Address warnings found by Localization compilation

diff --git a/Engine/Plugins/Experimental/TextureGraph/Source/MixerInsightEditor/Private/MixerInsightEditor.cpp b/Engine/Plugins/Experimental/TextureGraph/Source/MixerInsightEditor/Private/MixerInsightEditor.cpp
--- a/Engine/Plugins/Experimental/TextureGraph/Source/MixerInsightEditor/Private/MixerInsightEditor.cpp
+++ b/Engine/Plugins/Experimental/TextureGraph/Source/MixerInsightEditor/Private/MixerInsightEditor.cpp
@@ -227,7 +227,7 @@ TSharedRef<SDockTab> FMixerInsightEditorModule::OnSpawnPluginTab(const FSpawnTab
 			{
 				return SNew(SDockTab)
 					.TabRole(ETabRole::PanelTab)
-					.Label(LOCTEXT("FMixerInsightEditorTabTitleAction", "Mixer Inspector"))
+					.Label(LOCTEXT("FMixerInsightEditorTabTitleInspector", "Mixer Inspector"))
 					[
 						SNew(SMixerInsightInspectorView)
 					];
diff --git a/Engine/Plugins/Experimental/TextureGraph/Source/TextureGraphEditor/Private/EdGraph/TG_EdGraphSchema.cpp b/Engine/Plugins/Experimental/TextureGraph/Source/TextureGraphEditor/Private/EdGraph/TG_EdGraphSchema.cpp
--- a/Engine/Plugins/Experimental/TextureGraph/Source/TextureGraphEditor/Private/EdGraph/TG_EdGraphSchema.cpp
+++ b/Engine/Plugins/Experimental/TextureGraph/Source/TextureGraphEditor/Private/EdGraph/TG_EdGraphSchema.cpp
@@ -166,7 +166,7 @@ const FPinConnectionResponse UTG_EdGraphSchema::CanCreateConnection(const UEdGra
 
 	if (!CategorizePinsByDirection(A, B, /*out*/ InputPin, /*out*/ OutputPin))
 	{
-		return FPinConnectionResponse(CONNECT_RESPONSE_DISALLOW, LOCTEXT("ConnectionIncompatible", "Directions are not compatible"));
+		return FPinConnectionResponse(CONNECT_RESPONSE_DISALLOW, LOCTEXT("ConnectionDirectionsNotCompatible", "Directions are not compatible"));
 	}
 
 	// Check for new and existing loops
@@ -179,7 +179,7 @@ const FPinConnectionResponse UTG_EdGraphSchema::CanCreateConnection(const UEdGra
 	// Check for incompatible pins and get description if they cannot connect
 	if (!ArePinsCompatible_Internal(OutputPin, InputPin))
 	{
-		return FPinConnectionResponse(CONNECT_RESPONSE_DISALLOW, LOCTEXT("ConnectionLoop", "Pins are not compatible").ToString());
+		return FPinConnectionResponse(CONNECT_RESPONSE_DISALLOW, LOCTEXT("ConnectionPinsNotCompatible", "Pins are not compatible").ToString());
 	}
 
 	// For non-exec pins, break existing connections on inputs only - multiple output connections are acceptable
diff --git a/Engine/Plugins/Experimental/TextureGraph/Source/TextureGraphEditor/Private/Pins/STG_GraphPinOutputSettingsWidget.cpp b/Engine/Plugins/Experimental/TextureGraph/Source/TextureGraphEditor/Private/Pins/STG_GraphPinOutputSettingsWidget.cpp
--- a/Engine/Plugins/Experimental/TextureGraph/Source/TextureGraphEditor/Private/Pins/STG_GraphPinOutputSettingsWidget.cpp
+++ b/Engine/Plugins/Experimental/TextureGraph/Source/TextureGraphEditor/Private/Pins/STG_GraphPinOutputSettingsWidget.cpp
@@ -95,14 +95,14 @@ void STG_GraphPinOutputSettingsWidget::Construct(const FArguments& InArgs, UEdGr
 		+ SVerticalBox::Slot()
 		.Padding(2)
 		[
-			AddEnumComobox(LOCTEXT("OutputWidth", "Height"), GetHeightDelegate, OnGenerateHeightMenu)
+			AddEnumComobox(LOCTEXT("OutputHeight", "Height"), GetHeightDelegate, OnGenerateHeightMenu)
 		]
 
 		//Path Format
 		+ SVerticalBox::Slot()
 		.Padding(2)
 		[
-			AddEnumComobox(LOCTEXT("OutputWidth", "Format"), GetFormatDelegate, OnGenerateFormatMenu)
+			AddEnumComobox(LOCTEXT("OutputFormat", "Format"), GetFormatDelegate, OnGenerateFormatMenu)
 		]
 	];
 }
diff --git a/Engine/Plugins/Experimental/TextureGraph/Source/TextureGraphEditor/Private/STG_SelectionPreview.cpp b/Engine/Plugins/Experimental/TextureGraph/Source/TextureGraphEditor/Private/STG_SelectionPreview.cpp
--- a/Engine/Plugins/Experimental/TextureGraph/Source/TextureGraphEditor/Private/STG_SelectionPreview.cpp
+++ b/Engine/Plugins/Experimental/TextureGraph/Source/TextureGraphEditor/Private/STG_SelectionPreview.cpp
@@ -315,14 +315,14 @@ TSharedRef<SWidget> STG_SelectionPreview::MakeZoomControlWidget()
 			FCanExecuteAction(),
 			FIsActionChecked::CreateSP(this, &STG_SelectionPreview::IsZoomMenuFitChecked)
 		);
-		ZoomMenuBuilder.AddMenuEntry(NSLOCTEXT("STG_SelectionPreview", "ZoomFillAction", "Scale To Fit"), NSLOCTEXT("STG_SelectionPreview", "ZoomFitActionHint", "Scales the texture down to fit within the viewport if needed."), FSlateIcon(), ZoomFitAction, NAME_None, EUserInterfaceActionType::RadioButton);
+		ZoomMenuBuilder.AddMenuEntry(NSLOCTEXT("STG_SelectionPreview", "ZoomFitAction", "Scale To Fit"), NSLOCTEXT("STG_SelectionPreview", "ZoomFitActionHint", "Scales the texture down to fit within the viewport if needed."), FSlateIcon(), ZoomFitAction, NAME_None, EUserInterfaceActionType::RadioButton);
 		
 		FUIAction ZoomFillAction(
 			FExecuteAction::CreateSP(this, &STG_SelectionPreview::HandleZoomMenuFillClicked),
 			FCanExecuteAction(),
 			FIsActionChecked::CreateSP(this, &STG_SelectionPreview::IsZoomMenuFillChecked)
 		);
-		ZoomMenuBuilder.AddMenuEntry(NSLOCTEXT("STG_SelectionPreview", "ZoomFillAction", "Scale To Fill"), NSLOCTEXT("STG_SelectionPreview", "ZoomFitActionHint", "Scales the texture up and down to fill the viewport."), FSlateIcon(), ZoomFillAction, NAME_None, EUserInterfaceActionType::RadioButton);
+		ZoomMenuBuilder.AddMenuEntry(NSLOCTEXT("STG_SelectionPreview", "ZoomFillAction", "Scale To Fill"), NSLOCTEXT("STG_SelectionPreview", "ZoomFillActionHint", "Scales the texture up and down to fill the viewport."), FSlateIcon(), ZoomFillAction, NAME_None, EUserInterfaceActionType::RadioButton);
 	}
 
 	// zoom Dropdown
