TextureGraph: Fix TextureHelper::RawFromResource() not taking ptich vs width into account

diff --git a/Engine/Plugins/Experimental/TextureGraph/Source/MixerEngine/2D/TextureHelper.cpp b/Engine/Plugins/Experimental/TextureGraph/Source/MixerEngine/2D/TextureHelper.cpp
--- a/Engine/Plugins/Experimental/TextureGraph/Source/MixerEngine/2D/TextureHelper.cpp
+++ b/Engine/Plugins/Experimental/TextureGraph/Source/MixerEngine/2D/TextureHelper.cpp
@@ -819,7 +819,7 @@ size_t TextureHelper::RoundUpTo(size_t Size, size_t DesiredRounding)
 	return RoundedSize;
 }
 
-RawBufferPtr TextureHelper::RawFromResource(FTexture2DRHIRef ResourceRHI, const BufferDescriptor& Desc)
+RawBufferPtr TextureHelper::RawFromResource(const FTexture2DRHIRef& ResourceRHI, const BufferDescriptor& Desc)
 {
 	check(IsInRenderingThread()); 
 
@@ -827,18 +827,16 @@ RawBufferPtr TextureHelper::RawFromResource(FTexture2DRHIRef ResourceRHI, const
 	{
 		RawBufferPtr RawObj;
 
-		EPixelFormat PixelFormat = Desc.PixelFormat();
+		const EPixelFormat PixelFormat = Desc.PixelFormat();
 
 		/// These must be the same!
 		check(TextureHelper::GetBppFromPixelFormat(ResourceRHI->GetFormat()) == TextureHelper::GetBppFromPixelFormat(PixelFormat));
 
-		SIZE_T Length = CalculateImageBytes(Desc.Width, Desc.Height, 0, PixelFormat);
-
-		uint32 BitsPerPixel = TextureHelper::GetBppFromPixelFormat(PixelFormat);
+		const uint32 BitsPerPixel = TextureHelper::GetBppFromPixelFormat(PixelFormat);
 		check(BitsPerPixel % 8 == 0);
-		uint32 BytesPerPixel = BitsPerPixel / 8;
+		const uint32 BytesPerPixel = BitsPerPixel / 8;
 
-		TUniquePtr<FRHIGPUTextureReadback> TextureReadback = MakeUnique<FRHIGPUTextureReadback>(TEXT("RawFromResourceTextureReadback"));
+		const TUniquePtr<FRHIGPUTextureReadback> TextureReadback = MakeUnique<FRHIGPUTextureReadback>(TEXT("RawFromResourceTextureReadback"));
 
 		FRHICommandListImmediate& RHI = FRHICommandListExecutor::GetImmediateCommandList();
 		RHI.FlushResources();
@@ -846,48 +844,44 @@ RawBufferPtr TextureHelper::RawFromResource(FTexture2DRHIRef ResourceRHI, const
 
 		TextureReadback->EnqueueCopy(RHI, ResourceRHI);
 		RHI.BlockUntilGPUIdle();
-		
+
 		//check(TextureReadback->IsReady());
 		{
-			int32 OutBufferWidth;
-			int32 OutBufferHeight;
-			const uint8* SrcData = (const uint8*)TextureReadback->Lock(OutBufferWidth, &OutBufferHeight);
+			const size_t DstDataLength = Desc.Width * Desc.Height * BytesPerPixel;
+			uint8* const DstData = new uint8[DstDataLength];
 
-			/// The behaviour of Lock is erratic. For some image formats it returns the Pitch in bytes (OutBufferWidth)
-			/// but for some formats it returns that as the width of the image in pixels
-			int32 SrcDataLength = std::max(OutBufferWidth * OutBufferHeight, (int32)Length);
+			int32 OutBufferRowPitchInPixels, OutBufferHeight;
+			const uint8* SrcData = static_cast<const uint8*>(TextureReadback->Lock(OutBufferRowPitchInPixels, &OutBufferHeight));
+			check(SrcData);
+			check(static_cast<uint32>(OutBufferRowPitchInPixels) >= Desc.Width && static_cast<uint32>(OutBufferHeight) >= Desc.Height);
 
-			/// If the OutBufferWidth was actually a width rather than the pitch of the buffer then we need to ensure
-			/// that we have the correct length based on the rounded dimensions
-			if (OutBufferWidth > (int32)Desc.Width && OutBufferWidth < (int32)Desc.Height * 2)
+			if (OutBufferRowPitchInPixels == Desc.Width)
 			{
-				SrcDataLength = std::max(OutBufferWidth * OutBufferHeight * (int32)BytesPerPixel, SrcDataLength);
+				// If pitch and width are the same, we can just copy the entire buffer.
+
+				FMemory::Memcpy(DstData, SrcData, DstDataLength);
 			}
 			else
 			{
-				SrcDataLength = std::max(OutBufferWidth * OutBufferHeight, SrcDataLength);
-			}
+				// If pitch and width are NOT the same, we need to copy the valid pixels row by row.
 
-			/// Round it up to the optimal hashing size
-			int32 RoundedDataLength = SrcDataLength; // (int32)DataUtil::GetOptimalHashingSize((size_t)SrcDataLength);
-			check(RoundedDataLength >= SrcDataLength);
+				      uint8* RowDstPtr = DstData;
+				const size_t RowDstLength = Desc.Width * BytesPerPixel;
+				const uint8* RowSrcPtr = SrcData;
+				const size_t RowSrcLength = OutBufferRowPitchInPixels * BytesPerPixel;
 
-			auto Stride = BitsPerPixel * OutBufferWidth;
-			check(Length / Desc.Height <= Stride || (Length < 256 && Length <= Stride));
-			uint8* DstData = new uint8 [RoundedDataLength];
-			
-			check(SrcData); 
-			memcpy(DstData, SrcData, SrcDataLength);
+				for (uint32 Row = 0; Row < Desc.Height; ++Row)
+				{
+					FMemory::Memcpy(RowDstPtr, RowSrcPtr, RowDstLength);
 
-			/// Ensure that we're zero'ing the padding
-			if (RoundedDataLength != SrcDataLength)
-			{
-				memset(DstData + SrcDataLength, 0, RoundedDataLength - SrcDataLength);
+					RowDstPtr += RowDstLength;
+					RowSrcPtr += RowSrcLength;
+				}
 			}
 
 			TextureReadback->Unlock();
-			
-			RawObj = std::make_shared<RawBuffer>(DstData, RoundedDataLength, Desc);
+
+			RawObj = std::make_shared<RawBuffer>(DstData, DstDataLength, Desc);
 		}
 
 		return RawObj;
diff --git a/Engine/Plugins/Experimental/TextureGraph/Source/MixerEngine/2D/TextureHelper.h b/Engine/Plugins/Experimental/TextureGraph/Source/MixerEngine/2D/TextureHelper.h
--- a/Engine/Plugins/Experimental/TextureGraph/Source/MixerEngine/2D/TextureHelper.h
+++ b/Engine/Plugins/Experimental/TextureGraph/Source/MixerEngine/2D/TextureHelper.h
@@ -82,7 +82,7 @@ struct MIXERENGINE_API TextureHelper
 
 	static RawBufferPtr					RawFromRT(UTextureRenderTarget2D* RenderTarget, const BufferDescriptor& Desc);
 	static RawBufferPtr					RawFromTexture(UTexture2D* Texture, const BufferDescriptor& Desc);
-	static RawBufferPtr					RawFromResource(FTexture2DRHIRef ResourceRHI, const BufferDescriptor& Desc);
+	static RawBufferPtr					RawFromResource(const FTexture2DRHIRef& ResourceRHI, const BufferDescriptor& Desc);
 	static BufferFormat					FindOptimalSupportedFormat(BufferFormat SrcFormat);
 
 	static void							RawFromRT_Tiled(UTextureRenderTarget2D* RenderTarget, const BufferDescriptor& Desc, size_t TileSizeX, size_t TileSizeY, RawBufferPtrTiles& Tiles);
