Dynamic Pins in Nodes UI

Use CanEditChange instead of EditCondition in UProperty Meta to allow definining EditCondition logic in a single place for the Expressions and their properties.

diff --git a/Engine/Plugins/Experimental/TextureGraph/Source/TextureGraph/Private/Expressions/FIlter/TG_Expression_Blur.cpp b/Engine/Plugins/Experimental/TextureGraph/Source/TextureGraph/Private/Expressions/FIlter/TG_Expression_Blur.cpp
--- a/Engine/Plugins/Experimental/TextureGraph/Source/TextureGraph/Private/Expressions/FIlter/TG_Expression_Blur.cpp
+++ b/Engine/Plugins/Experimental/TextureGraph/Source/TextureGraph/Private/Expressions/FIlter/TG_Expression_Blur.cpp
@@ -3,7 +3,38 @@
 #include "Expressions/Filter/TG_Expression_Blur.h"
 
 #include "Transform/Expressions/T_Blur.h"
+#if WITH_EDITOR
+bool UTG_Expression_Blur::CanEditChange(const FProperty* InProperty) const
+{
+	bool bEditCondition = Super::CanEditChange(InProperty);
+
+	// if already set to false Or InProperty not directly owned by us, early out
+	if (!bEditCondition || this->GetClass() != InProperty->GetOwnerClass())
+	{
+		return bEditCondition;
+	}
 
+	const FName PropertyName = InProperty->GetFName();
+	
+
+	// Specific logic associated with Property
+	if (PropertyName == GET_MEMBER_NAME_CHECKED(UTG_Expression_Blur, Radius))
+	{
+		bEditCondition = BlurType == EBlurType::Gaussian || BlurType == EBlurType::Radial;
+	}
+	if (PropertyName == GET_MEMBER_NAME_CHECKED(UTG_Expression_Blur, Angle))
+	{
+		bEditCondition = BlurType == EBlurType::Directional;
+	}
+	if (PropertyName == GET_MEMBER_NAME_CHECKED(UTG_Expression_Blur, Strength))
+	{
+		bEditCondition = BlurType == EBlurType::Directional || BlurType == EBlurType::Radial;
+	}
+
+	// Default behaviour
+	return bEditCondition;
+}
+#endif
 void UTG_Expression_Blur::Evaluate(FTG_EvaluationContext* InContext)
 {
 	Super::Evaluate(InContext);
diff --git a/Engine/Plugins/Experimental/TextureGraph/Source/TextureGraph/Private/Expressions/Procedural/TG_Expression_Shape.cpp b/Engine/Plugins/Experimental/TextureGraph/Source/TextureGraph/Private/Expressions/Procedural/TG_Expression_Shape.cpp
--- a/Engine/Plugins/Experimental/TextureGraph/Source/TextureGraph/Private/Expressions/Procedural/TG_Expression_Shape.cpp
+++ b/Engine/Plugins/Experimental/TextureGraph/Source/TextureGraph/Private/Expressions/Procedural/TG_Expression_Shape.cpp
@@ -32,3 +32,41 @@ void UTG_Expression_Shape::Evaluate(FTG_EvaluationContext* InContext)
 
 	Output = T_ShapeMask::Create(InContext->Cycle, DesiredDescriptor, ShaderShapeType, Params, InContext->TargetId);
 }
+#if WITH_EDITOR
+bool UTG_Expression_Shape::CanEditChange(const FProperty* InProperty) const
+{
+	bool bEditCondition = Super::CanEditChange(InProperty);
+
+	// if already set to false Or InProperty not directly owned by us, early out
+	// the CanEditChange function also checks down the PropertyChain,
+	// checking direct ownership fixes an issue here where deeper in the hierarchy,
+	// FTG_TextureDescriptor::Height was also being affected because of the logic for TG_Expression_Shape::Height below 
+	if (!bEditCondition || this->GetClass() != InProperty->GetOwnerClass())
+	{
+		return bEditCondition;
+	}
+
+	const FName PropertyName = InProperty->GetFName();
+	
+	// Specific logic associated with Property
+	if (PropertyName == GET_MEMBER_NAME_CHECKED(UTG_Expression_Shape, PolygonNumSides))
+	{
+		bEditCondition = ShapeType == EShapeType::Polygon;
+	}
+	if (PropertyName == GET_MEMBER_NAME_CHECKED(UTG_Expression_Shape, Height))
+	{
+		bEditCondition = ShapeType == EShapeType::Rectangle || (ShapeType==EShapeType::Ellipse);
+	}
+	if (PropertyName == GET_MEMBER_NAME_CHECKED(UTG_Expression_Shape, Orientation))
+	{
+		bEditCondition = ShapeType != EShapeType::Circle;
+	}
+	if (PropertyName == GET_MEMBER_NAME_CHECKED(UTG_Expression_Shape, Rounding))
+	{
+		bEditCondition = ShapeType != EShapeType::Circle;
+	}
+
+	// Default behaviour
+	return bEditCondition;
+}
+#endif
\ No newline at end of file
diff --git a/Engine/Plugins/Experimental/TextureGraph/Source/TextureGraph/Public/Expressions/Filter/TG_Expression_Blur.h b/Engine/Plugins/Experimental/TextureGraph/Source/TextureGraph/Public/Expressions/Filter/TG_Expression_Blur.h
--- a/Engine/Plugins/Experimental/TextureGraph/Source/TextureGraph/Public/Expressions/Filter/TG_Expression_Blur.h
+++ b/Engine/Plugins/Experimental/TextureGraph/Source/TextureGraph/Public/Expressions/Filter/TG_Expression_Blur.h
@@ -23,16 +23,21 @@ class TEXTUREGRAPH_API UTG_Expression_Blur : public UTG_Expression
 public:
 	TG_DECLARE_EXPRESSION(TG_Category::Filter);
 	
-	UPROPERTY(EditAnywhere, BlueprintReadWrite, Category = NoCategory, meta = (TGType = "TG_Setting", ClampMin = 0, ClampMax = 100, EditCondition="BlurType==EBlurType::Gaussian || BlurType==EBlurType::Radial", EditConditionHides))
+#if WITH_EDITOR
+	// Used to implement EditCondition logic for both Node UI and Details View
+	virtual bool CanEditChange(const FProperty* InProperty) const override;
+#endif
+	
+	UPROPERTY(EditAnywhere, BlueprintReadWrite, Category = NoCategory, meta = (TGType = "TG_Setting", ClampMin = 0, ClampMax = 100, EditConditionHides))
 	int32 Radius = 1;
 	
-	UPROPERTY(EditAnywhere, BlueprintReadWrite, Category = NoCategory, meta = (TGType = "TG_Setting", ClampMin = -180, ClampMax = 180, EditCondition="BlurType==EBlurType::Directional", EditConditionHides))
+	UPROPERTY(EditAnywhere, BlueprintReadWrite, Category = NoCategory, meta = (TGType = "TG_Setting", ClampMin = -180, ClampMax = 180, EditConditionHides))
 	float Angle = 0.0f;
 	
-	UPROPERTY(EditAnywhere, BlueprintReadWrite, Category = NoCategory, meta = (TGType = "TG_Setting", ClampMin = 0, ClampMax = 1, EditCondition="BlurType==EBlurType::Directional || BlurType==EBlurType::Radial", EditConditionHides))
+	UPROPERTY(EditAnywhere, BlueprintReadWrite, Category = NoCategory, meta = (TGType = "TG_Setting", ClampMin = 0, ClampMax = 1, EditConditionHides))
 	float Strength = 0.1f;
 
-	UPROPERTY(EditAnywhere, BlueprintReadWrite, Category = NoCategory, meta = (TGType = "TG_Setting"))
+	UPROPERTY(EditAnywhere, BlueprintReadWrite, Category = NoCategory, meta = (TGType = "TG_Setting", RegenPinsOnChange))
 	EBlurType BlurType = EBlurType::Gaussian;
 
 	UPROPERTY(EditAnywhere, Category = NoCategory, meta = (TGType = "TG_Output", PinDisplayName = ""))
diff --git a/Engine/Plugins/Experimental/TextureGraph/Source/TextureGraph/Public/Expressions/Procedural/TG_Expression_Shape.h b/Engine/Plugins/Experimental/TextureGraph/Source/TextureGraph/Public/Expressions/Procedural/TG_Expression_Shape.h
--- a/Engine/Plugins/Experimental/TextureGraph/Source/TextureGraph/Public/Expressions/Procedural/TG_Expression_Shape.h
+++ b/Engine/Plugins/Experimental/TextureGraph/Source/TextureGraph/Public/Expressions/Procedural/TG_Expression_Shape.h
@@ -34,23 +34,28 @@ public:
 	
 	virtual void						Evaluate(FTG_EvaluationContext* InContext) override;
 
+#if WITH_EDITOR
+	// Used to implement EditCondition logic for both Node UI and Details View
+	virtual bool						CanEditChange(const FProperty* InProperty) const override;
+#endif
+	
 	// The type of shape
-	UPROPERTY(EditAnywhere, BlueprintReadWrite, Category = NoCategory, meta = (TGType = "TG_Setting"))
+	UPROPERTY(EditAnywhere, BlueprintReadWrite, Category = NoCategory, meta = (TGType = "TG_Setting", RegenPinsOnChange))
 	EShapeType							ShapeType = EShapeType::Circle;
 
-	UPROPERTY(EditAnywhere, BlueprintReadWrite, Category = NoCategory, meta = (TGType = "TG_Setting", ClampMin = 7, ClampMax = 32, EditCondition="ShapeType==EShapeType::Polygon", EditConditionHides))
+	UPROPERTY(EditAnywhere, BlueprintReadWrite, Category = NoCategory, meta = (TGType = "TG_Setting", ClampMin = 7, ClampMax = 32, EditConditionHides))
 	int32								PolygonNumSides = 8;
 
-	UPROPERTY(EditAnywhere, BlueprintReadWrite, Category = NoCategory, meta = (TGType = "TG_Setting", ClampMin = 0, ClampMax = 1))
+	UPROPERTY(EditAnywhere, BlueprintReadWrite, Category = NoCategory, meta = (TGType = "TG_Setting", ClampMin = 0, ClampMax =1))
 	float 								Width = 1.0f;
 
-	UPROPERTY(EditAnywhere, BlueprintReadWrite, Category = NoCategory, meta = (TGType = "TG_Setting", ClampMin = 0, ClampMax = 1, EditCondition = "(ShapeType==EShapeType::Rectangle) || (ShapeType==EShapeType::Ellipse)", EditConditionHides))
+	UPROPERTY(EditAnywhere, BlueprintReadWrite, Category = NoCategory, meta = (TGType = "TG_Setting", ClampMin = 0, ClampMax = 1, EditConditionHides))
 	float 								Height = 1.0f;
 
-	UPROPERTY(EditAnywhere, BlueprintReadWrite, Category = NoCategory, meta = (TGType = "TG_Setting", UIMin = "-180", ClampMin = "-180", UIMax = "180", ClampMax = "180", Delta = "1", Units = "Degrees", EditCondition = "!(ShapeType==EShapeType::Circle)", EditConditionHides))
+	UPROPERTY(EditAnywhere, BlueprintReadWrite, Category = NoCategory, meta = (TGType = "TG_Setting", UIMin = "-180", ClampMin = "-180", UIMax = "180", ClampMax = "180", Delta = "1", Units = "Degrees", EditConditionHides))
 	float								Orientation = 0.0f;
 
-	UPROPERTY(EditAnywhere, BlueprintReadWrite, Category = NoCategory, meta = (TGType = "TG_Setting", ClampMin = 0, ClampMax = 1, EditCondition = "!((ShapeType==EShapeType::Circle) || (ShapeType==EShapeType::Ellipse))", EditConditionHides))
+	UPROPERTY(EditAnywhere, BlueprintReadWrite, Category = NoCategory, meta = (TGType = "TG_Setting", ClampMin = 0, ClampMax = 1, EditConditionHides))
 	float								Rounding = 0.0f;
 
 	UPROPERTY(EditAnywhere, BlueprintReadWrite, Category = NoCategory, meta = (TGType = "TG_Setting", ClampMin = 0, ClampMax = 1))
diff --git a/Engine/Plugins/Experimental/TextureGraph/Source/TextureGraphEditor/Private/EdGraph/TG_EdGraph.cpp b/Engine/Plugins/Experimental/TextureGraph/Source/TextureGraphEditor/Private/EdGraph/TG_EdGraph.cpp
--- a/Engine/Plugins/Experimental/TextureGraph/Source/TextureGraphEditor/Private/EdGraph/TG_EdGraph.cpp
+++ b/Engine/Plugins/Experimental/TextureGraph/Source/TextureGraphEditor/Private/EdGraph/TG_EdGraph.cpp
@@ -6,12 +6,13 @@
 #include "EdGraph/TG_EdGraphNode.h"
 #include "Expressions/TG_Expression.h"
 #include "TextureGraph.h"
+#include "TG_Editor.h"
 #include "TG_Graph.h"
 
-void UTG_EdGraph::InitializeFromTextureGraph(UTextureGraph* InTextureGraph)
+void UTG_EdGraph::InitializeFromTextureGraph(UTextureGraph* InTextureGraph, TWeakPtr<FTG_Editor> InTGEditor)
 {
 	TextureGraph = InTextureGraph;
-
+	TGEditor = InTGEditor;
 	// And now build the EdGraph viewmodel matching the script
 	BuildEdGraphViewmodel();
 }
@@ -110,6 +111,15 @@ UTG_EdGraphNode* UTG_EdGraph::GetViewModelNode(FTG_Id NodeId)
 	return (Result ? Cast<UTG_EdGraphNode>(Result->Get()) : nullptr);
 }
 
+void UTG_EdGraph::RefreshEditorDetails() const
+{
+	//Tell Editor to update details
+	if (TGEditor.IsValid())
+	{
+		TGEditor.Pin()->RefreshDetailsView();
+	}
+}
+
 void UTG_EdGraph::GraphChanged(UTG_Graph* InGraph, UTG_Node* InNode, bool Tweaking)
 {
 	if (InNode)
diff --git a/Engine/Plugins/Experimental/TextureGraph/Source/TextureGraphEditor/Private/EdGraph/TG_EdGraphNode.cpp b/Engine/Plugins/Experimental/TextureGraph/Source/TextureGraphEditor/Private/EdGraph/TG_EdGraphNode.cpp
--- a/Engine/Plugins/Experimental/TextureGraph/Source/TextureGraphEditor/Private/EdGraph/TG_EdGraphNode.cpp
+++ b/Engine/Plugins/Experimental/TextureGraph/Source/TextureGraphEditor/Private/EdGraph/TG_EdGraphNode.cpp
@@ -205,7 +205,9 @@ void UTG_EdGraphNode::AllocateDefaultPins()
 				}
 				NewPin->bNotConnectable = true;
 			}
-			
+#if WITH_EDITOR
+			UpdatePinVisibility(NewPin, Property);
+#endif
 			NewPin->DefaultValue = DefaultValue;
 			// This updates the UObject (Texture/Material) picker UI in the Node to get updated
 			if (Property && Property->GetClass()->IsChildOf(FObjectProperty::StaticClass()))
@@ -358,6 +360,7 @@ void UTG_EdGraphNode::ReconstructNode()
 		}
 	}
 
+	
 	// Remove old pins
 	for (UEdGraphPin* OldPin : OldPins)
 	{
@@ -440,6 +443,44 @@ void UTG_EdGraphNode::PostPasteNode()
 	Graph->AddPostPasteNode(Node);
 }
 
+#if WITH_EDITOR
+void UTG_EdGraphNode::UpdatePinVisibility(UEdGraphPin* Pin, FProperty* Property) const
+{
+	UTG_Expression* Expression = GetNode()->GetExpression();
+	bool bCanEditChange = Expression->CanEditChange(Property);
+	bool bEditConditionHides = Property->HasMetaData("EditConditionHides");
+
+	// update Metadata with EditCondition to also update details view
+	FName NAME_EditCondition(TEXT("EditCondition"));
+	Property->SetMetaData(NAME_EditCondition, bCanEditChange ? TEXT("true") : TEXT("false"));
+	
+	Pin->bHidden = !bCanEditChange && bEditConditionHides;
+	Pin->bDefaultValueIsReadOnly = bCanEditChange;
+}
+
+void UTG_EdGraphNode::UpdateInputPinsVisibility() const
+{
+	auto Expression = Node->GetExpression();
+	// go through all pins and match them with their property hide states
+	for (FTG_Id OtherPinId : Node->GetInputPinIds())
+	{
+		TObjectPtr<UTG_Pin> TGPin = Node->GetGraph()->GetPin(OtherPinId);
+		check(TGPin);
+
+		const FName& OtherPinName = TGPin->GetArgumentName();
+		UEdGraphPin* EdPin = FindPinChecked(OtherPinName, EEdGraphPinDirection::EGPD_Input);
+
+		// check hide state
+		
+		auto PinProperty = TGPin->GetExpressionProperty();
+		if (PinProperty)
+		{
+			UpdatePinVisibility(EdPin, PinProperty);
+		}
+	}
+}
+#endif
+
 void UTG_EdGraphNode::PinDefaultValueChangedWithTweaking(UEdGraphPin* Pin, bool bIsTweaking)
 {
 	const UTG_EdGraphSchema* Schema = Cast<const UTG_EdGraphSchema>(GetSchema());
@@ -589,8 +630,9 @@ void UTG_EdGraphNode::OnNodePostEvaluate(const FTG_EvaluationContext* Evaluation
 	OnNodeChanged(GetNode());
 }
 
-void UTG_EdGraphNode::UpdateEdPinDefaultValue(UEdGraphPin* EdPin, const UTG_EdGraphSchema* Schema)
+bool UTG_EdGraphNode::UpdateEdPinDefaultValue(UEdGraphPin* EdPin, const UTG_EdGraphSchema* Schema)
 {
+	bool bShouldUpdatePinsVisibility = false;
 	UTG_Pin* TSPin = Schema->GetTGPinFromEdPin(EdPin);
 	if (TSPin)
 	{
@@ -623,6 +665,8 @@ void UTG_EdGraphNode::UpdateEdPinDefaultValue(UEdGraphPin* EdPin, const UTG_EdGr
 				DefaultValue = Enum->GetNameByValue(FCString::Atoi(*DefaultValue)).ToString();
 			}
 		}
+		// check here if the property has a tag for regen pins and value has changed
+		bShouldUpdatePinsVisibility = Property->HasMetaData("RegenPinsOnChange") && EdPin->DefaultValue != DefaultValue;
 		
 		EdPin->bDefaultValueIsReadOnly = false;
 		EdPin->DefaultValue = DefaultValue;
@@ -634,16 +678,34 @@ void UTG_EdGraphNode::UpdateEdPinDefaultValue(UEdGraphPin* EdPin, const UTG_EdGr
 		}
 		EdPin->bDefaultValueIsReadOnly = EdPin->HasAnyConnections();
 	}
+	return bShouldUpdatePinsVisibility;
 }
 
 void UTG_EdGraphNode::OnNodeChanged(UTG_Node* InNode)
 {
+	bool bShouldUpdatePinsVisibility = false;
 	// Update Default Values to their respective Var values;
 	for(UEdGraphPin* EdPin : Pins)
 	{
 		const UTG_EdGraphSchema* Schema = Cast<const UTG_EdGraphSchema>(GetSchema());
-		UpdateEdPinDefaultValue(EdPin, Schema);
+		bShouldUpdatePinsVisibility |= UpdateEdPinDefaultValue(EdPin, Schema);
 	}
+#if WITH_EDITOR
+	// if there was a pin which triggered regeneration of all other input pins (used meta "RegenPinsOnChange"), we update them
+	if (bShouldUpdatePinsVisibility)
+	{
+		UpdateInputPinsVisibility();
+		
+		// cache advanced visibility
+		TEnumAsByte<ENodeAdvancedPins::Type> CachedAdvancedPinDisplay = this->AdvancedPinDisplay;
+		ReconstructNode();
+		this->AdvancedPinDisplay = CachedAdvancedPinDisplay;
+
+		//Tell Editor to update details
+		Cast<UTG_EdGraph>(GetGraph())->RefreshEditorDetails();
+	}
+#endif
+	
 }
 
 void UTG_EdGraphNode::PostTransacted(const FTransactionObjectEvent& TransactionEvent)
diff --git a/Engine/Plugins/Experimental/TextureGraph/Source/TextureGraphEditor/Private/TG_Editor.cpp b/Engine/Plugins/Experimental/TextureGraph/Source/TextureGraphEditor/Private/TG_Editor.cpp
--- a/Engine/Plugins/Experimental/TextureGraph/Source/TextureGraphEditor/Private/TG_Editor.cpp
+++ b/Engine/Plugins/Experimental/TextureGraph/Source/TextureGraphEditor/Private/TG_Editor.cpp
@@ -190,7 +190,7 @@ void FTG_Editor::InitEditor(const EToolkitMode::Type Mode, const TSharedPtr< cla
 	TG_EdGraph = NewObject<UTG_EdGraph>(EditedTextureGraph, UTG_EdGraph::StaticClass(), NAME_None, RF_Transactional | RF_Transient);
 	TG_EdGraph->Schema = UTG_EdGraphSchema::StaticClass();
 
-	TG_EdGraph->InitializeFromTextureGraph(EditedTextureGraph);
+	TG_EdGraph->InitializeFromTextureGraph(EditedTextureGraph, SharedThis(this));
 
 	TG_EdGraph->PinSelectionManager.OnPinSelectionUpdated.AddSP(this, &FTG_Editor::OnPinSelectionUpdated);
 
@@ -1392,6 +1392,11 @@ void FTG_Editor::RefreshTool()
 	RefreshViewport();
 }
 
+void FTG_Editor::RefreshDetailsView() const
+{
+	DetailsView->ForceRefresh();
+}
+
 void FTG_Editor::SaveAsset_Execute()
 {
 	UE_LOG(LogTextureGraphEditor, Log, TEXT("Saving Texture Graph %s"), *GetEditingObjects()[0]->GetName());
diff --git a/Engine/Plugins/Experimental/TextureGraph/Source/TextureGraphEditor/Private/TG_Editor.h b/Engine/Plugins/Experimental/TextureGraph/Source/TextureGraphEditor/Private/TG_Editor.h
--- a/Engine/Plugins/Experimental/TextureGraph/Source/TextureGraphEditor/Private/TG_Editor.h
+++ b/Engine/Plugins/Experimental/TextureGraph/Source/TextureGraphEditor/Private/TG_Editor.h
@@ -95,20 +95,21 @@ public:
 	/** Called to bring focus to the details panel */
 	void											FocusDetailsPanel();
 
-	/*
-	 * Refresh Errors and Warnings
-	 */
+	/** Refresh Errors and Warnings **/
 	void											RefreshErrors();
-	/**
-	 * Refreshes the viewport containing the preview mesh.
-	 */
+
+	/** Refreshes the viewport containing the preview mesh. **/
 	void											RefreshPreviewViewport();
+	
 	/** Called to update the selection view */
 	void											RefreshSelectionPreview(const TSet<class UObject*>& NewSelection);
 
 	void											SetMesh(class UMeshComponent* InPreviewMesh, class UWorld* InWorld) override;
-	bool SetPreviewAsset(UObject* InAsset);
-	bool SetPreviewAssetByName(const TCHAR* InAssetName);
+	bool 											SetPreviewAsset(UObject* InAsset);
+	bool 											SetPreviewAssetByName(const TCHAR* InAssetName);
+
+	/** Force Refresh Details View **/
+	void											RefreshDetailsView() const;
 
 protected:
 	TArray<UTG_EdGraphNode*>						GetCurrentSelectedTG_EdGraphNodes() const;
@@ -121,9 +122,9 @@ protected:
 
 	virtual void									RefreshViewport() override;
 
-
 	virtual void									RefreshTool() override;
 
+
 	// ~Begin FAssetEditorToolkit interface
 	/** Called when "Save" is clicked for this asset */
 	virtual void									SaveAsset_Execute() override;
diff --git a/Engine/Plugins/Experimental/TextureGraph/Source/TextureGraphEditor/Public/EdGraph/TG_EdGraph.h b/Engine/Plugins/Experimental/TextureGraph/Source/TextureGraphEditor/Public/EdGraph/TG_EdGraph.h
--- a/Engine/Plugins/Experimental/TextureGraph/Source/TextureGraphEditor/Public/EdGraph/TG_EdGraph.h
+++ b/Engine/Plugins/Experimental/TextureGraph/Source/TextureGraphEditor/Public/EdGraph/TG_EdGraph.h
@@ -29,7 +29,7 @@ public:
      *
      * @param TextureGraph 
      */
-    void InitializeFromTextureGraph(UTextureGraph* InTextureGraph);
+    void InitializeFromTextureGraph(UTextureGraph* InTextureGraph, TWeakPtr<class FTG_Editor> InTGEditor);
 
 	/**
 	 * Add a TG_Node to the Graph and create a EdGraphNode controlling it
@@ -48,10 +48,15 @@ public:
 	 */
 	UTG_EdGraphNode* GetViewModelNode(FTG_Id NodeId);
 
-
+	/**
+	 * Force refresh Editor details
+	 */
+	void RefreshEditorDetails() const;
+	
 	UPROPERTY()
 	TObjectPtr<UTextureGraph> TextureGraph = nullptr;
-
+	
+	TWeakPtr<FTG_Editor> TGEditor = nullptr;
 	FTG_PinSelectionManager PinSelectionManager;
 
 private:
diff --git a/Engine/Plugins/Experimental/TextureGraph/Source/TextureGraphEditor/Public/EdGraph/TG_EdGraphNode.h b/Engine/Plugins/Experimental/TextureGraph/Source/TextureGraphEditor/Public/EdGraph/TG_EdGraphNode.h
--- a/Engine/Plugins/Experimental/TextureGraph/Source/TextureGraphEditor/Public/EdGraph/TG_EdGraphNode.h
+++ b/Engine/Plugins/Experimental/TextureGraph/Source/TextureGraphEditor/Public/EdGraph/TG_EdGraphNode.h
@@ -63,6 +63,11 @@ public:
 	void PostCopyNode();
 	virtual void PostPasteNode() override;
 
+#if WITH_EDITOR
+	void UpdatePinVisibility(UEdGraphPin* Pin, FProperty* Property) const;
+	void UpdateInputPinsVisibility() const;
+#endif
+	
 	/** Called when the DefaultValue of one of the pins of this node is changed in the editor */
 	virtual void PinDefaultValueChangedWithTweaking(UEdGraphPin* Pin, bool bIsTweaking);
 	virtual void PinConnectionListChanged(UEdGraphPin* Pin) override;
@@ -101,7 +106,7 @@ protected:
 	void UpdatePosition();
 
 	void OnNodePostEvaluate(const FTG_EvaluationContext* Fts_EvaluationContext);
-	void UpdateEdPinDefaultValue(UEdGraphPin* EdPin, const UTG_EdGraphSchema* Schema);
+	bool UpdateEdPinDefaultValue(UEdGraphPin* EdPin, const UTG_EdGraphSchema* Schema);
 
 	void OnUpdateCommentText(const FString& NewComment) override;
 	void OnCommentBubbleToggled(bool bInCommentBubbleVisible) override;
