TextureGraph: Update plugin definition

diff --git a/Engine/Plugins/Experimental/TextureGraph/TextureGraph.uplugin b/Engine/Plugins/Experimental/TextureGraph/TextureGraph.uplugin
--- a/Engine/Plugins/Experimental/TextureGraph/TextureGraph.uplugin
+++ b/Engine/Plugins/Experimental/TextureGraph/TextureGraph.uplugin
@@ -2,8 +2,7 @@
 	"FileVersion": 3,
 	"Version": 1,
 	"VersionName": "0.1",
-	"FriendlyName": "TextureGraphEditor",
-	"Description": "Asset editor for procedural texture graphs",
+	"Description": "Asset type and asset editor for procedural texturing",
 	"Category": "Editor",
 	"CreatedBy": "Epic Games, Inc.",
 	"CreatedByURL": "https://epicgames.com",
