Export tests

diff --git a/Engine/Plugins/Experimental/TextureGraph/Source/MixerEngine/2D/Tex.cpp b/Engine/Plugins/Experimental/TextureGraph/Source/MixerEngine/2D/Tex.cpp
--- a/Engine/Plugins/Experimental/TextureGraph/Source/MixerEngine/2D/Tex.cpp
+++ b/Engine/Plugins/Experimental/TextureGraph/Source/MixerEngine/2D/Tex.cpp
@@ -996,7 +996,9 @@ AsyncTiledBlobRef Tex::ToSingleBlob(CHashPtr Hash, bool TransferToRT /* = false
 			TiledBlobRef FinalBlobRef = TiledBlobObj;
 
 			if (!NoCache)
+			{
 				FinalBlobRef = MixerEngine::GetBlobber()->AddTiledResult(BlobHash, std::move(TiledBlobObj));
+			}
 
 			check(FinalBlobRef);
 			
@@ -1005,7 +1007,12 @@ AsyncTiledBlobRef Tex::ToSingleBlob(CHashPtr Hash, bool TransferToRT /* = false
 			Texture = nullptr;
 
 			if (!ResolveOnRenderThread)
-				return static_cast<AsyncTiledBlobRef>(PromiseUtil::OnGameThread().then([=]() { return FinalBlobRef; }));
+			{
+				return static_cast<AsyncTiledBlobRef>(PromiseUtil::OnGameThread().then([=]() 
+					{ 
+						return FinalBlobRef; 
+					}));
+			}
 
 			return static_cast<AsyncTiledBlobRef>(cti::make_ready_continuable<TiledBlobRef>(std::move(FinalBlobRef)));
 		});
@@ -1092,7 +1099,7 @@ AsyncTiledBlobRef Tex::ToBlob(int32 XTiles, int32 YTiles, uint32 Width /* = 0 */
 
 			bool bIsNowRT = RT != nullptr;	// Have we shifted to being an RT now?
 
-			DeviceBufferRef Buffer = bIsNowRT ? Device_FX::Get()->CreateFromRT(RT, Desc.ToBufferDescriptor()) : Device_FX::Get()->CreateFromTexture(Texture,Desc.ToBufferDescriptor());
+			DeviceBufferRef Buffer = bIsNowRT ? Device_FX::Get()->CreateFromRT(RT, Desc.ToBufferDescriptor()) : Device_FX::Get()->CreateFromTexture(Texture, Desc.ToBufferDescriptor());
 			T_Tiles<DeviceBufferRef> TileBuffers(XTiles, YTiles);
 
 			for (int32 TileX = 0; TileX < XTiles; TileX++)
@@ -1124,6 +1131,8 @@ AsyncTiledBlobRef Tex::ToBlob(int32 XTiles, int32 YTiles, uint32 Width /* = 0 */
 					BlobUPtr& Tile = TilesPtr[TileX * YTiles + TileY];
 					check(Tile);
 
+					//check(Tile->IsFinalised());
+
 					/// make sure that there's a Hash
 					CHashPtr TileHash = Tile->Hash();
 					check(TileHash && TileHash->IsValid() && TileHash->IsFinal());
diff --git a/Engine/Plugins/Experimental/TextureGraph/Source/MixerEngine/Data/Blob.cpp b/Engine/Plugins/Experimental/TextureGraph/Source/MixerEngine/Data/Blob.cpp
--- a/Engine/Plugins/Experimental/TextureGraph/Source/MixerEngine/Data/Blob.cpp
+++ b/Engine/Plugins/Experimental/TextureGraph/Source/MixerEngine/Data/Blob.cpp
@@ -195,7 +195,7 @@ AsyncDeviceBufferRef Blob::TransferTo(Device* TargetDevice)
 
 AsyncBufferResultPtr Blob::Bind(const BlobTransform* Transform, const ResourceBindInfo& BindInfo)
 {
-	/// Ok, we must have a raw buffer over here to transfer over to the device buffer
+	/// Ok, we must have a RawObj buffer over here to transfer over to the device buffer
 	//check(_buffer);
 
 	Device* Dev = BindInfo.Dev;
@@ -213,9 +213,9 @@ AsyncBufferResultPtr Blob::Bind(const BlobTransform* Transform, const ResourceBi
 	/// If the buffer isn't compatible then we transfer it over to the other new device
 	if (!Buffer->IsCompatible(Dev))
 	{
-		return Dev->Transfer(Buffer).then([this, Transform, BindInfo](DeviceBufferRef result)
+		return Dev->Transfer(Buffer).then([this, Transform, BindInfo](DeviceBufferRef Result)
 		{
-			Buffer = result;
+			Buffer = Result;
 			return Buffer->Bind(Transform, BindInfo);
 		});
 	}
@@ -238,27 +238,27 @@ bool Blob::IsValid() const
 AsyncRawBufferPtr Blob::Raw()
 {
 	/// Could possibly be going to another thread
-	/// Save the current temp hash
-	CHashPtr prevHash = Buffer->Hash(false);
-	check(!prevHash || prevHash->IsTemp());
+	/// Save the current temp Hash
+	CHashPtr PrevHash = Buffer->Hash(false);
+	check(!PrevHash || PrevHash->IsTemp());
 
-	return Buffer->Raw().then([this, prevHash](RawBufferPtr raw)
+	return Buffer->Raw().then([this, PrevHash](RawBufferPtr RawObj)
 	{
-		const BufferDescriptor& desc = Buffer->Descriptor();
-		CHashPtr hash = Buffer->Hash(false);
+		const BufferDescriptor& BufferDesc = Buffer->Descriptor();
+		CHashPtr Hash = Buffer->Hash(false);
 
-		if (!desc.bIsTransient && hash->IsFinal() && (!prevHash || !prevHash->IsFinal()))
+		if (!BufferDesc.bIsTransient && Hash->IsFinal() && (!PrevHash || !PrevHash->IsFinal()))
 			Buffer = Buffer->GetOwnerDevice()->AddInternal(Buffer);
 		else
 		{
-			HashType prevHashValue = prevHash ? prevHash->Value() : DataUtil::GNullHash;
-			UE_LOG(LogDevice, VeryVerbose, TEXT("DeviceBuffer has a new hash without owning reference. Unless this is manually cached by the device, this buffer will be deleted which is undesirable. Name: %s, Hash: %llu [Prev Hash: %llu, Size: %dx%d]"),
-				*desc.Name, prevHashValue, hash->Value(), desc.Width, desc.Height);
+			HashType PrevHashValue = PrevHash ? PrevHash->Value() : DataUtil::GNullHash;
+			UE_LOG(LogDevice, VeryVerbose, TEXT("DeviceBuffer has a new Hash without owning reference. Unless this is manually cached by the device, this buffer will be deleted which is undesirable. Name: %s, Hash: %llu [Prev Hash: %llu, Size: %dx%d]"),
+				*BufferDesc.Name, PrevHashValue, Hash->Value(), BufferDesc.Width, BufferDesc.Height);
 		}
 
 		/// TODO: do we really need this?
-		if (prevHash != nullptr)
-			MixerEngine::GetBlobber()->UpdateHash(prevHash->Value(), hash);
+		if (PrevHash != nullptr)
+			MixerEngine::GetBlobber()->UpdateHash(PrevHash->Value(), Hash);
 
 		return Buffer->Raw_Now();
 	});
@@ -266,12 +266,12 @@ AsyncRawBufferPtr Blob::Raw()
 
 AscynCHashPtr Blob::CalcHash()
 {
-	/// If we already have a have a hash for the buffer then don't bother
+	/// If we already have a have a Hash for the buffer then don't bother
 	if (Buffer)
 	{
-		CHashPtr bufferHash = Buffer->Hash(false);
-		if (bufferHash && bufferHash->IsFinal())
-			return cti::make_ready_continuable(bufferHash);
+		CHashPtr BufferHash = Buffer->Hash(false);
+		if (BufferHash && BufferHash->IsFinal())
+			return cti::make_ready_continuable(BufferHash);
 	}
 
 	return Raw().then([this] 
@@ -287,7 +287,7 @@ bool Blob::IsNull() const
 
 CHashPtr Blob::Hash() const
 {
-	/// Use own address as hash. This will still ensure that we detect object re-use
+	/// Use own address as Hash. This will still ensure that we detect object re-use
 	/// even if they're late bound or haven't been calculated yet
 	return Buffer->Hash(false);
 }
@@ -475,19 +475,65 @@ void Blob::SetLODLevel(int32 Level, BlobPtr LODBlob, BlobPtrW LODParentBlob, Blo
 	}
 }
 
-void Blob::Finalise_Now(bool bNoCalcHash, CHashPtr FixedHash)
+void Blob::UpdateLinkedBlobs(bool bDoFinalise)
+{
+	check(IsInGameThread());
+
+	for (BlobPtrW LinkedBlobW : LinkedBlobs)
+	{
+		BlobPtr LinkedBlob = LinkedBlobW.lock();
+
+		if (LinkedBlob && bDoFinalise)
+		{
+			LinkedBlob->FinaliseFrom(this);
+		}
+	}
+}
+
+void Blob::AddLinkedBlob(BlobPtr LinkedBlob)
+{
+	if (IsFinalised())
+	{
+		*LinkedBlob = *this;
+		return;
+	}
+	LinkedBlobs.push_back(LinkedBlob);
+}
+
+void Blob::FinaliseFrom(const Blob* RHS)
+{
+	Buffer = RHS->Buffer;
+	LODLevels = RHS->LODLevels;
+	MinMax = RHS->MinMax;
+	MinValue = RHS->MinValue;
+	MaxValue = RHS->MaxValue;
+	Histogram = RHS->Histogram;
+	LODParent = RHS->LODParent;
+	LODSource = RHS->LODSource;
+	bIsLODLevel = RHS->bIsLODLevel;
+
+	FinaliseNow(true, nullptr);
+}
+
+void Blob::FinaliseNow(bool bNoCalcHash, CHashPtr FixedHash)
 {
 	/// If already finalised then nothing to do over here
 	if (bIsFinalised)
+	{
 		return;
+	}
+
+	check(IsInGameThread());
 
 	bIsFinalised = true;
 	FinaliseTS = FDateTime::Now();
+
+	UpdateLinkedBlobs(true);
 }
 
 AsyncBufferResultPtr Blob::Finalise(bool bNoCalcHash, CHashPtr FixedHash)
 {
-	Finalise_Now(bNoCalcHash, FixedHash);
+	FinaliseNow(bNoCalcHash, FixedHash);
 	return cti::make_ready_continuable(std::make_shared<BufferResult>());
 }
 
@@ -508,7 +554,7 @@ CHashPtrVec Blob::CalculateMipHashes(CHashPtr MainHash, CHashPtr ParentHash, int
 	CHashPtr MainHashAtLOD = CalculateMipHash(MainHash, Level);
 	CHashPtr ParentHashAtLOD = CalculateMipHash(ParentHash, Level);
 
-	/// Just return one unique hash
+	/// Just return one unique Hash
 	if (MainHashAtLOD->Value() != ParentHashAtLOD->Value())
 		return { MainHashAtLOD, ParentHashAtLOD };
 
diff --git a/Engine/Plugins/Experimental/TextureGraph/Source/MixerEngine/Data/Blob.h b/Engine/Plugins/Experimental/TextureGraph/Source/MixerEngine/Data/Blob.h
--- a/Engine/Plugins/Experimental/TextureGraph/Source/MixerEngine/Data/Blob.h
+++ b/Engine/Plugins/Experimental/TextureGraph/Source/MixerEngine/Data/Blob.h
@@ -40,6 +40,7 @@ class MIXERENGINE_API Blob : public std::enable_shared_from_this<Blob>
 public:
 	static const char*				LODTransformName;	
 	typedef std::vector<BlobPtrW>	OwnerList;
+	typedef std::vector<std::weak_ptr<Blob>> LinkedBlobsVec;
 
 protected:
 	DeviceBufferRef					Buffer;							/// Device buffer that is attached to this blob
@@ -59,6 +60,13 @@ protected:
 	BlobPtrW						LODSource;						/// The original source of this blob, the mip Level 0
 
 	bool							bIsLODLevel = false;			/// Whether this blob already represents an LOD'd blob
+	LinkedBlobsVec					LinkedBlobs;					/// Blobs that are linked to this tiled blob. These are objects that are 
+																	/// uniquely created within the system but eventually resolve to the same
+																	/// tiled blob. We need to ensure that these are resolved correctly
+																	/// so that if some object is holding a pointer to a particular TiledBlob_Promise
+																	/// it gets the same view upon the completion.
+																	/// This is usually very uncommon with Blobs but more of a possibility with TiledBlob_Promise
+																	/// but we need to ensure that it works correctly for all blobs
 
 #if DEBUG_BLOB_REF_KEEPING == 1
 	OwnerList						_owners;						/// The TiledBlobs that contain this blob as a tile
@@ -72,6 +80,11 @@ protected:
 	virtual void					UpdateAccessInfo(uint64 batchId);
 	virtual void					SetHash(CHashPtr Hash);
 
+	virtual void					UpdateLinkedBlobs(bool bDoFinalise);
+	virtual void					AddLinkedBlob(BlobPtr LinkedBlob);
+
+	virtual void					FinaliseFrom(const Blob* RHS);
+
 public:
 									Blob(DeviceBufferRef InBuffer);
 									Blob(const BufferDescriptor& InDesc, CHashPtr InHash);
@@ -97,7 +110,7 @@ public:
 	virtual bool					IsTiled() const { return false; }
 	virtual bool					CanCalculateHash() const { return IsFinalised(); }
 
-	virtual void					Finalise_Now(bool bNoCalcHash, CHashPtr FixedHash);
+	virtual void					FinaliseNow(bool bNoCalcHash, CHashPtr FixedHash);
 	virtual AsyncBufferResultPtr	Finalise(bool bNoCalcHash, CHashPtr FixedHash);
 
 	virtual AsyncRawBufferPtr		Raw();			/// SLOW: Read doc for DeviceBuffer::Raw() in DeviceBuffer.h
diff --git a/Engine/Plugins/Experimental/TextureGraph/Source/MixerEngine/Data/Blobber.cpp b/Engine/Plugins/Experimental/TextureGraph/Source/MixerEngine/Data/Blobber.cpp
--- a/Engine/Plugins/Experimental/TextureGraph/Source/MixerEngine/Data/Blobber.cpp
+++ b/Engine/Plugins/Experimental/TextureGraph/Source/MixerEngine/Data/Blobber.cpp
@@ -721,9 +721,26 @@ TiledBlobRef Blobber::AddTiledResult(TiledBlobPtr Result, BlobCacheOptions Optio
 TiledBlobRef Blobber::AddTiledResult(CHashPtr LHash, TiledBlobPtr Result, BlobCacheOptions Options /* = {} */)
 {
 	BlobRef Ret = AddResult(LHash, std::static_pointer_cast<Blob>(Result), Options);
-	auto Ptr = Ret.get();
-	check(Ptr->IsTiled());
-	return TiledBlobRef(std::static_pointer_cast<TiledBlob>(Ptr), Ret.IsKeepStrong(), false);
+	check(Ret->IsTiled());
+	TiledBlobPtr CachedPtr = std::static_pointer_cast<TiledBlob>(Ret.get());
+
+	if (CachedPtr && Result != CachedPtr)
+	{
+		if (typeid(*CachedPtr) == typeid(*Result))
+			CachedPtr->AddLinkedBlob(Result);
+	}
+
+	/// If we got a different pointer back from the blobber then that means that this result was already cached into 
+	/// the system. We need to copy it over to the original pointer if it's a finalised blob
+	/// so that anyone keeping a copy of Result has the exact same view as the original result 
+	/// and we don't have do any awkward pointer adjustment shenanigans to make things work
+	//if (Result && CachedPtr->IsFinalised() && Result != CachedPtr)
+	//{
+	//	/// Copy the contents over
+	//	*Result = *CachedPtr;
+	//}
+
+	return TiledBlobRef(std::static_pointer_cast<TiledBlob>(CachedPtr), Ret.IsKeepStrong(), false);
 }
 
 BlobRef Blobber::AddResult(CHashPtr LHash, BlobRef Result, BlobCacheOptions Options)
diff --git a/Engine/Plugins/Experimental/TextureGraph/Source/MixerEngine/Data/TiledBlob.cpp b/Engine/Plugins/Experimental/TextureGraph/Source/MixerEngine/Data/TiledBlob.cpp
--- a/Engine/Plugins/Experimental/TextureGraph/Source/MixerEngine/Data/TiledBlob.cpp
+++ b/Engine/Plugins/Experimental/TextureGraph/Source/MixerEngine/Data/TiledBlob.cpp
@@ -57,7 +57,7 @@ std::shared_ptr<TiledBlob> TiledBlob::InitFromTiles(const BufferDescriptor& InDe
 	Desc.Width = std::max(Desc.Width, static_cast<uint32>(Tiles.Rows()));
 	Desc.Height = std::max(Desc.Height, static_cast<uint32>(Tiles.Cols()));
 	
-	auto tblob = std::make_shared<TiledBlob>(Desc, Tiles);
+	TiledBlobPtr TBlob = std::make_shared<TiledBlob>(Desc, Tiles);
 
 	for (size_t TileX = 0; TileX < Tiles.Rows(); TileX++)
 	{
@@ -71,7 +71,10 @@ std::shared_ptr<TiledBlob> TiledBlob::InitFromTiles(const BufferDescriptor& InDe
 		}
 	}
 
-	return tblob;
+	/// Mark as finalised
+	TBlob->bIsFinalised = true;
+
+	return TBlob;
 }
 
 TiledBlob::TiledBlob(BlobRef BlobObj)
@@ -120,6 +123,21 @@ TiledBlobPtr TiledBlob::AsTiledBlob(BlobPtr BlobObj)
 	return std::make_shared<TiledBlob>(BlobObj);
 }
 
+void TiledBlob::FinaliseFrom(const Blob* RHS)
+{
+	check(RHS->IsTiled());
+	const TiledBlob* RHSTiled = static_cast<const TiledBlob*>(RHS);
+
+	Tiles = RHSTiled->Tiles;
+	HashValue = RHSTiled->HashValue;
+	Desc = RHSTiled->Desc;
+	bReady = RHSTiled->bReady;
+	bTiledTarget = RHSTiled->bTiledTarget;
+	SingleBlob = RHSTiled->SingleBlob;
+
+	Blob::FinaliseFrom(RHS);
+}
+
 void TiledBlob::SetTransient()
 {
 	Desc.bIsTransient = true;
@@ -297,6 +315,8 @@ void TiledBlob::MakeSingleBlob_Internal()
 	}
 
 	Buffer = DeviceBufferRef();
+
+	UpdateLinkedBlobs(false);
 }
 
 AsyncBufferResultPtr TiledBlob::MakeSingleBlob()
@@ -858,6 +878,13 @@ FString TiledBlob::DisplayName() const
 	return Desc.Name;
 }
 
+void TiledBlob::AddLinkedBlob(BlobPtr LinkedBlob)
+{
+	check(IsInGameThread());
+	check(LinkedBlob->IsTiled());
+	Blob::AddLinkedBlob(LinkedBlob);
+}
+
 bool TiledBlob::CanCalculateHash() const
 {
 	check(IsInGameThread());
@@ -987,7 +1014,60 @@ AsyncBufferResultPtr TiledBlob_Promise::Unbind(const BlobTransform* Transform, c
 	return TiledBlob::Unbind(Transform, BindInfo);
 }
 
-void TiledBlob_Promise::Finalise_Now(bool NocalcHash, CHashPtr FixedHash)
+void TiledBlob_Promise::AddLinkedBlob(BlobPtr LinkedBlob)
+{
+	check(IsInGameThread());
+	check(LinkedBlob->IsTiled());
+
+	TiledBlobPtr TiledLinkedBlob = std::static_pointer_cast<TiledBlob>(LinkedBlob);
+	check(TiledLinkedBlob->IsPromise());
+
+	TiledBlob_PromisePtr TiledPromiseLinkedBlob = std::static_pointer_cast<TiledBlob_Promise>(LinkedBlob);
+
+	/// If we got a different pointer back from the blobber then that means that this result was already cached into 
+	/// the system. We need to copy it over to the original pointer if it's a finalised blob
+	/// so that anyone keeping a copy of Result has the exact same view as the original result 
+	/// and we don't have do any awkward pointer adjustment shenanigans to make things work
+	if (IsFinalised())
+	{
+		/// Copy the contents over
+		*TiledPromiseLinkedBlob = *this;
+	}
+
+	TiledPromiseLinkedBlob->CachedBlob = std::static_pointer_cast<TiledBlob_Promise>(shared_from_this());
+
+	TiledBlob::AddLinkedBlob(LinkedBlob);
+}
+
+//void TiledBlob_Promise::UpdateLinkedBlobs(bool bDoFinalise)
+//{
+//	check(IsInGameThread());
+//
+//	for (BlobPtrW LinkedBlobW : LinkedBlobs)
+//	{
+//		BlobPtr LinkedBlob = LinkedBlobW.lock();
+//
+//		if (LinkedBlob)
+//		{
+//			check(LinkedBlob->IsTiled());
+//			TiledBlobPtr TiledLinkedBlob = std::static_pointer_cast<TiledBlob>(LinkedBlob);
+//			check(TiledLinkedBlob->IsPromise());
+//
+//			TiledBlob_PromisePtr TiledPromiseLinkedBlob = std::static_pointer_cast<TiledBlob_Promise>(TiledLinkedBlob);
+//			*TiledPromiseLinkedBlob = *this;
+//
+//			/// We don't want to keep recursive links
+//			TiledPromiseLinkedBlob->LinkedBlobs.clear();
+//
+//			if (bDoFinalise)
+//			{
+//				TiledPromiseLinkedBlob->FinaliseNow(false, nullptr);
+//			}
+//		}
+//	}
+//}
+
+void TiledBlob_Promise::FinaliseNow(bool bNoCalcHash, CHashPtr FixedHash)
 {
 	check(IsInGameThread());
 
@@ -1000,12 +1080,14 @@ void TiledBlob_Promise::Finalise_Now(bool NocalcHash, CHashPtr FixedHash)
 		{
 			check(Tiles[TileX][TileY]);
 			BlobPtr Tile = Tiles[TileX][TileY];
-			Tile->Finalise_Now(NocalcHash, nullptr);
+			Tile->FinaliseNow(bNoCalcHash, nullptr);
 		}
 	}
 
 	/// Trigger the callbacks - Should always be the last step of finalize
 	NotifyCallbacks();
+
+	UpdateLinkedBlobs(true);
 }
 
 AsyncBufferResultPtr TiledBlob_Promise::FinaliseFrom(std::shared_ptr<TiledBlob_Promise> RHS)
@@ -1030,8 +1112,20 @@ AsyncBufferResultPtr TiledBlob_Promise::FinaliseFrom(std::shared_ptr<TiledBlob_P
 	return cti::make_ready_continuable<BufferResultPtr>(std::make_shared<BufferResult>());
 }
 
+void TiledBlob_Promise::FinaliseFrom(const Blob* RHS)
+{
+	check(RHS->IsTiled());
+	const TiledBlob* RHSTiled = static_cast<const TiledBlob*>(RHS);
+	check(RHSTiled->IsPromise());
 
-AsyncBufferResultPtr TiledBlob_Promise::Finalise(bool NocalcHash, CHashPtr FixedHash)
+	const TiledBlob_Promise* RHSTiledPromise = static_cast<const TiledBlob_Promise*>(RHS);
+
+	bMakeSingleBlob = RHSTiledPromise->bMakeSingleBlob;
+
+	TiledBlob::FinaliseFrom(RHS);
+}
+
+AsyncBufferResultPtr TiledBlob_Promise::Finalise(bool bNoCalcHash, CHashPtr FixedHash)
 {
 	check(IsInGameThread());
 
@@ -1046,15 +1140,15 @@ AsyncBufferResultPtr TiledBlob_Promise::Finalise(bool NocalcHash, CHashPtr Fixed
 		/// If this is not due to be turned into a single BlobObj then just return
 		if (!bMakeSingleBlob)
 		{
-			Finalise_Now(NocalcHash, FixedHash);
+			FinaliseNow(bNoCalcHash, FixedHash);
 			return cti::make_ready_continuable(std::make_shared<BufferResult>());
 		}
 
 		/// Otherwise we turn into a single BlobObj
 		return MakeSingleBlob()
-			.then([this, NocalcHash, FixedHash](BufferResultPtr SingleBuffer) {
+			.then([this, bNoCalcHash, FixedHash](BufferResultPtr SingleBuffer) {
 				if (bMakeSingleBlob)
-					Finalise_Now(NocalcHash, FixedHash);
+					FinaliseNow(bNoCalcHash, FixedHash);
 				
 				return SingleBuffer;
 			});
@@ -1065,7 +1159,7 @@ AsyncBufferResultPtr TiledBlob_Promise::Finalise(bool NocalcHash, CHashPtr Fixed
 
 	/// Otherwise we need to Tile the Buffer first
 	return TiledBlob::TileBuffer(Buffer, Tiles)
-		.then([this, NocalcHash, FixedHash](BufferResultPtr Result) mutable
+		.then([this, bNoCalcHash, FixedHash](BufferResultPtr Result) mutable
 		{
 			CHashPtrVec TileHashes(Tiles.Rows() * Tiles.Cols());
 
@@ -1076,7 +1170,7 @@ AsyncBufferResultPtr TiledBlob_Promise::Finalise(bool NocalcHash, CHashPtr Fixed
 					check(Tiles[TileX][TileY]);
 
 					BlobPtr Tile = Tiles[TileX][TileY];
-					if (!NocalcHash)
+					if (!bNoCalcHash)
 						Tile->CalcHash();
 					else
 					{
@@ -1097,7 +1191,7 @@ AsyncBufferResultPtr TiledBlob_Promise::Finalise(bool NocalcHash, CHashPtr Fixed
 			bTiledTarget = true;
 
 			/// Tiled blobs always recalculate their has to match those of the children
-			Finalise_Now(true, fullHash);
+			FinaliseNow(true, fullHash);
 
 			return Result;
 		});
@@ -1140,6 +1234,11 @@ void TiledBlob_Promise::SetTile(int32 TileX, int32 TileY, BlobRef InTile)
 
 AsyncBufferResultPtr TiledBlob_Promise::MakeSingleBlob()
 {
+	if (!CachedBlob.expired())
+	{
+		return CachedBlob.lock()->MakeSingleBlob();
+	}
+
 	/// If already finalised then return
 	if (bIsFinalised)
 		return TiledBlob::MakeSingleBlob();
diff --git a/Engine/Plugins/Experimental/TextureGraph/Source/MixerEngine/Data/TiledBlob.h b/Engine/Plugins/Experimental/TextureGraph/Source/MixerEngine/Data/TiledBlob.h
--- a/Engine/Plugins/Experimental/TextureGraph/Source/MixerEngine/Data/TiledBlob.h
+++ b/Engine/Plugins/Experimental/TextureGraph/Source/MixerEngine/Data/TiledBlob.h
@@ -36,6 +36,9 @@ protected:
 	TiledBlob&						operator = (const TiledBlob& RHS);
 	virtual void					SetHash(CHashPtr Hash) override;
 
+	virtual void					AddLinkedBlob(BlobPtr LinkedBlob) override;
+	virtual void					FinaliseFrom(const Blob* RHS) override;
+
 public:
 	static AsyncBufferResultPtr		TileBuffer(DeviceBufferRef Buffer, BlobPtrTiles& Tiles);
 	AsyncBufferResultPtr			CombineTiles(bool bTouch, bool bIsArray, uint64 BatchId = 0);
@@ -153,11 +156,15 @@ class MIXERENGINE_API TiledBlob_Promise : public TiledBlob
 protected:
 	mutable std::vector<BlobReadyCallback> Callbacks;		/// The callbacks to call when the promised tiled BlobObj is ready
 	bool							bMakeSingleBlob = false;/// Make single BlobObj on finalise
+	std::weak_ptr<TiledBlob_Promise> CachedBlob;			/// The cached blob that this is derived off
 
 	virtual void					OnFinaliseInternal(BlobReadyCallback Callback) const override;
 	TiledBlob_Promise&				operator = (const TiledBlob_Promise& RHS);
 	virtual void					NotifyCallbacks();
 
+	virtual void					AddLinkedBlob(BlobPtr LinkedBlob) override;
+	virtual void					FinaliseFrom(const Blob* RHS) override;
+
 public:
 									TiledBlob_Promise(TiledBlobPtr Source);
 									TiledBlob_Promise(const BufferDescriptor& Desc, size_t NumTilesX, size_t NumTilesY, CHashPtr HashValue);
@@ -173,8 +180,8 @@ public:
 	virtual void					CopyResolveLateBound(BlobPtr RHS) override;
 	virtual bool					IsPromise() const override { return true; }
 
-	virtual void					Finalise_Now(bool NocalcHash, CHashPtr FixedHash) override;
-	virtual AsyncBufferResultPtr	Finalise(bool NocalcHash, CHashPtr FixedHash) override;
+	virtual void					FinaliseNow(bool bNoCalcHash, CHashPtr FixedHash) override;
+	virtual AsyncBufferResultPtr	Finalise(bool bNoCalcHash, CHashPtr FixedHash) override;
 	AsyncBufferResultPtr			FinaliseFrom(std::shared_ptr<TiledBlob_Promise> RHS);
 
 	void							ResetForReplay(); /// For debug purpose, reset the state of the tile as a promise to NOT finalised, increment the replayCount
diff --git a/Engine/Plugins/Experimental/TextureGraph/Source/MixerEngine/Device/Device.cpp b/Engine/Plugins/Experimental/TextureGraph/Source/MixerEngine/Device/Device.cpp
--- a/Engine/Plugins/Experimental/TextureGraph/Source/MixerEngine/Device/Device.cpp
+++ b/Engine/Plugins/Experimental/TextureGraph/Source/MixerEngine/Device/Device.cpp
@@ -714,10 +714,10 @@ AsyncDeviceBufferRef Device::CombineFromTiles(const CombineSplitArgs& CombineArg
 	return cti::make_ready_continuable(Buffer);
 }
 
-AsyncDeviceBufferRef Device::SplitToTiles_Generic(const CombineSplitArgs& splitArgs)
+AsyncDeviceBufferRef Device::SplitToTiles_Generic(const CombineSplitArgs& SplitArgs)
 {
-	const T_Tiles<DeviceBufferRef>& Tiles = splitArgs.Tiles;
-	auto Buffer = splitArgs.Buffer;
+	const T_Tiles<DeviceBufferRef>& Tiles = SplitArgs.Tiles;
+	auto Buffer = SplitArgs.Buffer;
 
 	RawBufferPtrTiles RawTileBuffers(0, 0);
 	RawBufferPtr Raw = Buffer->Raw_Now();
@@ -739,18 +739,18 @@ AsyncDeviceBufferRef Device::SplitToTiles_Generic(const CombineSplitArgs& splitA
 	{
 		for (int32 TileY = 0; TileY < Tiles.Cols(); TileY++)
 		{
-			DeviceBufferRef tile = Tiles[TileX][TileY];
-			RawBufferPtr rawTile = RawTileBuffers[TileX][TileY];
-			check(rawTile);
+			DeviceBufferRef Tile = Tiles[TileX][TileY];
+			RawBufferPtr RawTile = RawTileBuffers[TileX][TileY];
+			check(RawTile);
 
-			AsyncBufferResultPtr tilePromise = tile->UpdateRaw(rawTile);
-			Promises.push_back(std::move(tilePromise));
+			AsyncBufferResultPtr TilePromise = Tile->UpdateRaw(RawTile);
+			Promises.push_back(std::move(TilePromise));
 		}
 	}
 
 	return cti::when_all(Promises.begin(), Promises.end()).then([Buffer]() mutable
 	{
-		UE_LOG(LogDevice, Log, TEXT("SplitTiles_Generic tile updates finished!"));
+		UE_LOG(LogDevice, Log, TEXT("SplitTiles_Generic Tile updates finished!"));
 		return Buffer;
 	});
 }
diff --git a/Engine/Plugins/Experimental/TextureGraph/Source/MixerEngine/Export/T_ExportAsUAsset.cpp b/Engine/Plugins/Experimental/TextureGraph/Source/MixerEngine/Export/T_ExportAsUAsset.cpp
--- a/Engine/Plugins/Experimental/TextureGraph/Source/MixerEngine/Export/T_ExportAsUAsset.cpp
+++ b/Engine/Plugins/Experimental/TextureGraph/Source/MixerEngine/Export/T_ExportAsUAsset.cpp
@@ -40,7 +40,7 @@ cti::continuable<int32> Job_ExportAsUAsset::PreExecAsync(ENamedThreads::Type exe
 	auto hash = Job::Hash();
 	BlobPtr cachedResult = MixerEngine::GetBlobber()->Find(hash->Value());
 
-	//TODO: Need to cache the result so that the job do not execute every time with matching arguments
+	//TODO: Need to cache the Result so that the job do not execute every time with matching arguments
 	if (cachedResult)
 	{
 		bIsCulled = true;
@@ -69,46 +69,48 @@ cti::continuable<int32> Job_ExportAsUAsset::ExecAsync(ENamedThreads::Type execTh
 	RawBufferPtr rawPtr;
 
 	return Job::BeginNative(RunInfo)
-	.then([]()
-	{
-		return PromiseUtil::OnGameThread();
-	})
-	.then([this](int32)
-	{
-		return Setting.Map->CombineTiles(false,false);
-	})
-	.then([this](BufferResultPtr bufferPtr)
-	{
+		.then([]()
+		{
+			return PromiseUtil::OnGameThread();
+		})
+		.then([this](int32)
+		{
+			return Setting.Map->CombineTiles(false,false);
+		})
+		.then([this](BufferResultPtr bufferPtr)
+		{
 			UE_LOG(LogData, Log, TEXT("[Export As UAsset - %s] Calling Export ..."), *Setting.Name.ToString());
-			DeviceBufferRef buffer = Setting.Map->GetBufferRef();
-
-		return buffer->Raw();
-	})
-	.then([this](RawBufferPtr raw)
-	{
-		FName fileName = Setting.Name;
-
-		check(raw->GetData());
-
-		UE_LOG(LogData, Display, TEXT("   - Writing raw buffer to filename: %s"), *fileName.ToString());
-		return TextureExporter::ExportRawAsUAsset(raw, Setting, OutFolder, fileName);
-	})
-	.then([this](int32 result) mutable
-	{
-		Setting.OnDone.ExecuteIfBound(Setting);
-		//promise.set_value(0);
-		UE_LOG(LogData, Display, TEXT("Exporting finished"));
-		EndNative();
-		SetPromise(0);
-		return 0;
-	})
-	.fail([this](std::exception_ptr e) mutable
-	{
-		UE_LOG(LogData, Log, TEXT("[Exporting As UAsset - %s] Promise failure!"), *Setting.Name.ToString());
-		Setting.OnDone.ExecuteIfBound(Setting);
-		EndNative();
-		return -1;
-	});
+			DeviceBufferRef Buffer = Setting.Map->GetBufferRef();
+			return Buffer->Raw();
+		})
+		.then([this](RawBufferPtr RawObj)
+		{
+			FName fileName = Setting.Name;
+
+			check(RawObj->GetData());
+
+			UE_LOG(LogData, Display, TEXT("   - Writing RawObj Buffer to filename: %s"), *fileName.ToString());
+			return TextureExporter::ExportRawAsUAsset(RawObj, Setting, OutFolder, fileName);
+		})
+		.then([this](int32) mutable
+		{
+			Setting.OnDone.ExecuteIfBound(Setting);
+			//promise.set_value(0);
+			UE_LOG(LogData, Display, TEXT("Exporting finished"));
+			return EndNative();
+		})
+		.then([this](JobResultPtr)
+		{
+			SetPromise(0);
+			return 0;
+		})
+		.fail([this](std::exception_ptr e) mutable
+		{
+			UE_LOG(LogData, Log, TEXT("[Exporting As UAsset - %s] Promise failure!"), *Setting.Name.ToString());
+			Setting.OnDone.ExecuteIfBound(Setting);
+			EndNative();
+			return -1;
+		});
 #else
 	return cti::make_ready_continuable(0);
 #endif 
diff --git a/Engine/Plugins/Experimental/TextureGraph/Source/MixerEngine/Helper/DataUtil.cpp b/Engine/Plugins/Experimental/TextureGraph/Source/MixerEngine/Helper/DataUtil.cpp
--- a/Engine/Plugins/Experimental/TextureGraph/Source/MixerEngine/Helper/DataUtil.cpp
+++ b/Engine/Plugins/Experimental/TextureGraph/Source/MixerEngine/Helper/DataUtil.cpp
@@ -359,8 +359,6 @@ void CHash::UpdateLinks()
 
 CHashPtr CHash::UpdateHash(CHashPtr NewHash, CHashPtr PrevHash)
 {
-	check(NewHash->IsFinal());
-
 	if (PrevHash)
 	{
 		check(NewHash != PrevHash);
diff --git a/Engine/Plugins/Experimental/TextureGraph/Source/TextureGraphEditor/Private/STG_OutputSelectionDlg.cpp b/Engine/Plugins/Experimental/TextureGraph/Source/TextureGraphEditor/Private/STG_OutputSelectionDlg.cpp
--- a/Engine/Plugins/Experimental/TextureGraph/Source/TextureGraphEditor/Private/STG_OutputSelectionDlg.cpp
+++ b/Engine/Plugins/Experimental/TextureGraph/Source/TextureGraphEditor/Private/STG_OutputSelectionDlg.cpp
@@ -87,12 +87,19 @@ void STG_OutputSelectionDlg::AddExportItems()
 			{
 				UTG_EdGraphNode* EdNode = EdGraph->GetViewModelNode(Node->GetId());
 				TiledBlobPtr ThumbBlob = EdNode->GetCachedThumbBlob(Id);
+
+				if (!ThumbBlob)
+				{
+					ThumbBlob = TextureHelper::GetBlack();
+				}
+
 				TSharedPtr<STG_NodeThumbnail> NodeThumbnail = SNew(STG_NodeThumbnail);
 				ThumbBlob->OnFinalise()
-				.then([ThumbBlob, NodeThumbnail]
-				{
-					NodeThumbnail->UpdateBlob(ThumbBlob);
-				});
+					.then([ThumbBlob, NodeThumbnail]
+					{
+						NodeThumbnail->UpdateBlob(ThumbBlob);
+					});
+
 				ThumbnailWidget = NodeThumbnail;
 			}
 			// else if color do something
