Property Ptr needed a null check

diff --git a/Engine/Plugins/Experimental/TextureGraph/Source/TextureGraphEditor/Private/EdGraph/TG_EdGraphNode.cpp b/Engine/Plugins/Experimental/TextureGraph/Source/TextureGraphEditor/Private/EdGraph/TG_EdGraphNode.cpp
--- a/Engine/Plugins/Experimental/TextureGraph/Source/TextureGraphEditor/Private/EdGraph/TG_EdGraphNode.cpp
+++ b/Engine/Plugins/Experimental/TextureGraph/Source/TextureGraphEditor/Private/EdGraph/TG_EdGraphNode.cpp
@@ -665,8 +665,12 @@ bool UTG_EdGraphNode::UpdateEdPinDefaultValue(UEdGraphPin* EdPin, const UTG_EdGr
 				DefaultValue = Enum->GetNameByValue(FCString::Atoi(*DefaultValue)).ToString();
 			}
 		}
+		
+#if WITH_EDITOR
 		// check here if the property has a tag for regen pins and value has changed
-		bShouldUpdatePinsVisibility = Property->HasMetaData("RegenPinsOnChange") && EdPin->DefaultValue != DefaultValue;
+		bShouldUpdatePinsVisibility = (Property != nullptr && Property->HasMetaData("RegenPinsOnChange"))
+										&& EdPin->DefaultValue != DefaultValue;
+#endif
 		
 		EdPin->bDefaultValueIsReadOnly = false;
 		EdPin->DefaultValue = DefaultValue;
