Fix for Texture Graph Undo Redo Issues
Also added the tests for undo redo

diff --git a/Engine/Plugins/Experimental/TextureGraph/Source/TextureGraph/Private/Expressions/CommandChange/TG_PinConnectionChange.cpp b/Engine/Plugins/Experimental/TextureGraph/Source/TextureGraph/Private/Expressions/CommandChange/TG_PinConnectionChange.cpp
new file mode 100644
--- /dev/null
+++ b/Engine/Plugins/Experimental/TextureGraph/Source/TextureGraph/Private/Expressions/CommandChange/TG_PinConnectionChange.cpp
@@ -0,0 +1,71 @@
+// Copyright Epic Games, Inc. All Rights Reserved.
+
+#include "Expressions/CommandChange/TG_PinConnectionChange.h"
+#include "TG_Pin.h"
+#include "TG_Node.h"
+#include "TG_Graph.h"
+#include "Misc/ITransaction.h"
+
+void FTG_PinConnectionChange::StoreChange(UTG_Graph* InGraph, UTG_Node& NodeFrom, FTG_Name& PinFromName)
+{
+#if WITH_EDITOR
+	if (!GIsTransacting)
+	{
+		TUniquePtr<FTG_PinConnectionChange> Change = MakeUnique<FTG_PinConnectionChange>(InGraph, NodeFrom, PinFromName);
+		GUndo->StoreUndo(InGraph, MoveTemp(Change));
+	}
+#endif
+}
+
+FTG_PinConnectionChange::FTG_PinConnectionChange(UTG_Graph* InGraph, UTG_Node& NodeTo, FTG_Name& PinToName)
+	: FCommandChange()
+	, Graph(InGraph)
+{
+	PinTo = NodeTo.GetPin(PinToName);
+}
+
+void FTG_PinConnectionChange::Revert(UObject* Object)
+{ // undo
+	check(PinTo);
+
+	//If the pin is not conneccted we will copy its self var to the expression property
+	if(!PinTo->IsConnected())
+		PinTo->GetNodePtr()->OnPinConnectionUndo(PinTo->GetId());
+}
+
+///////////////////////////////////////////////////////////////////////////////////////////////////////////////
+
+void FTG_PinConnectionBreakChange::StoreChange(UTG_Graph* InGraph, UTG_Node& NodeFrom, FTG_Name& PinFromName)
+{
+#if WITH_EDITOR
+	if (!GIsTransacting)
+	{
+		TUniquePtr<FTG_PinConnectionBreakChange> Change = MakeUnique<FTG_PinConnectionBreakChange>(InGraph, NodeFrom, PinFromName);
+		GUndo->StoreUndo(InGraph, MoveTemp(Change));
+	}
+#endif
+}
+
+FTG_PinConnectionBreakChange::FTG_PinConnectionBreakChange(UTG_Graph* InGraph, UTG_Node& NodeTo, FTG_Name& PinToName)
+	: FCommandChange()
+	, Graph(InGraph)
+{
+	PinTo = NodeTo.GetPin(PinToName);
+}
+
+void FTG_PinConnectionBreakChange::Apply(UObject* Object)
+{ // undo
+	check(PinTo);
+
+	//we will rest its self var and copy to the expression property
+	for (auto OtherPinId : PinTo->GetEdges())
+	{
+		auto OtherPin = Graph->GetPin(OtherPinId);
+		if (OtherPin)
+		{
+			OtherPin->GetNodePtr()->OnPinConnectionUndo(OtherPin->GetId());
+		}
+	}
+
+	PinTo->GetNodePtr()->OnPinConnectionUndo(PinTo->GetId());
+}
\ No newline at end of file
diff --git a/Engine/Plugins/Experimental/TextureGraph/Source/TextureGraph/Private/TG_Node.cpp b/Engine/Plugins/Experimental/TextureGraph/Source/TextureGraph/Private/TG_Node.cpp
--- a/Engine/Plugins/Experimental/TextureGraph/Source/TextureGraph/Private/TG_Node.cpp
+++ b/Engine/Plugins/Experimental/TextureGraph/Source/TextureGraph/Private/TG_Node.cpp
@@ -187,6 +187,33 @@ void UTG_Node::OnPinConnectionChanged(FTG_Id InPinId, FTG_Id OldPinId, FTG_Id Ne
 		FTG_Variant::EType CommonType = EvalExpressionCommonVariantType();
 		GetExpression()->NotifyCommonVariantTypeChanged(CommonType);
 	}
+
+	if (!ThePin->IsConnected())
+	{
+		if (ThePin->IsArgTexture())
+		{
+			//Reset the FTGTexture
+			ThePin->EditSelfVar()->EditAs<FTG_Texture>() = nullptr;
+			GetExpression()->CopyVarToExpressionArgument(ThePin->GetArgument(), ThePin->EditSelfVar());
+		}
+	}
+}
+
+void UTG_Node::OnPinConnectionUndo(FTG_Id InPinId)
+{
+	UTG_Pin* PinTo = GetPin(InPinId);
+	// Reset the Pin to Default in this case
+	if (PinTo->IsArgVariant())
+	{
+		FTG_Variant::EType CommonType = FTG_Variant::EType::Scalar;
+		GetExpression()->NotifyCommonVariantTypeChanged(CommonType);
+	}
+	else if (PinTo->IsArgTexture())
+	{
+		//Reset the FTGTexture
+		PinTo->EditSelfVar()->EditAs<FTG_Texture>() = nullptr;
+		GetExpression()->CopyVarToExpressionArgument(PinTo->GetArgument(), PinTo->EditSelfVar());
+	}
 }
 
 
diff --git a/Engine/Plugins/Experimental/TextureGraph/Source/TextureGraph/Private/TG_Pin.cpp b/Engine/Plugins/Experimental/TextureGraph/Source/TextureGraph/Private/TG_Pin.cpp
--- a/Engine/Plugins/Experimental/TextureGraph/Source/TextureGraph/Private/TG_Pin.cpp
+++ b/Engine/Plugins/Experimental/TextureGraph/Source/TextureGraph/Private/TG_Pin.cpp
@@ -210,6 +210,13 @@ void UTG_Pin::PostEditChangeProperty(FPropertyChangedEvent& PropertyChangedEvent
 		SetAliasName(AliasName);
 	}
 }
+
+bool UTG_Pin::Modify(bool bAlwaysMarkDirty)
+{
+	//Calling the Pin expression Modify here to make sure we have the snapshot of data
+	GetNodePtr()->GetExpression()->Modify();
+	return Super::Modify(bAlwaysMarkDirty);
+}
 #endif
 
 FProperty* UTG_Pin::GetExpressionProperty() const
diff --git a/Engine/Plugins/Experimental/TextureGraph/Source/TextureGraph/Public/Expressions/CommandChange/TG_PinConnectionChange.h b/Engine/Plugins/Experimental/TextureGraph/Source/TextureGraph/Public/Expressions/CommandChange/TG_PinConnectionChange.h
new file mode 100644
--- /dev/null
+++ b/Engine/Plugins/Experimental/TextureGraph/Source/TextureGraph/Public/Expressions/CommandChange/TG_PinConnectionChange.h
@@ -0,0 +1,55 @@
+// Copyright Epic Games, Inc. All Rights Reserved.
+
+#pragma once
+
+#include "CoreMinimal.h"
+#include "TG_SystemTypes.h"  
+#include "Misc/Change.h"
+
+class UTG_Graph;
+class UTG_Node;
+class UTG_Pin;
+
+/** This class helps to Undo the pin connection correctly
+as the Blob is not serialized so if the expression property is not connected and is holding 
+a reference of blob it is not rest when undoing the connection*/
+class TEXTUREGRAPH_API FTG_PinConnectionChange : public FCommandChange 
+{
+public:
+	static void StoreChange(UTG_Graph* InGraph, UTG_Node& NodeFrom, FTG_Name& PinFromName);
+	FTG_PinConnectionChange(UTG_Graph* InGraph, UTG_Node& NodeFrom, FTG_Name& PinFromName);
+
+	virtual FString ToString() const override
+	{
+		return FString(TEXT("Pin Connection"));
+	}
+
+	virtual void Apply(UObject* Object) override {};
+	virtual void Revert(UObject* Object) override;
+
+private:
+	UTG_Graph* Graph;
+	UTG_Pin* PinTo;
+};
+
+/** This class helps to Redo the pin connection Break correctly
+as the Blob is not serialized so if the expression property is not connected and is holding
+a reference of blob it is not rest when redoing the connection*/
+class TEXTUREGRAPH_API FTG_PinConnectionBreakChange : public FCommandChange
+{
+public:
+	static void StoreChange(UTG_Graph* InGraph, UTG_Node& NodeFrom, FTG_Name& PinFromName);
+	FTG_PinConnectionBreakChange(UTG_Graph* InGraph, UTG_Node& NodeFrom, FTG_Name& PinFromName);
+
+	virtual FString ToString() const override
+	{
+		return FString(TEXT("Pin Connection Break"));
+	}
+
+	virtual void Apply(UObject* Object) override;
+	virtual void Revert(UObject* Object) override {};
+
+private:
+	UTG_Graph* Graph;
+	UTG_Pin* PinTo;
+};
diff --git a/Engine/Plugins/Experimental/TextureGraph/Source/TextureGraph/Public/TG_Node.h b/Engine/Plugins/Experimental/TextureGraph/Source/TextureGraph/Public/TG_Node.h
--- a/Engine/Plugins/Experimental/TextureGraph/Source/TextureGraph/Public/TG_Node.h
+++ b/Engine/Plugins/Experimental/TextureGraph/Source/TextureGraph/Public/TG_Node.h
@@ -170,6 +170,8 @@ public:
 	// only valid for expression with FTG_Variant arguments
 	FTG_Variant::EType GetExpressionCommonVariantType() const;
 
+	void OnPinConnectionUndo(FTG_Id InPinId);
+
 protected:
 
 	// Return the array of alias name of all the pins
diff --git a/Engine/Plugins/Experimental/TextureGraph/Source/TextureGraph/Public/TG_Pin.h b/Engine/Plugins/Experimental/TextureGraph/Source/TextureGraph/Public/TG_Pin.h
--- a/Engine/Plugins/Experimental/TextureGraph/Source/TextureGraph/Public/TG_Pin.h
+++ b/Engine/Plugins/Experimental/TextureGraph/Source/TextureGraph/Public/TG_Pin.h
@@ -91,6 +91,9 @@ protected:
 public:
 	// Override serialization for conditional serialization of SelfVar
 	virtual	void		Serialize(FArchive& Ar) override;
+#if WITH_EDITOR
+	virtual bool		Modify(bool bAlwaysMarkDirty = true) override;
+#endif
 
 	bool				IsValid() const { return Id != FTG_Id::INVALID; }
 
diff --git a/Engine/Plugins/Experimental/TextureGraph/Source/TextureGraphEditor/Private/EdGraph/TG_EdGraphSchema.cpp b/Engine/Plugins/Experimental/TextureGraph/Source/TextureGraphEditor/Private/EdGraph/TG_EdGraphSchema.cpp
--- a/Engine/Plugins/Experimental/TextureGraph/Source/TextureGraphEditor/Private/EdGraph/TG_EdGraphSchema.cpp
+++ b/Engine/Plugins/Experimental/TextureGraph/Source/TextureGraphEditor/Private/EdGraph/TG_EdGraphSchema.cpp
@@ -19,6 +19,7 @@
 #include "Expressions/Input/TG_Expression_Material.h"
 #include "Expressions/Input/TG_Expression_MaterialFunction.h"
 #include "Expressions/Input/TG_Expression_Texture.h"
+#include "Expressions/CommandChange/TG_PinConnectionChange.h"
 
 #define LOCTEXT_NAMESPACE "TG_GraphSchema"
 
@@ -276,6 +277,8 @@ bool UTG_EdGraphSchema::TryCreateConnection(UEdGraphPin* InA, UEdGraphPin* InB)
 		UTG_Graph* TG_Graph = TG_NodeA->GetGraph();
 		check(TG_Graph);
 
+		FTG_PinConnectionChange::StoreChange(TG_Graph, *TG_NodeB, B->PinName);
+
 		auto edge = TG_Graph->Connect(*TG_NodeA, A->PinName, *TG_NodeB, B->PinName);
 	}
 
@@ -328,6 +331,8 @@ void UTG_EdGraphSchema::BreakPinLinks(UEdGraphPin& TargetPin, bool bSendsNodeNot
 	UTG_Graph* TextureGraph = TSNode->GetGraph();
 	check(TextureGraph);
 
+	FTG_PinConnectionBreakChange::StoreChange(TextureGraph, *TSNode, TargetPin.PinName);
+
 	TextureGraph->RemovePinEdges(*TSNode, TargetPin.PinName);
 }
 
@@ -349,6 +354,8 @@ void UTG_EdGraphSchema::BreakSinglePinLink(UEdGraphPin* SourcePin, UEdGraphPin*
 
 	UTG_Graph* TextureGraph = SourceTSNode->GetGraph();
 	check(TextureGraph);
+
+	FTG_PinConnectionBreakChange::StoreChange(TextureGraph, *TargetTSNode, TargetPin->PinName);
 	TextureGraph->RemoveEdge(*SourceTSNode, SourcePin->PinName, *TargetTSNode, TargetPin->PinName);
 }
 
@@ -421,6 +428,9 @@ void UTG_EdGraphSchema::DroppedAssetsOnNode(const TArray<FAssetData>& Assets, co
 
 		if(Expression->CanHandleAsset(Asset))
 		{
+			const FScopedTransaction Transaction(NSLOCTEXT("UTG_EdGraphSchema", "DroppedAssetsOnNode", "Dropped Assets OnNode"));
+			Expression->Modify();
+
 			Expression->SetAsset(Asset);
 		}
 	}
diff --git a/Engine/Plugins/Experimental/TextureGraph/Source/TextureGraphEditor/Private/Pins/STG_TextureDescriptor.cpp b/Engine/Plugins/Experimental/TextureGraph/Source/TextureGraphEditor/Private/Pins/STG_TextureDescriptor.cpp
--- a/Engine/Plugins/Experimental/TextureGraph/Source/TextureGraphEditor/Private/Pins/STG_TextureDescriptor.cpp
+++ b/Engine/Plugins/Experimental/TextureGraph/Source/TextureGraphEditor/Private/Pins/STG_TextureDescriptor.cpp
@@ -13,6 +13,7 @@
 #include "SGraphPinComboBox.h"
 #include "FrameWork/MultiBox/MultiBoxBuilder.h"
 #include "TG_Graph.h"
+#include "Expressions/TG_Expression.h"
 #include "EDGraph/TG_EdGraphSchema.h"
 #include "Widgets/Layout/SSeparator.h"
 #include "TG_HelperFunctions.h"
@@ -181,13 +182,12 @@ FTG_TextureDescriptor STG_TextureDescriptor::GetTextureDescriptor() const
 
 void STG_TextureDescriptor::SetValue(FTG_TextureDescriptor Descriptor)
 {
+	const FScopedTransaction Transaction(NSLOCTEXT("STG_TextureDescriptor", "SetValue", "Output Settings Pin Changed"));
+
 	const UTG_EdGraphSchema* Schema = Cast<const UTG_EdGraphSchema>(GraphPinObj->GetSchema());
 	UTG_Pin* TGPin = Schema->GetTGPinFromEdPin(GraphPinObj);
 	check(TGPin->GetArgument().IsTexture());
-	//FTG_Texture OutTexture;
-	//TGPin->GetValue(OutTexture);
-	
-	//OutTexture.Descriptor = Descriptor;
+	TGPin->Modify();
 
 	TGPin->SetValue(Descriptor);
 
diff --git a/Engine/Source/Editor/GraphEditor/Public/KismetPins/SGraphPinNumSlider.h b/Engine/Source/Editor/GraphEditor/Public/KismetPins/SGraphPinNumSlider.h
--- a/Engine/Source/Editor/GraphEditor/Public/KismetPins/SGraphPinNumSlider.h
+++ b/Engine/Source/Editor/GraphEditor/Public/KismetPins/SGraphPinNumSlider.h
@@ -105,6 +105,9 @@ protected:
 	*/
 	void OnBeginSliderMovement()
 	{
+		GEditor->BeginTransaction(NSLOCTEXT("GraphEditor", "ChangeNumberPinValueSlider", "Change Number Pin Value slider"));
+		GraphPinObj->Modify();
+
 		SliderValue = LastSliderCommittedValue = GetNumericValue().GetValue();
 		bIsUsingSlider = true;
 	}
@@ -115,6 +118,7 @@ protected:
 	void OnEndSliderMovement( NumericType NewValue )
 	{
 		bIsUsingSlider = false;
+		GEditor->EndTransaction();
 	}
 	
 	TOptional<NumericType> GetNumericValue() const
