Fixed Weird tooltip

Fix/How: In the case of transient objects, tooltips use UObject NamePrivate and that name needs to be unique. On further investigation concluded that we dont need a transient for the EdGraph node. The variable is declared in TG_Editor and it should work fine without transient flag.

diff --git a/Engine/Plugins/Experimental/TextureGraph/Source/TextureGraphEditor/Private/TG_Editor.cpp b/Engine/Plugins/Experimental/TextureGraph/Source/TextureGraphEditor/Private/TG_Editor.cpp
--- a/Engine/Plugins/Experimental/TextureGraph/Source/TextureGraphEditor/Private/TG_Editor.cpp
+++ b/Engine/Plugins/Experimental/TextureGraph/Source/TextureGraphEditor/Private/TG_Editor.cpp
@@ -187,7 +187,7 @@ void FTG_Editor::InitEditor(const EToolkitMode::Type Mode, const TSharedPtr< cla
 
 	MixerEngine::RegisterErrorReporter(EditedTextureGraph, std::make_shared<FTG_EditorErrorReporter>(this));
 
-	TG_EdGraph = NewObject<UTG_EdGraph>(EditedTextureGraph, UTG_EdGraph::StaticClass(), NAME_None, RF_Transactional | RF_Transient);
+	TG_EdGraph = NewObject<UTG_EdGraph>(EditedTextureGraph, UTG_EdGraph::StaticClass(), NAME_None, RF_Transactional);
 	TG_EdGraph->Schema = UTG_EdGraphSchema::StaticClass();
 
 	TG_EdGraph->InitializeFromTextureGraph(EditedTextureGraph, SharedThis(this));
