Fix Normal looking washed out

Linear textures are now being displayed in linear gamma. Turning off the conversion to match UE's convention.

#virtualized

diff --git a/Engine/Plugins/Experimental/TextureGraph/Source/MixerEngine/2D/TargetTextureSet.cpp b/Engine/Plugins/Experimental/TextureGraph/Source/MixerEngine/2D/TargetTextureSet.cpp
--- a/Engine/Plugins/Experimental/TextureGraph/Source/MixerEngine/2D/TargetTextureSet.cpp
+++ b/Engine/Plugins/Experimental/TextureGraph/Source/MixerEngine/2D/TargetTextureSet.cpp
@@ -132,8 +132,13 @@ void TargetTextureSet::BindOnTextureUpdate(RenderMaterial_BPPtr InMaterial, FMat
 				{
 					BindInfo.Target = MaterialMappingInfo.MaterialInput.ToString();
 					//TODO: Need to move it into the Resource Bind Info
-					InMaterial->SetInt(TEXT("sRGB"), texture.get()->GetDescriptor().bIsSRGB ? 1 : 0);
+
+					// Linear textures are now being displayed in linear gamma.
+					// So turning off the conversion to match UE's convention.
+					// InMaterial->SetInt(TEXT("sRGB"), texture.get()->GetDescriptor().bIsSRGB ? 1 : 0);
+
 					InMaterial->SetInt(TEXT("IsGrayscale"), texture.get()->GetDescriptor().ItemsPerPoint == 1 ? 1 : 0);
+
 					texture->OnFinalise().then([=]()
 					{
 						texture->Bind(InMaterial.get(), BindInfo);
