Fixing crash when undo Texture Input convert param to const

Adding a simple check for a valid input for thumbnail generation

diff --git a/Engine/Plugins/Experimental/TextureGraph/Source/TextureGraphEditor/Private/EdGraph/STG_EditorGraphNode.cpp b/Engine/Plugins/Experimental/TextureGraph/Source/TextureGraphEditor/Private/EdGraph/STG_EditorGraphNode.cpp
--- a/Engine/Plugins/Experimental/TextureGraph/Source/TextureGraphEditor/Private/EdGraph/STG_EditorGraphNode.cpp
+++ b/Engine/Plugins/Experimental/TextureGraph/Source/TextureGraphEditor/Private/EdGraph/STG_EditorGraphNode.cpp
@@ -990,8 +990,13 @@ void STG_EditorGraphNode::UpdateThumbnail(const FTG_EvaluationContext* InContext
 				UMixInterface* Mix = InContext->Cycle->GetMix();
 				auto TargetId = InContext->TargetId;
 
-				TiledBlobPtr ThumbBlob = T_Thumbnail::Bind(Mix,  Pin, OutTexture.RasterBlob, TargetId);
-				
+				if (!OutTexture.RasterBlob)
+				{
+					OutTexture = FTG_Texture::GetBlack();
+				}
+
+				TiledBlobPtr ThumbBlob = T_Thumbnail::Bind(Mix, Pin, OutTexture.RasterBlob, TargetId);
+
 				TSEditorGraphNode->CacheThumbBlob(Pin->GetId(), ThumbBlob);
 				ThumbBlob->OnFinalise()
 					.then([=, this](const Blob* FinalisedBlob) mutable
@@ -1006,7 +1011,6 @@ void STG_EditorGraphNode::UpdateThumbnail(const FTG_EvaluationContext* InContext
 			}
 		}
 	}
-
 }
 
 
