Fixed issue: Parameters tab should show the parameter names

Fix/How: Renamed the files to match the structure name to fix the issue.

diff --git a/Engine/Plugins/Experimental/TextureGraph/Source/TextureGraphEditor/Private/TextureGraphEditorModule.cpp b/Engine/Plugins/Experimental/TextureGraph/Source/TextureGraphEditor/Private/TextureGraphEditorModule.cpp
--- a/Engine/Plugins/Experimental/TextureGraph/Source/TextureGraphEditor/Private/TextureGraphEditorModule.cpp
+++ b/Engine/Plugins/Experimental/TextureGraph/Source/TextureGraphEditor/Private/TextureGraphEditorModule.cpp
@@ -30,7 +30,7 @@ void FTextureGraphEditorModule::StartupModule()
 	RegisterAssetTypeAction(AssetTools, MakeShareable(new FAssetTypeActions_TSX()));
 
 	FPropertyEditorModule& PropertyEditorModule = FModuleManager::GetModuleChecked<FPropertyEditorModule>("PropertyEditor");
-	PropertyEditorModule.RegisterCustomPropertyTypeLayout("TSParameterInfo", FOnGetPropertyTypeCustomizationInstance::CreateStatic(&FTG_ParameterInfoCustomization::Create));
+	PropertyEditorModule.RegisterCustomPropertyTypeLayout("TG_ParameterInfo", FOnGetPropertyTypeCustomizationInstance::CreateStatic(&FTG_ParameterInfoCustomization::Create));
 	PropertyEditorModule.RegisterCustomPropertyTypeLayout("TG_Texture", FOnGetPropertyTypeCustomizationInstance::CreateStatic(&FTG_TextureCustomization::Create));
 	auto ScalarIdentifier = MakeShared<FTG_ScalarTypeIdentifier>();
 	PropertyEditorModule.RegisterCustomPropertyTypeLayout("FloatProperty", FOnGetPropertyTypeCustomizationInstance::CreateStatic(&FTG_ScalarCustomization::Create), ScalarIdentifier);
@@ -60,7 +60,7 @@ void FTextureGraphEditorModule::ShutdownModule()
 	FTG_EditorCommands::Unregister();
 	
 	FPropertyEditorModule& PropertyEditorModule = FModuleManager::GetModuleChecked<FPropertyEditorModule>("PropertyEditor");
-	PropertyEditorModule.UnregisterCustomPropertyTypeLayout("TSParameterInfo");
+	PropertyEditorModule.UnregisterCustomPropertyTypeLayout("TG_ParameterInfo");
 	PropertyEditorModule.UnregisterCustomPropertyTypeLayout("TG_Texture");
 	auto ScalarIdentifier = MakeShared<FTG_ScalarTypeIdentifier>();
 	PropertyEditorModule.UnregisterCustomPropertyTypeLayout("FloatProperty", ScalarIdentifier);
