Minor display change to HSV node

diff --git a/Engine/Plugins/Experimental/TextureGraph/Source/TextureGraph/Public/Expressions/Color/TG_Expression_HSV.h b/Engine/Plugins/Experimental/TextureGraph/Source/TextureGraph/Public/Expressions/Color/TG_Expression_HSV.h
--- a/Engine/Plugins/Experimental/TextureGraph/Source/TextureGraph/Public/Expressions/Color/TG_Expression_HSV.h
+++ b/Engine/Plugins/Experimental/TextureGraph/Source/TextureGraph/Public/Expressions/Color/TG_Expression_HSV.h
@@ -75,11 +75,11 @@ public:
 	UPROPERTY(meta = (TGType = "TG_Input", PinDisplayName = ""))
 	FTG_Texture							Input;
 
-	// The hue  
-	UPROPERTY(EditAnywhere, Category = NoCategory, meta = (TGType = "TG_Setting", UIMin = "0", ClampMin = "0", UIMax = "1", ClampMax = "1"))
+	// The normalized hue. Please divide your [0, 359] hue values by 359 for this input
+	UPROPERTY(EditAnywhere, Category = NoCategory, meta = (TGType = "TG_Setting", UIMin = "0", ClampMin = "0", UIMax = "1", ClampMax = "1", PinDisplayName = "Hue (Normalized)", DisplayName = "Hue (Normalized)"))
 	float								Hue = 1.0f;
 	
-	// The saturation  
+	// The saturation
 	UPROPERTY(EditAnywhere, Category = NoCategory, meta = (TGType = "TG_Setting", UIMin = "0", ClampMin = "0", UIMax = "1", ClampMax = "1"))
 	float								Saturation = 1.0f;
 
