Prevent accessing a NUll property of a Pin from EdGraphNode regarding visibility update

Modify the interface of the function to work at the level of the Pin , not Fproperty which is not guaranteed per Pin

Added a editor test creating a MAterial and Texture node and settign up with test asset to validate the basic use case

Adding boolrean rreturns to the Setup function of the EngineInteraction in test in order to use them directly in the test sequence
#tests 65a83b75555fab6aca0ce75c

diff --git a/Engine/Plugins/Experimental/TextureGraph/Source/TextureGraphEditor/Private/EdGraph/TG_EdGraphNode.cpp b/Engine/Plugins/Experimental/TextureGraph/Source/TextureGraphEditor/Private/EdGraph/TG_EdGraphNode.cpp
--- a/Engine/Plugins/Experimental/TextureGraph/Source/TextureGraphEditor/Private/EdGraph/TG_EdGraphNode.cpp
+++ b/Engine/Plugins/Experimental/TextureGraph/Source/TextureGraphEditor/Private/EdGraph/TG_EdGraphNode.cpp
@@ -184,7 +184,7 @@ void UTG_EdGraphNode::AllocateDefaultPins()
 			NewPin->bNotConnectable = Pin->IsNotConnectable() || Pin->IsParam();
 
 			FString DefaultValue = Pin->GetSelfVar()->LogValue();
-			FProperty* Property = Pin->GetExpressionProperty();
+			FProperty* Property = Pin->GetExpressionProperty(); //// BEWARE Pin do NOT all have a Property, it is potentially NULL
 			FByteProperty* ByteProperty = CastField<FByteProperty>(Property);
 			if (ByteProperty)
 			{
@@ -205,9 +205,6 @@ void UTG_EdGraphNode::AllocateDefaultPins()
 				}
 				NewPin->bNotConnectable = true;
 			}
-#if WITH_EDITOR
-			UpdatePinVisibility(NewPin, Property);
-#endif
 			NewPin->DefaultValue = DefaultValue;
 			// This updates the UObject (Texture/Material) picker UI in the Node to get updated
 			if (Property && Property->GetClass()->IsChildOf(FObjectProperty::StaticClass()))
@@ -216,6 +213,10 @@ void UTG_EdGraphNode::AllocateDefaultPins()
 			}
 			NewPin->PinFriendlyName = FText::FromName(Pin->GetAliasName());
 			NewPin->PinToolTip = Pin->LogTooltip();
+
+#if WITH_EDITOR
+			UpdatePinVisibility(NewPin, Pin);
+#endif		
 		}
 	}
 	AdvancedPinDisplay = showAdvancedPinDisplay ? ENodeAdvancedPins::Hidden : ENodeAdvancedPins::NoPins;
@@ -444,16 +445,25 @@ void UTG_EdGraphNode::PostPasteNode()
 }
 
 #if WITH_EDITOR
-void UTG_EdGraphNode::UpdatePinVisibility(UEdGraphPin* Pin, FProperty* Property) const
+void UTG_EdGraphNode::UpdatePinVisibility(UEdGraphPin* Pin, UTG_Pin* TGPin) const
 {
 	UTG_Expression* Expression = GetNode()->GetExpression();
-	bool bCanEditChange = Expression->CanEditChange(Property);
-	bool bEditConditionHides = Property->HasMetaData("EditConditionHides");
+	bool bCanEditChange = !TGPin->IsConnected();
+	bool bEditConditionHides = false;
+
+	// Rely on the FProperty meta information for potential conditional editability / visibility
+	FProperty* Property = TGPin->GetExpressionProperty();
+	if (Property)
+	{
+		bCanEditChange = Expression->CanEditChange(Property);
+		bEditConditionHides = Property->HasMetaData("EditConditionHides");
+
+		// update Metadata with EditCondition to also update details view
+		FName NAME_EditCondition(TEXT("EditCondition"));
+		Property->SetMetaData(NAME_EditCondition, bCanEditChange ? TEXT("true") : TEXT("false"));
+
+	}
 
-	// update Metadata with EditCondition to also update details view
-	FName NAME_EditCondition(TEXT("EditCondition"));
-	Property->SetMetaData(NAME_EditCondition, bCanEditChange ? TEXT("true") : TEXT("false"));
-	
 	Pin->bHidden = !bCanEditChange && bEditConditionHides;
 	Pin->bDefaultValueIsReadOnly = bCanEditChange;
 }
@@ -471,12 +481,7 @@ void UTG_EdGraphNode::UpdateInputPinsVisibility() const
 		UEdGraphPin* EdPin = FindPinChecked(OtherPinName, EEdGraphPinDirection::EGPD_Input);
 
 		// check hide state
-		
-		auto PinProperty = TGPin->GetExpressionProperty();
-		if (PinProperty)
-		{
-			UpdatePinVisibility(EdPin, PinProperty);
-		}
+		UpdatePinVisibility(EdPin, TGPin);
 	}
 }
 #endif
diff --git a/Engine/Plugins/Experimental/TextureGraph/Source/TextureGraphEditor/Public/EdGraph/TG_EdGraphNode.h b/Engine/Plugins/Experimental/TextureGraph/Source/TextureGraphEditor/Public/EdGraph/TG_EdGraphNode.h
--- a/Engine/Plugins/Experimental/TextureGraph/Source/TextureGraphEditor/Public/EdGraph/TG_EdGraphNode.h
+++ b/Engine/Plugins/Experimental/TextureGraph/Source/TextureGraphEditor/Public/EdGraph/TG_EdGraphNode.h
@@ -64,7 +64,7 @@ public:
 	virtual void PostPasteNode() override;
 
 #if WITH_EDITOR
-	void UpdatePinVisibility(UEdGraphPin* Pin, FProperty* Property) const;
+	void UpdatePinVisibility(UEdGraphPin* Pin, UTG_Pin* TGPin) const;
 	void UpdateInputPinsVisibility() const;
 #endif
 	
