Rename MixerEngine and other files with the Mixer branding in them.

diff --git a/Engine/Plugins/Experimental/TextureGraph/Source/MixerInsight/MixerInsight.Build.cs b/Engine/Plugins/Experimental/TextureGraph/Source/MixerInsight/MixerInsight.Build.cs
--- a/Engine/Plugins/Experimental/TextureGraph/Source/MixerInsight/MixerInsight.Build.cs
+++ b/Engine/Plugins/Experimental/TextureGraph/Source/MixerInsight/MixerInsight.Build.cs
@@ -27,7 +27,7 @@ public class MixerInsight : ModuleRules
 
 		PrivateDependencyModuleNames.AddRange(new string[]
 		{
-			"MixerEngine",
+			"TextureGraphEngine",
 			"FreeImage",
 			"Function2",
 			"Continuable",
diff --git a/Engine/Plugins/Experimental/TextureGraph/Source/MixerInsight/Private/MixerInsight.cpp b/Engine/Plugins/Experimental/TextureGraph/Source/MixerInsight/Private/MixerInsight.cpp
--- a/Engine/Plugins/Experimental/TextureGraph/Source/MixerInsight/Private/MixerInsight.cpp
+++ b/Engine/Plugins/Experimental/TextureGraph/Source/MixerInsight/Private/MixerInsight.cpp
@@ -1,7 +1,7 @@
 // Copyright Epic Games, Inc. All Rights Reserved.
 #include "MixerInsight.h"
 #include "Modules/ModuleManager.h"
-#include "MixerEngine.h"
+#include "TextureGraphEngine.h"
 
 #include "Model/MixerInsightObserver.h"
 #include "Model/MixerInsightSession.h"
@@ -45,7 +45,7 @@ bool MixerInsight::Create()
 
 		/// Need the engine observer to watch what is happening
 		auto EngineObserver = std::make_shared<MixerInsightEngineObserver>();
-		MixerEngine::RegisterObserverSource(EngineObserver); // will also install other observers
+		TextureGraphEngine::RegisterObserverSource(EngineObserver); // will also install other observers
 
 		return true;
 	}
@@ -58,7 +58,7 @@ bool MixerInsight::Destroy()
 	if (GInstance)
 	{
 		/// Remove Engine observer
-		MixerEngine::RegisterObserverSource(nullptr);
+		TextureGraphEngine::RegisterObserverSource(nullptr);
 
 		delete GInstance;
 		GInstance = nullptr;
diff --git a/Engine/Plugins/Experimental/TextureGraph/Source/MixerInsight/Private/Model/MixerInsightObserver.cpp b/Engine/Plugins/Experimental/TextureGraph/Source/MixerInsight/Private/Model/MixerInsightObserver.cpp
--- a/Engine/Plugins/Experimental/TextureGraph/Source/MixerInsight/Private/Model/MixerInsightObserver.cpp
+++ b/Engine/Plugins/Experimental/TextureGraph/Source/MixerInsight/Private/Model/MixerInsightObserver.cpp
@@ -50,7 +50,7 @@ void MixerInsightSchedulerObserver::Start()
 
 void MixerInsightSchedulerObserver::UpdateIdle()
 {
-	if (MixerEngine::IsTestMode())
+	if (TextureGraphEngine::IsTestMode())
 		return;
 
 	MixerInsight::Instance()->GetSession()->UpdateIdle();
@@ -63,7 +63,7 @@ void MixerInsightSchedulerObserver::Stop()
 
 void MixerInsightSchedulerObserver::BatchAdded(JobBatchPtr Batch)
 {
-	if (MixerEngine::IsTestMode())
+	if (TextureGraphEngine::IsTestMode())
 		return;
 
 	MixerInsight::Instance()->GetSession()->BatchAdded(Batch);
@@ -71,7 +71,7 @@ void MixerInsightSchedulerObserver::BatchAdded(JobBatchPtr Batch)
 
 void MixerInsightSchedulerObserver::BatchDone(JobBatchPtr Batch)
 {
-	if (MixerEngine::IsTestMode())
+	if (TextureGraphEngine::IsTestMode())
 		return;
 
 	MixerInsight::Instance()->GetSession()->BatchDone(Batch);
@@ -79,7 +79,7 @@ void MixerInsightSchedulerObserver::BatchDone(JobBatchPtr Batch)
 
 void MixerInsightSchedulerObserver::BatchJobsDone(JobBatchPtr Batch)
 {
-	if (MixerEngine::IsTestMode())
+	if (TextureGraphEngine::IsTestMode())
 		return;
 
 	MixerInsight::Instance()->GetSession()->BatchJobsDone(Batch);
@@ -104,22 +104,22 @@ MixerInsightEngineObserver::~MixerInsightEngineObserver()
 
 void MixerInsightEngineObserver::Created()
 {
-	if (MixerEngine::IsTestMode())
+	if (TextureGraphEngine::IsTestMode())
 		return;
 
 	UE_LOG(LogMixerInsight, Log, TEXT("MixerInsightEngineObserver::Created"));
 
 	/// Engine is created when notified, this should be true
-	if (MixerEngine::GetInstance())
+	if (TextureGraphEngine::GetInstance())
 	{
 		for (int i = 0; i < (uint32)DeviceType::Count; ++i)
 		{
-			auto Dev = MixerEngine::GetInstance()->GetDeviceManager()->GetDevice(i);
+			auto Dev = TextureGraphEngine::GetInstance()->GetDeviceManager()->GetDevice(i);
 			if (Dev)
 				Dev->RegisterObserverSource(_deviceObservers[i]);
 		}
-		MixerEngine::GetInstance()->GetBlobber()->RegisterObserverSource(BlobberObserver);
-		MixerEngine::GetInstance()->GetScheduler()->RegisterObserverSource(SchedulerObserver);
+		TextureGraphEngine::GetInstance()->GetBlobber()->RegisterObserverSource(BlobberObserver);
+		TextureGraphEngine::GetInstance()->GetScheduler()->RegisterObserverSource(SchedulerObserver);
 	}
 
 	/// This is a brand new session engine
@@ -129,22 +129,22 @@ void MixerInsightEngineObserver::Created()
 
 void MixerInsightEngineObserver::Destroyed()
 {
-	if (MixerEngine::IsTestMode())
+	if (TextureGraphEngine::IsTestMode())
 		return;
 
 	UE_LOG(LogMixerInsight, Log, TEXT("MixerInsightEngineObserver::Destroyed"));
 
 	/// Engine is already destroyed when notified
-	if (MixerEngine::GetInstance())
+	if (TextureGraphEngine::GetInstance())
 	{
 		for (int i = 0; i < (uint32)DeviceType::Count; ++i)
 		{
-			auto Dev = MixerEngine::GetInstance()->GetDeviceManager()->GetDevice(i);
+			auto Dev = TextureGraphEngine::GetInstance()->GetDeviceManager()->GetDevice(i);
 			if (Dev)
 				Dev->RegisterObserverSource(nullptr);
 		}
-		MixerEngine::GetInstance()->GetBlobber()->RegisterObserverSource(nullptr); /// remove blobber observer from the engine
-		MixerEngine::GetInstance()->GetScheduler()->RegisterObserverSource(nullptr); /// remove scheduler observer from the engine
+		TextureGraphEngine::GetInstance()->GetBlobber()->RegisterObserverSource(nullptr); /// remove blobber observer from the engine
+		TextureGraphEngine::GetInstance()->GetScheduler()->RegisterObserverSource(nullptr); /// remove scheduler observer from the engine
 	}
 
 	MixerInsight::Instance()->GetSession()->EngineDestroyed();
diff --git a/Engine/Plugins/Experimental/TextureGraph/Source/MixerInsight/Private/Model/MixerInsightSession.cpp b/Engine/Plugins/Experimental/TextureGraph/Source/MixerInsight/Private/Model/MixerInsightSession.cpp
--- a/Engine/Plugins/Experimental/TextureGraph/Source/MixerInsight/Private/Model/MixerInsightSession.cpp
+++ b/Engine/Plugins/Experimental/TextureGraph/Source/MixerInsight/Private/Model/MixerInsightSession.cpp
@@ -4,7 +4,7 @@
 #include "Job/Job.h"
 #include "Data/Blob.h"
 
-#include "MixerEngine.h"
+#include "TextureGraphEngine.h"
 #include "Data/Blobber.h"
 #include "Device/DeviceManager.h"
 #include "Device/Device.h"
@@ -38,7 +38,7 @@ BlobPtr MixerInsightSession::LiveCache::GetBlob(RecordID RecId) const
 	if (RecId.IsBlob() && RecId.Blob() < Blobs.size())
 	{
 		const auto& blobEntry = Blobs[RecId.Blob()];
-		return MixerEngine::GetInstance()->GetBlobber()->Find<Blob>(blobEntry).get();
+		return TextureGraphEngine::GetInstance()->GetBlobber()->Find<Blob>(blobEntry).get();
 	}
 	return nullptr;
 }
@@ -93,7 +93,7 @@ DeviceBufferPtr MixerInsightSession::LiveCache::GetDeviceBuffer(RecordID RecId)
 			Hash = SourceBlobPtr->Hash()->Value();
 
 		// Try to find the buffer in the device under the  first hash
-		auto Dev = MixerEngine::GetInstance()->GetDeviceManager()->GetDevice(RecId.Buffer_DeviceType());
+		auto Dev = TextureGraphEngine::GetInstance()->GetDeviceManager()->GetDevice(RecId.Buffer_DeviceType());
 		auto Buffer = Dev->Find(Hash, false);
 
 		// Couldn't find with what we think is final hash, try with prev hash if exists ?
@@ -106,7 +106,7 @@ DeviceBufferPtr MixerInsightSession::LiveCache::GetDeviceBuffer(RecordID RecId)
 
 		for (int32 i = RecId.Buffer_DeviceType() + 1; i < (int32) DeviceType::Count; i++)
 		{
-			Dev = MixerEngine::GetInstance()->GetDeviceManager()->GetDevice(i);
+			Dev = TextureGraphEngine::GetInstance()->GetDeviceManager()->GetDevice(i);
 			if (Dev && !Buffer) {
 				Buffer = Dev->Find(Hash, false);
 				if (Buffer)
@@ -117,7 +117,7 @@ DeviceBufferPtr MixerInsightSession::LiveCache::GetDeviceBuffer(RecordID RecId)
 		// Couldn't find in the device cache... then look up in the blobber
 		if (!Buffer) 
 		{
-			auto BlobPtr = MixerEngine::GetInstance()->GetBlobber()->Find(BufferRecord.HashValue);
+			auto BlobPtr = TextureGraphEngine::GetInstance()->GetBlobber()->Find(BufferRecord.HashValue);
 			if (BlobPtr) {
 				Buffer = BlobPtr->GetBufferRef();
 			}
@@ -130,7 +130,7 @@ DeviceBufferPtr MixerInsightSession::LiveCache::GetDeviceBuffer(RecordID RecId)
 		if (!Buffer) {
 			if (BufferRecord.PrevHashValue != 0)
 			{
-				auto blobPtr = MixerEngine::GetInstance()->GetBlobber()->Find(BufferRecord.PrevHashValue);
+				auto blobPtr = TextureGraphEngine::GetInstance()->GetBlobber()->Find(BufferRecord.PrevHashValue);
 				if (blobPtr) {
 					Buffer = blobPtr->GetBufferRef();
 				}
@@ -152,7 +152,7 @@ DeviceBufferPtr MixerInsightSession::LiveCache::GetDeviceBuffer(RecordID RecId)
 
 void MixerInsightSession::EngineCreated()
 {
-	if (MixerEngine::IsTestMode())
+	if (TextureGraphEngine::IsTestMode())
 		return;
 
 	UE_LOG(LogMixerInsight, Log, TEXT("MixerInsightActionManagerObserver::EngineCreated"));
@@ -167,7 +167,7 @@ void MixerInsightSession::EngineCreated()
 
 void MixerInsightSession::EngineDestroyed()
 {
-	if (MixerEngine::IsTestMode())
+	if (TextureGraphEngine::IsTestMode())
 		return;
 
 	UE_LOG(LogMixerInsight, Log, TEXT("MixerInsightActionManagerObserver::EngineDestroyed"));
@@ -202,7 +202,7 @@ uint32 MixerInsightSession::ParseDeviceBuffersCache(DeviceType DevType, RecordID
 	uint32 NumNewBuffers = 0;
 
 	auto DevIndex = (int32) DevType;
-	auto Dev = MixerEngine::GetInstance()->GetDeviceManager()->GetDevice(DevType);
+	auto Dev = TextureGraphEngine::GetInstance()->GetDeviceManager()->GetDevice(DevType);
 	if (Dev)
 	{
 		auto Version = Dev->GetObserverSource()->GetVersion();
@@ -663,7 +663,7 @@ RecordID MixerInsightSession::FindOrCreateBlobRecordFromHash(HashType hash)
 	auto RecId = GetRecord().FindBlobRecord(hash);
 	if (!RecId.IsValid())
 	{
-		auto BlobObj = MixerEngine::GetInstance()->GetBlobber()->Find(hash);
+		auto BlobObj = TextureGraphEngine::GetInstance()->GetBlobber()->Find(hash);
 		if (BlobObj)
 		{
 			return FindOrCreateBlobRecord(BlobObj.get());
@@ -823,8 +823,8 @@ bool MixerInsightSession::ReplayBatch(RecordID batchRecordID, bool captureRender
 		batchPtr->ResetForReplay();
 		batchPtr->SetCaptureRenderDoc(false);
 		if (captureRenderDoc)
-			MixerEngine::GetInstance()->GetScheduler()->SetCaptureRenderDocNextBatch(true);
-		MixerEngine::GetInstance()->GetScheduler()->AddBatch(batchPtr);
+			TextureGraphEngine::GetInstance()->GetScheduler()->SetCaptureRenderDocNextBatch(true);
+		TextureGraphEngine::GetInstance()->GetScheduler()->AddBatch(batchPtr);
 		return true;
 	}
 	return false;
@@ -841,8 +841,8 @@ bool MixerInsightSession::ReplayJob(RecordID jobRecordID, bool captureRenderDoc)
 		batchPtr->ResetForReplay(jobRecordID.Job());
 		batchPtr->SetCaptureRenderDoc(false);
 		if(captureRenderDoc)
-			MixerEngine::GetInstance()->GetScheduler()->SetCaptureRenderDocNextBatch(true);
-		MixerEngine::GetInstance()->GetScheduler()->AddBatch(batchPtr);
+			TextureGraphEngine::GetInstance()->GetScheduler()->SetCaptureRenderDocNextBatch(true);
+		TextureGraphEngine::GetInstance()->GetScheduler()->AddBatch(batchPtr);
 
 		return true;
 	}
diff --git a/Engine/Plugins/Experimental/TextureGraph/Source/MixerInsight/Private/View/SMixerInsightDeviceBufferView.cpp b/Engine/Plugins/Experimental/TextureGraph/Source/MixerInsight/Private/View/SMixerInsightDeviceBufferView.cpp
--- a/Engine/Plugins/Experimental/TextureGraph/Source/MixerInsight/Private/View/SMixerInsightDeviceBufferView.cpp
+++ b/Engine/Plugins/Experimental/TextureGraph/Source/MixerInsight/Private/View/SMixerInsightDeviceBufferView.cpp
@@ -5,7 +5,7 @@
 #include "MixerInsight.h"
 #include "Model/MixerInsightSession.h"
 
-#include "MixerEngine.h"
+#include "TextureGraphEngine.h"
 #include "Device/DeviceManager.h"
 #include "Device/FX/DeviceBuffer_FX.h"
 
diff --git a/Engine/Plugins/Experimental/TextureGraph/Source/MixerInsight/Private/View/SMixerInsightDeviceView.cpp b/Engine/Plugins/Experimental/TextureGraph/Source/MixerInsight/Private/View/SMixerInsightDeviceView.cpp
--- a/Engine/Plugins/Experimental/TextureGraph/Source/MixerInsight/Private/View/SMixerInsightDeviceView.cpp
+++ b/Engine/Plugins/Experimental/TextureGraph/Source/MixerInsight/Private/View/SMixerInsightDeviceView.cpp
@@ -6,7 +6,7 @@
 
 #include "Device/DeviceManager.h"
 #include "Data/Blobber.h"
-#include "MixerEngine.h"
+#include "TextureGraphEngine.h"
 #include <Widgets/Views/ITableRow.h>
 #include <Widgets/Views/STableViewBase.h>
 #include <Widgets/Text/STextBlock.h>
@@ -189,14 +189,14 @@ public:
 	{
 		if (_recordID.IsDevice())
 		{
-			if (!MixerEngine::GetInstance())
+			if (!TextureGraphEngine::GetInstance())
 			{
 				SetColorAndOpacity(FLinearColor(0.3, 0.3, 0.3));
 			}
 			else
 			{
 				auto item = _item.Pin();
-				auto device = MixerEngine::GetInstance()->GetDeviceManager()->GetDevice(_recordID.Buffer_DeviceType());
+				auto device = TextureGraphEngine::GetInstance()->GetDeviceManager()->GetDevice(_recordID.Buffer_DeviceType());
 				if (!device)
 				{
 					SetColorAndOpacity(FLinearColor(0.2, 0.2, 0.2));
@@ -560,7 +560,7 @@ void SMixerInsightDeviceView::Construct(const FArguments& Args)
 
 void SMixerInsightDeviceView::Refresh()
 {
-	auto device = MixerEngine::GetInstance()->GetDeviceManager()->GetDevice(_deviceType);
+	auto device = TextureGraphEngine::GetInstance()->GetDeviceManager()->GetDevice(_deviceType);
 	_threadCount->SetText(FText::FromString(FString::FromInt(device->GetMaxThreads())));
 
 	auto memMax = device->GetMaxMemory();
@@ -607,7 +607,7 @@ void SMixerInsightDeviceManagerView::Reset()
 		SAssignNew(vbox, SVerticalBox)
 	]; 
 
-	if (MixerEngine::GetInstance())
+	if (TextureGraphEngine::GetInstance())
 	{
 		for (int i = 0; i < (int)DeviceType::Count; ++i)
 		{
@@ -615,7 +615,7 @@ void SMixerInsightDeviceManagerView::Reset()
 			//	continue;
 
 			auto deviceType = (DeviceType)i;
-			auto device = MixerEngine::GetInstance()->GetDeviceManager()->GetDevice(deviceType);
+			auto device = TextureGraphEngine::GetInstance()->GetDeviceManager()->GetDevice(deviceType);
 			if (device)
 			{
 				TSharedPtr< SMixerInsightDeviceView> deviceView;
diff --git a/Engine/Plugins/Experimental/TextureGraph/Source/MixerInsight/Private/View/SMixerInsightMixView.cpp b/Engine/Plugins/Experimental/TextureGraph/Source/MixerInsight/Private/View/SMixerInsightMixView.cpp
--- a/Engine/Plugins/Experimental/TextureGraph/Source/MixerInsight/Private/View/SMixerInsightMixView.cpp
+++ b/Engine/Plugins/Experimental/TextureGraph/Source/MixerInsight/Private/View/SMixerInsightMixView.cpp
@@ -6,7 +6,7 @@
 
 #include "Device/DeviceManager.h"
 #include "Data/Blobber.h"
-#include "MixerEngine.h"
+#include "TextureGraphEngine.h"
 
 
 class SMixerInsightMixListViewRow : public SMultiColumnTableRow<SMixerInsightMixListView::FItem>
diff --git a/Engine/Plugins/Experimental/TextureGraph/Source/MixerInsight/Public/MixerInsight.h b/Engine/Plugins/Experimental/TextureGraph/Source/MixerInsight/Public/MixerInsight.h
--- a/Engine/Plugins/Experimental/TextureGraph/Source/MixerInsight/Public/MixerInsight.h
+++ b/Engine/Plugins/Experimental/TextureGraph/Source/MixerInsight/Public/MixerInsight.h
@@ -29,7 +29,7 @@ private:
 
 	MixerInsightSessionPtr					Session;
 public:
-	static bool								Create();  /// Create the Insight Instance ONLY if no other currently created AND if a MixerEngine is already created
+	static bool								Create();  /// Create the Insight Instance ONLY if no other currently created AND if a TextureGraphEngine is already created
 
 	/// Destroy the current instance
 	static bool								Destroy(); /// Destroy the instance of Insight if it exists.
diff --git a/Engine/Plugins/Experimental/TextureGraph/Source/MixerInsightEditor/MixerInsightEditor.Build.cs b/Engine/Plugins/Experimental/TextureGraph/Source/MixerInsightEditor/MixerInsightEditor.Build.cs
--- a/Engine/Plugins/Experimental/TextureGraph/Source/MixerInsightEditor/MixerInsightEditor.Build.cs
+++ b/Engine/Plugins/Experimental/TextureGraph/Source/MixerInsightEditor/MixerInsightEditor.Build.cs
@@ -50,7 +50,7 @@ public class MixerInsightEditor : ModuleRules
 				// ... add private dependencies that you statically link with here ...	
 				"Continuable",
 				"DeveloperSettings",
-                "MixerEngine",
+                "TextureGraphEngine",
                 "MixerInsight",
             }
 			);
diff --git a/Engine/Plugins/Experimental/TextureGraph/Source/TextureGraph/Private/Blueprint/TG_AsyncTask.cpp b/Engine/Plugins/Experimental/TextureGraph/Source/TextureGraph/Private/Blueprint/TG_AsyncTask.cpp
--- a/Engine/Plugins/Experimental/TextureGraph/Source/TextureGraph/Private/Blueprint/TG_AsyncTask.cpp
+++ b/Engine/Plugins/Experimental/TextureGraph/Source/TextureGraph/Private/Blueprint/TG_AsyncTask.cpp
@@ -2,7 +2,7 @@
 
 #include "Blueprint/TG_AsyncTask.h"
 
-#include "MixerEngine.h"
+#include "TextureGraphEngine.h"
 #include "Blueprint/TG_AsyncTaskManager.h"
 
 
@@ -24,5 +24,5 @@ void UTG_AsyncTask::SetReadyToDestroy()
 
 void UTG_AsyncTask::Activate()
 {
-	MixerEngine::SetRunEngine();
+	TextureGraphEngine::SetRunEngine();
 }
diff --git a/Engine/Plugins/Experimental/TextureGraph/Source/TextureGraph/Private/Expressions/Adjustment/TG_Expression_NormalFromHeightMap.cpp b/Engine/Plugins/Experimental/TextureGraph/Source/TextureGraph/Private/Expressions/Adjustment/TG_Expression_NormalFromHeightMap.cpp
--- a/Engine/Plugins/Experimental/TextureGraph/Source/TextureGraph/Private/Expressions/Adjustment/TG_Expression_NormalFromHeightMap.cpp
+++ b/Engine/Plugins/Experimental/TextureGraph/Source/TextureGraph/Private/Expressions/Adjustment/TG_Expression_NormalFromHeightMap.cpp
@@ -29,9 +29,9 @@ bool UTG_Expression_NormalFromHeightMap::Validate(MixUpdateCyclePtr Cycle)
 	if(Output.Descriptor.TextureFormat != ETG_TextureFormat::Auto && TextureHelper::GetNumChannelsFromTGTextureFormat(Output.Descriptor.TextureFormat) < 3)
 	{
 		UMixInterface* ParentMix = Cast<UMixInterface>(GetOutermostObject());
-		auto ErrorType = static_cast<int32>(EMixerErrorType::NODE_WARNING);
+		auto ErrorType = static_cast<int32>(ETextureGraphErrorType::NODE_WARNING);
 
-		MixerEngine::GetErrorReporter(ParentMix)->ReportWarning(ErrorType, FString::Printf(TEXT("Output requires minimum 3 channels to display normal correctly")), GetParentNode());
+		TextureGraphEngine::GetErrorReporter(ParentMix)->ReportWarning(ErrorType, FString::Printf(TEXT("Output requires minimum 3 channels to display normal correctly")), GetParentNode());
 	}
 	return Super::Validate(Cycle);
 }
diff --git a/Engine/Plugins/Experimental/TextureGraph/Source/TextureGraph/Private/Expressions/Channel/TG_Expression_ChannelCombiner.cpp b/Engine/Plugins/Experimental/TextureGraph/Source/TextureGraph/Private/Expressions/Channel/TG_Expression_ChannelCombiner.cpp
--- a/Engine/Plugins/Experimental/TextureGraph/Source/TextureGraph/Private/Expressions/Channel/TG_Expression_ChannelCombiner.cpp
+++ b/Engine/Plugins/Experimental/TextureGraph/Source/TextureGraph/Private/Expressions/Channel/TG_Expression_ChannelCombiner.cpp
@@ -9,7 +9,7 @@ void UTG_Expression_ChannelCombiner::Evaluate(FTG_EvaluationContext* InContext)
 {
 	Super::Evaluate(InContext);
 
-	const RenderMaterial_FXPtr RenderMaterial = MixerEngine::GetMaterialManager()->CreateMaterial_FX<VSH_Simple, FSH_ChannelCombiner>(TEXT("T_ChannelCombiner"));
+	const RenderMaterial_FXPtr RenderMaterial = TextureGraphEngine::GetMaterialManager()->CreateMaterial_FX<VSH_Simple, FSH_ChannelCombiner>(TEXT("T_ChannelCombiner"));
 
 	check(RenderMaterial);
 	TiledBlobPtr Source = nullptr;
diff --git a/Engine/Plugins/Experimental/TextureGraph/Source/TextureGraph/Private/Expressions/Channel/TG_Expression_ChannelSplitter.cpp b/Engine/Plugins/Experimental/TextureGraph/Source/TextureGraph/Private/Expressions/Channel/TG_Expression_ChannelSplitter.cpp
--- a/Engine/Plugins/Experimental/TextureGraph/Source/TextureGraph/Private/Expressions/Channel/TG_Expression_ChannelSplitter.cpp
+++ b/Engine/Plugins/Experimental/TextureGraph/Source/TextureGraph/Private/Expressions/Channel/TG_Expression_ChannelSplitter.cpp
@@ -10,7 +10,7 @@
 template <typename FSH_PixelShaderType>
 static void GenericEvaluate(FTG_EvaluationContext* InContext, FString ChannelName, FTG_Texture& Input, FTG_Texture& Output)
 {
-	const RenderMaterial_FXPtr RenderMaterial = MixerEngine::GetMaterialManager()->CreateMaterial_FX<VSH_Simple, FSH_PixelShaderType>(TEXT("T_ChannelSplitter_") + ChannelName);
+	const RenderMaterial_FXPtr RenderMaterial = TextureGraphEngine::GetMaterialManager()->CreateMaterial_FX<VSH_Simple, FSH_PixelShaderType>(TEXT("T_ChannelSplitter_") + ChannelName);
 
 	check(RenderMaterial);
 
diff --git a/Engine/Plugins/Experimental/TextureGraph/Source/TextureGraph/Private/Expressions/Color/TG_Expression_Grayscale.cpp b/Engine/Plugins/Experimental/TextureGraph/Source/TextureGraph/Private/Expressions/Color/TG_Expression_Grayscale.cpp
--- a/Engine/Plugins/Experimental/TextureGraph/Source/TextureGraph/Private/Expressions/Color/TG_Expression_Grayscale.cpp
+++ b/Engine/Plugins/Experimental/TextureGraph/Source/TextureGraph/Private/Expressions/Color/TG_Expression_Grayscale.cpp
@@ -9,7 +9,7 @@ void UTG_Expression_Grayscale::Evaluate(FTG_EvaluationContext* InContext)
 {
 	Super::Evaluate(InContext);
 
-	const RenderMaterial_FXPtr RenderMaterial = MixerEngine::GetMaterialManager()->CreateMaterial_FX<VSH_Simple, FSH_Grayscale>(TEXT("T_Grayscale"));
+	const RenderMaterial_FXPtr RenderMaterial = TextureGraphEngine::GetMaterialManager()->CreateMaterial_FX<VSH_Simple, FSH_Grayscale>(TEXT("T_Grayscale"));
 
 	check(RenderMaterial);
 
diff --git a/Engine/Plugins/Experimental/TextureGraph/Source/TextureGraph/Private/Expressions/Color/TG_Expression_HSV.cpp b/Engine/Plugins/Experimental/TextureGraph/Source/TextureGraph/Private/Expressions/Color/TG_Expression_HSV.cpp
--- a/Engine/Plugins/Experimental/TextureGraph/Source/TextureGraph/Private/Expressions/Color/TG_Expression_HSV.cpp
+++ b/Engine/Plugins/Experimental/TextureGraph/Source/TextureGraph/Private/Expressions/Color/TG_Expression_HSV.cpp
@@ -9,7 +9,7 @@ void UTG_Expression_HSV::Evaluate(FTG_EvaluationContext* InContext)
 {
 	Super::Evaluate(InContext);
 
-	const RenderMaterial_FXPtr RenderMaterial = MixerEngine::GetMaterialManager()->CreateMaterial_FX<VSH_Simple, FSH_HSV>(TEXT("T_HSV"));
+	const RenderMaterial_FXPtr RenderMaterial = TextureGraphEngine::GetMaterialManager()->CreateMaterial_FX<VSH_Simple, FSH_HSV>(TEXT("T_HSV"));
 
 	check(RenderMaterial);
 
@@ -39,7 +39,7 @@ void UTG_Expression_RGB2HSV::Evaluate(FTG_EvaluationContext* InContext)
 {
 	Super::Evaluate(InContext);
 
-	const RenderMaterial_FXPtr RenderMaterial = MixerEngine::GetMaterialManager()->CreateMaterial_FX<VSH_Simple, FSH_RGB2HSV>(TEXT("T_RGB2HSV"));
+	const RenderMaterial_FXPtr RenderMaterial = TextureGraphEngine::GetMaterialManager()->CreateMaterial_FX<VSH_Simple, FSH_RGB2HSV>(TEXT("T_RGB2HSV"));
 
 	check(RenderMaterial);
 
@@ -66,7 +66,7 @@ void UTG_Expression_HSV2RGB::Evaluate(FTG_EvaluationContext* InContext)
 {
 	Super::Evaluate(InContext);
 
-	const RenderMaterial_FXPtr RenderMaterial = MixerEngine::GetMaterialManager()->CreateMaterial_FX<VSH_Simple, FSH_HSV2RGB>(TEXT("T_HSV2RGB"));
+	const RenderMaterial_FXPtr RenderMaterial = TextureGraphEngine::GetMaterialManager()->CreateMaterial_FX<VSH_Simple, FSH_HSV2RGB>(TEXT("T_HSV2RGB"));
 
 	check(RenderMaterial);
 
diff --git a/Engine/Plugins/Experimental/TextureGraph/Source/TextureGraph/Private/Expressions/FIlter/TG_Expression_Brightness.cpp b/Engine/Plugins/Experimental/TextureGraph/Source/TextureGraph/Private/Expressions/FIlter/TG_Expression_Brightness.cpp
--- a/Engine/Plugins/Experimental/TextureGraph/Source/TextureGraph/Private/Expressions/FIlter/TG_Expression_Brightness.cpp
+++ b/Engine/Plugins/Experimental/TextureGraph/Source/TextureGraph/Private/Expressions/FIlter/TG_Expression_Brightness.cpp
@@ -9,7 +9,7 @@ void UTG_Expression_Brightness::Evaluate(FTG_EvaluationContext* InContext)
 {
 	Super::Evaluate(InContext);
 
-	const RenderMaterial_FXPtr RenderMaterial = MixerEngine::GetMaterialManager()->CreateMaterial_FX<VSH_Simple, FSH_Brightness>(TEXT("T_Brightness"));
+	const RenderMaterial_FXPtr RenderMaterial = TextureGraphEngine::GetMaterialManager()->CreateMaterial_FX<VSH_Simple, FSH_Brightness>(TEXT("T_Brightness"));
 
 	check(RenderMaterial);
 
diff --git a/Engine/Plugins/Experimental/TextureGraph/Source/TextureGraph/Private/Expressions/FIlter/TG_Expression_Levels.cpp b/Engine/Plugins/Experimental/TextureGraph/Source/TextureGraph/Private/Expressions/FIlter/TG_Expression_Levels.cpp
--- a/Engine/Plugins/Experimental/TextureGraph/Source/TextureGraph/Private/Expressions/FIlter/TG_Expression_Levels.cpp
+++ b/Engine/Plugins/Experimental/TextureGraph/Source/TextureGraph/Private/Expressions/FIlter/TG_Expression_Levels.cpp
@@ -93,7 +93,7 @@ void UTG_Expression_Levels::Evaluate(FTG_EvaluationContext* InContext)
 {
 	Super::Evaluate(InContext);
 
-	const RenderMaterial_FXPtr RenderMaterial = MixerEngine::GetMaterialManager()->CreateMaterial_FX<VSH_Simple, FSH_Levels>(TEXT("T_Levels"));
+	const RenderMaterial_FXPtr RenderMaterial = TextureGraphEngine::GetMaterialManager()->CreateMaterial_FX<VSH_Simple, FSH_Levels>(TEXT("T_Levels"));
 
 	check(RenderMaterial);
 
diff --git a/Engine/Plugins/Experimental/TextureGraph/Source/TextureGraph/Private/Expressions/FIlter/TG_Expression_Threshold.cpp b/Engine/Plugins/Experimental/TextureGraph/Source/TextureGraph/Private/Expressions/FIlter/TG_Expression_Threshold.cpp
--- a/Engine/Plugins/Experimental/TextureGraph/Source/TextureGraph/Private/Expressions/FIlter/TG_Expression_Threshold.cpp
+++ b/Engine/Plugins/Experimental/TextureGraph/Source/TextureGraph/Private/Expressions/FIlter/TG_Expression_Threshold.cpp
@@ -10,7 +10,7 @@ void UTG_Expression_Threshold::Evaluate(FTG_EvaluationContext* InContext)
 {
 	Super::Evaluate(InContext);
 
-	const RenderMaterial_FXPtr RenderMaterial = MixerEngine::GetMaterialManager()->CreateMaterial_FX<VSH_Simple, FSH_Threshold>(TEXT("T_Threshold"));
+	const RenderMaterial_FXPtr RenderMaterial = TextureGraphEngine::GetMaterialManager()->CreateMaterial_FX<VSH_Simple, FSH_Threshold>(TEXT("T_Threshold"));
 
 	check(RenderMaterial);
 
diff --git a/Engine/Plugins/Experimental/TextureGraph/Source/TextureGraph/Private/Expressions/Input/TG_Expression_Graph.cpp b/Engine/Plugins/Experimental/TextureGraph/Source/TextureGraph/Private/Expressions/Input/TG_Expression_Graph.cpp
--- a/Engine/Plugins/Experimental/TextureGraph/Source/TextureGraph/Private/Expressions/Input/TG_Expression_Graph.cpp
+++ b/Engine/Plugins/Experimental/TextureGraph/Source/TextureGraph/Private/Expressions/Input/TG_Expression_Graph.cpp
@@ -63,9 +63,9 @@ void UTG_Expression_Graph::Evaluate(FTG_EvaluationContext* InContext)
 		}
 		else
 		{
-			auto ErrorType = static_cast<int32>(EMixerErrorType::SUBGRAPH_INTERNAL_ERROR);
+			auto ErrorType = static_cast<int32>(ETextureGraphErrorType::SUBGRAPH_INTERNAL_ERROR);
 			const FString ErrorMsg = FString::Format(TEXT("Internal error in {0} node."), {GetTitleName().ToString()} );
-			MixerEngine::GetErrorReporter(InContext->Cycle->GetMix())->ReportError(ErrorType, ErrorMsg, GetParentNode());
+			TextureGraphEngine::GetErrorReporter(InContext->Cycle->GetMix())->ReportError(ErrorType, ErrorMsg, GetParentNode());
 		}
 		
 		InContext->GraphDepth--;	
@@ -176,10 +176,10 @@ bool UTG_Expression_TextureGraph::Validate(MixUpdateCyclePtr Cycle)
 	// if TextureGraph is already in execution list then
 	if(TextureGraph && Cycle->ContainsMix(TextureGraph))
 	{
-		auto ErrorType = static_cast<int32>(EMixerErrorType::RECURSIVE_CALL);
+		auto ErrorType = static_cast<int32>(ETextureGraphErrorType::RECURSIVE_CALL);
 
 		const FString ErrorMsg = FString::Format(TEXT("Recursive Texture Graph call found in '{0}' asset."), {Cycle->TopMix()->GetName()} );
-		MixerEngine::GetErrorReporter(Cycle->GetMix())->ReportError(ErrorType, ErrorMsg, GetParentNode());
+		TextureGraphEngine::GetErrorReporter(Cycle->GetMix())->ReportError(ErrorType, ErrorMsg, GetParentNode());
 		return false;
 	}
 
diff --git a/Engine/Plugins/Experimental/TextureGraph/Source/TextureGraph/Private/Expressions/Input/TG_Expression_Texture.cpp b/Engine/Plugins/Experimental/TextureGraph/Source/TextureGraph/Private/Expressions/Input/TG_Expression_Texture.cpp
--- a/Engine/Plugins/Experimental/TextureGraph/Source/TextureGraph/Private/Expressions/Input/TG_Expression_Texture.cpp
+++ b/Engine/Plugins/Experimental/TextureGraph/Source/TextureGraph/Private/Expressions/Input/TG_Expression_Texture.cpp
@@ -63,12 +63,12 @@ bool UTG_Expression_Texture::Validate(MixUpdateCyclePtr Cycle)
 	//Check here if Source is VT
 	if (Source && !CanHandleAsset(Source))
 	{
-		auto ErrorType = static_cast<int32>(EMixerErrorType::UNSUPPORTED_TYPE);
+		auto ErrorType = static_cast<int32>(ETextureGraphErrorType::UNSUPPORTED_TYPE);
 
 		UClass* Type = Source->GetClass();
 		FString TypeName = Source->VirtualTextureStreaming ? "Virtual Texture" : Type->GetName();
 
-		MixerEngine::GetErrorReporter(ParentMix)->ReportError(ErrorType, FString::Printf(TEXT("%s not supported at the moment"), *TypeName), GetParentNode());
+		TextureGraphEngine::GetErrorReporter(ParentMix)->ReportError(ErrorType, FString::Printf(TEXT("%s not supported at the moment"), *TypeName), GetParentNode());
 		return false;
 	}
 	
diff --git a/Engine/Plugins/Experimental/TextureGraph/Source/TextureGraph/Private/Expressions/Maths/TG_Expression_Blend.cpp b/Engine/Plugins/Experimental/TextureGraph/Source/TextureGraph/Private/Expressions/Maths/TG_Expression_Blend.cpp
--- a/Engine/Plugins/Experimental/TextureGraph/Source/TextureGraph/Private/Expressions/Maths/TG_Expression_Blend.cpp
+++ b/Engine/Plugins/Experimental/TextureGraph/Source/TextureGraph/Private/Expressions/Maths/TG_Expression_Blend.cpp
@@ -10,7 +10,7 @@ void UTG_Expression_Blend::Evaluate(FTG_EvaluationContext* InContext)
 	if (!Mask)
 	{
 		/// Temporary, since our Grayscale texture saving (for test framework) isn't working correctly
-		if (MixerEngine::IsTestMode())
+		if (TextureGraphEngine::IsTestMode())
 		{
 			Mask = FTG_Texture::GetWhite();
 		}
diff --git a/Engine/Plugins/Experimental/TextureGraph/Source/TextureGraph/Private/Expressions/Maths/TG_Expression_Invert.cpp b/Engine/Plugins/Experimental/TextureGraph/Source/TextureGraph/Private/Expressions/Maths/TG_Expression_Invert.cpp
--- a/Engine/Plugins/Experimental/TextureGraph/Source/TextureGraph/Private/Expressions/Maths/TG_Expression_Invert.cpp
+++ b/Engine/Plugins/Experimental/TextureGraph/Source/TextureGraph/Private/Expressions/Maths/TG_Expression_Invert.cpp
@@ -15,7 +15,7 @@ void UTG_Expression_Invert::Evaluate(FTG_EvaluationContext* InContext)
 		return;
 	}
 
-	const RenderMaterial_FXPtr RenderMaterial = MixerEngine::GetMaterialManager()->CreateMaterial_FX<VSH_Simple, FSH_Invert>(TEXT("T_Invert"));
+	const RenderMaterial_FXPtr RenderMaterial = TextureGraphEngine::GetMaterialManager()->CreateMaterial_FX<VSH_Simple, FSH_Invert>(TEXT("T_Invert"));
 	check(RenderMaterial);
 
 	Output = T_Invert::Create(InContext->Cycle, Output.GetBufferDescriptor(), Input, MaxValue, InContext->TargetId);
diff --git a/Engine/Plugins/Experimental/TextureGraph/Source/TextureGraph/Private/Expressions/Output/TG_Expression_Output.cpp b/Engine/Plugins/Experimental/TextureGraph/Source/TextureGraph/Private/Expressions/Output/TG_Expression_Output.cpp
--- a/Engine/Plugins/Experimental/TextureGraph/Source/TextureGraph/Private/Expressions/Output/TG_Expression_Output.cpp
+++ b/Engine/Plugins/Experimental/TextureGraph/Source/TextureGraph/Private/Expressions/Output/TG_Expression_Output.cpp
@@ -38,7 +38,7 @@ void UTG_Expression_Output::Evaluate(FTG_EvaluationContext* InContext)
 			/// If the descriptors are not the same
 			if (OutputDesc != InputDesc)
 			{
-				RenderMaterial_FXPtr RenderMaterial = MixerEngine::GetMaterialManager()->CreateMaterialOfType_FX<Fx_FullScreenCopy>(TEXT("Tex::FullScreenCopy"));
+				RenderMaterial_FXPtr RenderMaterial = TextureGraphEngine::GetMaterialManager()->CreateMaterialOfType_FX<Fx_FullScreenCopy>(TEXT("Tex::FullScreenCopy"));
 
 				check(RenderMaterial);
 
diff --git a/Engine/Plugins/Experimental/TextureGraph/Source/TextureGraph/Private/Expressions/Procedural/TG_Expression_Noise.cpp b/Engine/Plugins/Experimental/TextureGraph/Source/TextureGraph/Private/Expressions/Procedural/TG_Expression_Noise.cpp
--- a/Engine/Plugins/Experimental/TextureGraph/Source/TextureGraph/Private/Expressions/Procedural/TG_Expression_Noise.cpp
+++ b/Engine/Plugins/Experimental/TextureGraph/Source/TextureGraph/Private/Expressions/Procedural/TG_Expression_Noise.cpp
@@ -12,7 +12,7 @@ void UTG_Expression_Noise::Evaluate(FTG_EvaluationContext* InContext)
 	FSH_NoiseMask::FPermutationDomain PermutationVector;
 	PermutationVector.Set<FVar_NoiseType>(NoiseType);
 
-	const RenderMaterial_FXPtr RenderMaterial = MixerEngine::GetMaterialManager()->CreateMaterial_FX<VSH_Simple, FSH_NoiseMask>(TEXT("T_NoiseMask"), PermutationVector);
+	const RenderMaterial_FXPtr RenderMaterial = TextureGraphEngine::GetMaterialManager()->CreateMaterial_FX<VSH_Simple, FSH_NoiseMask>(TEXT("T_NoiseMask"), PermutationVector);
 	check(RenderMaterial);
 
 	JobUPtr RenderJob = std::make_unique<Job>(InContext->Cycle->GetMix(), InContext->TargetId, std::static_pointer_cast<BlobTransform>(RenderMaterial));
diff --git a/Engine/Plugins/Experimental/TextureGraph/Source/TextureGraph/Private/Expressions/Procedural/TG_Expression_Pattern.cpp b/Engine/Plugins/Experimental/TextureGraph/Source/TextureGraph/Private/Expressions/Procedural/TG_Expression_Pattern.cpp
--- a/Engine/Plugins/Experimental/TextureGraph/Source/TextureGraph/Private/Expressions/Procedural/TG_Expression_Pattern.cpp
+++ b/Engine/Plugins/Experimental/TextureGraph/Source/TextureGraph/Private/Expressions/Procedural/TG_Expression_Pattern.cpp
@@ -13,7 +13,7 @@ void UTG_Expression_Pattern::Evaluate(FTG_EvaluationContext* InContext)
 	FSH_PatternMask::FPermutationDomain PermutationVector;
 	PermutationVector.Set<FVar_PatternType>(PatternType);
 
-	const RenderMaterial_FXPtr RenderMaterial = MixerEngine::GetMaterialManager()->CreateMaterial_FX<VSH_Simple, FSH_PatternMask>(TEXT("T_PatternMask"), PermutationVector);
+	const RenderMaterial_FXPtr RenderMaterial = TextureGraphEngine::GetMaterialManager()->CreateMaterial_FX<VSH_Simple, FSH_PatternMask>(TEXT("T_PatternMask"), PermutationVector);
 	check(RenderMaterial);
 
 	JobUPtr RenderJob = std::make_unique<Job>(InContext->Cycle->GetMix(), InContext->TargetId, std::static_pointer_cast<BlobTransform>(RenderMaterial));
diff --git a/Engine/Plugins/Experimental/TextureGraph/Source/TextureGraph/Private/Expressions/TG_Expression_MaterialBase.cpp b/Engine/Plugins/Experimental/TextureGraph/Source/TextureGraph/Private/Expressions/TG_Expression_MaterialBase.cpp
--- a/Engine/Plugins/Experimental/TextureGraph/Source/TextureGraph/Private/Expressions/TG_Expression_MaterialBase.cpp
+++ b/Engine/Plugins/Experimental/TextureGraph/Source/TextureGraph/Private/Expressions/TG_Expression_MaterialBase.cpp
@@ -71,7 +71,7 @@ void UTG_Expression_MaterialBase::Initialize()
 
 TiledBlobPtr UTG_Expression_MaterialBase::CreateRenderMaterialJob(FTG_EvaluationContext* InContext, const FString& InName, const FString& InMaterialPath, const BufferDescriptor& InDescriptor, EDrawMaterialAttributeTarget InDrawMaterialAttributeTarget)
 {
-	const RenderMaterial_BPPtr RenderMaterial = MixerEngine::GetMaterialManager()->CreateMaterial_BP(InName, InMaterialPath);
+	const RenderMaterial_BPPtr RenderMaterial = TextureGraphEngine::GetMaterialManager()->CreateMaterial_BP(InName, InMaterialPath);
 	return CreateRenderMaterialJob(InContext, RenderMaterial, InDescriptor, InDrawMaterialAttributeTarget);
 }
 
diff --git a/Engine/Plugins/Experimental/TextureGraph/Source/TextureGraph/Private/TG_HelperFunctions.cpp b/Engine/Plugins/Experimental/TextureGraph/Source/TextureGraph/Private/TG_HelperFunctions.cpp
--- a/Engine/Plugins/Experimental/TextureGraph/Source/TextureGraph/Private/TG_HelperFunctions.cpp
+++ b/Engine/Plugins/Experimental/TextureGraph/Source/TextureGraph/Private/TG_HelperFunctions.cpp
@@ -142,13 +142,13 @@ AsyncBool FTG_HelperFunctions::ExportAsync(UTextureGraph* InTextureGraph, FStrin
 				//Log Error to Error System
 				if (!HasOutputs)
 				{
-					auto ErrorType = static_cast<int32>(EMixerErrorType::UNSUPPORTED_TYPE);
-					MixerEngine::GetErrorReporter(InTextureGraph)->ReportError(ErrorType, FString::Format(TEXT("Texture Export Error : No valid output found for OutputSetting {0}"), { OutputSetting.OutputName.ToString() }), nullptr);
+					auto ErrorType = static_cast<int32>(ETextureGraphErrorType::UNSUPPORTED_TYPE);
+					TextureGraphEngine::GetErrorReporter(InTextureGraph)->ReportError(ErrorType, FString::Format(TEXT("Texture Export Error : No valid output found for OutputSetting {0}"), { OutputSetting.OutputName.ToString() }), nullptr);
 				}
 				if (!IsNameValid)
 				{
-					auto ErrorType = static_cast<int32>(EMixerErrorType::UNSUPPORTED_TYPE);
-					MixerEngine::GetErrorReporter(InTextureGraph)->ReportError(ErrorType, FString::Format(TEXT("Texture Export Error : Invalid path set for OutputSetting {0}"), { OutputSetting.OutputName.ToString() }), nullptr);
+					auto ErrorType = static_cast<int32>(ETextureGraphErrorType::UNSUPPORTED_TYPE);
+					TextureGraphEngine::GetErrorReporter(InTextureGraph)->ReportError(ErrorType, FString::Format(TEXT("Texture Export Error : Invalid path set for OutputSetting {0}"), { OutputSetting.OutputName.ToString() }), nullptr);
 				}
 			}
 		}
@@ -162,7 +162,7 @@ AsyncBool FTG_HelperFunctions::ExportAsync(UTextureGraph* InTextureGraph, FStrin
 
 void FTG_HelperFunctions::InitTargets(UTextureGraph* InTextureGraph)
 {
-	MixerEngine::RegisterErrorReporter(InTextureGraph, std::make_shared<FMixerErrorReporter>());
+	TextureGraphEngine::RegisterErrorReporter(InTextureGraph, std::make_shared<FTextureGraphErrorReporter>());
 	/// Now run the update Cycle
 
 	int num = 1;
@@ -227,27 +227,27 @@ AsyncBool FTG_HelperFunctions::RenderAsync(UTextureGraph* InTextureGraph, JobBat
 			std::unique_lock<std::mutex> lock(*mutex);
 
 			/// Lock the engine
-			if (MixerEngine::IsTestMode())
-				MixerEngine::Lock();
+			if (TextureGraphEngine::IsTestMode())
+				TextureGraphEngine::Lock();
 
 			/// We must get out of the loop if the engine is being destroyed. This can happen
 			/// when an test has offended the time limit for the test. Can result in tests exiting
-			/// and calling MixerEngine::Destroy even though this loop might still be running. 
+			/// and calling TextureGraphEngine::Destroy even though this loop might still be running. 
 			/// We need a safe passage out of this.
-			while (!*IsMixRendered && !MixerEngine::IsDestroying())
+			while (!*IsMixRendered && !TextureGraphEngine::IsDestroying())
 			{
 				Util::OnGameThread([]()
 				{
-					if (!MixerEngine::IsDestroying())
-						MixerEngine::Update(0);
+					if (!TextureGraphEngine::IsDestroying())
+						TextureGraphEngine::Update(0);
 				});
 
 				std::this_thread::sleep_for(std::chrono::milliseconds(10));
 			}
 
 			/// Unlock it
-			if (MixerEngine::IsTestMode())
-				MixerEngine::Unlock();
+			if (TextureGraphEngine::IsTestMode())
+				TextureGraphEngine::Unlock();
 		}
 	});
 
@@ -256,11 +256,11 @@ AsyncBool FTG_HelperFunctions::RenderAsync(UTextureGraph* InTextureGraph, JobBat
 		Batch->OnDone([=, FWD_PROMISE(promise)](JobBatch*) mutable
 		{
 			/// We need to do this on the background thread because we can potentially have a deadlock
-			/// where this thread blocks on the mutex wait and the MixerEngine::Update above is waiting
+			/// where this thread blocks on the mutex wait and the TextureGraphEngine::Update above is waiting
 			/// for Util::OnGameThread update. This situation cannot be allowed to happen
 			Util::OnBackgroundThread([=, FWD_PROMISE(promise)]() mutable
 			{
-				/// Set the atomic so that the loop above with MixerEngine::Update can exit
+				/// Set the atomic so that the loop above with TextureGraphEngine::Update can exit
 				/// and release the mutex that we're going to acquire below
 				*IsMixRendered = true;
 
@@ -285,7 +285,7 @@ AsyncBool FTG_HelperFunctions::RenderAsync(UTextureGraph* InTextureGraph, JobBat
 		{
 			Util::OnGameThread([=]()
 				{
-					MixerEngine::GetScheduler()->AddBatch(Batch);
+					TextureGraphEngine::GetScheduler()->AddBatch(Batch);
 				});
 		}
 	}); 
diff --git a/Engine/Plugins/Experimental/TextureGraph/Source/TextureGraph/Private/TG_Node.cpp b/Engine/Plugins/Experimental/TextureGraph/Source/TextureGraph/Private/TG_Node.cpp
--- a/Engine/Plugins/Experimental/TextureGraph/Source/TextureGraph/Private/TG_Node.cpp
+++ b/Engine/Plugins/Experimental/TextureGraph/Source/TextureGraph/Private/TG_Node.cpp
@@ -16,8 +16,8 @@ bool UTG_Node::Validate(MixUpdateCyclePtr Cycle)
 		UMixInterface* ParentMix = Cast<UMixInterface>(GetOutermostObject());
 		for (auto Warning : WarningStack)
 		{
-			auto ErrorType = static_cast<int32>(EMixerErrorType::NODE_WARNING);
-			FMixerErrorReport Report = MixerEngine::GetErrorReporter(ParentMix)->ReportWarning(ErrorType, Warning.ToString(), this);
+			auto ErrorType = static_cast<int32>(ETextureGraphErrorType::NODE_WARNING);
+			FTextureGraphErrorReport Report = TextureGraphEngine::GetErrorReporter(ParentMix)->ReportWarning(ErrorType, Warning.ToString(), this);
 		}
 	}
 
diff --git a/Engine/Plugins/Experimental/TextureGraph/Source/TextureGraph/Private/TextureGraph.cpp b/Engine/Plugins/Experimental/TextureGraph/Source/TextureGraph/Private/TextureGraph.cpp
--- a/Engine/Plugins/Experimental/TextureGraph/Source/TextureGraph/Private/TextureGraph.cpp
+++ b/Engine/Plugins/Experimental/TextureGraph/Source/TextureGraph/Private/TextureGraph.cpp
@@ -21,7 +21,7 @@
 void UTextureGraph::Construct(FString InName)
 {
 	// On the first new texture script we set the engine in run mode
-	MixerEngine::SetRunEngine();
+	TextureGraphEngine::SetRunEngine();
 
 	Settings = NewObject<UMixSettings>(this);
 	bInvalidateTextures = false;
@@ -54,7 +54,7 @@ void UTextureGraph::Serialize(FArchive& Ar)
 void UTextureGraph::PostLoad()
 {
 	// On the first script load we set the engine in run mode as well
-	MixerEngine::SetRunEngine();
+	TextureGraphEngine::SetRunEngine();
 
 	Super::PostLoad();
 	bInvalidateTextures = false;
@@ -146,7 +146,7 @@ void UTextureGraph::TriggerUpdate(bool Tweaking)
 	InvalidateInfo.Details.bTweaking = Tweaking;
 	InvalidateInfo.Details.Mix = this;
 
-	MixerEngine::GetMixManager()->InvalidateMix(this, InvalidateInfo.Details);
+	TextureGraphEngine::GetMixManager()->InvalidateMix(this, InvalidateInfo.Details);
 }
 
 void UTextureGraph::InvalidateAll()
@@ -154,9 +154,9 @@ void UTextureGraph::InvalidateAll()
 	FInvalidationDetails Details = FInvalidationDetails().All();
 	Details.Mix = this;
 
-	InvalidationFrameId = MixerEngine::GetFrameId();
+	InvalidationFrameId = TextureGraphEngine::GetFrameId();
 
-	MixerEngine::GetMixManager()->InvalidateMix(this, Details);
+	TextureGraphEngine::GetMixManager()->InvalidateMix(this, Details);
 }
 
 void UTextureGraph::Log() const
diff --git a/Engine/Plugins/Experimental/TextureGraph/Source/TextureGraph/Public/TG_SystemTypes.h b/Engine/Plugins/Experimental/TextureGraph/Source/TextureGraph/Public/TG_SystemTypes.h
--- a/Engine/Plugins/Experimental/TextureGraph/Source/TextureGraph/Public/TG_SystemTypes.h
+++ b/Engine/Plugins/Experimental/TextureGraph/Source/TextureGraph/Public/TG_SystemTypes.h
@@ -11,7 +11,7 @@ DECLARE_LOG_CATEGORY_EXTERN(LogTextureGraph, Log, All);
 using FTG_Index = int32;
 using FTG_Indices = TArray<FTG_Index>;
 
-using ErrorReportMap = TMap<int32, TArray<struct FMixerErrorReport>>;
+using ErrorReportMap = TMap<int32, TArray<struct FTextureGraphErrorReport>>;
 USTRUCT()
 struct TEXTUREGRAPH_API FTG_Id
 {
diff --git a/Engine/Plugins/Experimental/TextureGraph/Source/TextureGraph/Public/TextureGraph.h b/Engine/Plugins/Experimental/TextureGraph/Source/TextureGraph/Public/TextureGraph.h
--- a/Engine/Plugins/Experimental/TextureGraph/Source/TextureGraph/Public/TextureGraph.h
+++ b/Engine/Plugins/Experimental/TextureGraph/Source/TextureGraph/Public/TextureGraph.h
@@ -7,7 +7,7 @@
 #include "TextureGraph.generated.h"
 
 class UTG_Graph;
-using ErrorReportMap = TMap<int32, TArray<FMixerErrorReport>>;
+using ErrorReportMap = TMap<int32, TArray<FTextureGraphErrorReport>>;
 
 UCLASS(BlueprintType)
 class TEXTUREGRAPH_API UTextureGraph : public UMixInterface
diff --git a/Engine/Plugins/Experimental/TextureGraph/Source/TextureGraphEditor/Private/EdGraph/TG_EditorErrorReporter.cpp b/Engine/Plugins/Experimental/TextureGraph/Source/TextureGraphEditor/Private/EdGraph/TG_EditorErrorReporter.cpp
--- a/Engine/Plugins/Experimental/TextureGraph/Source/TextureGraphEditor/Private/EdGraph/TG_EditorErrorReporter.cpp
+++ b/Engine/Plugins/Experimental/TextureGraph/Source/TextureGraphEditor/Private/EdGraph/TG_EditorErrorReporter.cpp
@@ -7,14 +7,14 @@
 #include "EdGraph/TG_EdGraph.h"
 #include "EdGraph/TG_EdGraphNode.h"
 
-FMixerErrorReport FTG_EditorErrorReporter::ReportLog(int32 ErrorId, const FString& ErrorMsg, UObject* ReferenceObj /*= nullptr*/)
+FTextureGraphErrorReport FTG_EditorErrorReporter::ReportLog(int32 ErrorId, const FString& ErrorMsg, UObject* ReferenceObj /*= nullptr*/)
 {
 	return {ErrorId, ErrorMsg, ReferenceObj};
 }
 
-FMixerErrorReport FTG_EditorErrorReporter::Report(int32 ErrorId, const FString& ErrorMsg, UObject* ReferenceObj, EMessageSeverity::Type ErrorType)
+FTextureGraphErrorReport FTG_EditorErrorReporter::Report(int32 ErrorId, const FString& ErrorMsg, UObject* ReferenceObj, EMessageSeverity::Type ErrorType)
 {
-	FMixerErrorReport Report {ErrorId, ErrorMsg, ReferenceObj};
+	FTextureGraphErrorReport Report {ErrorId, ErrorMsg, ReferenceObj};
 	
 	// If the reference object is a TG_Node
 	if (UTG_Node* NodeObj = Cast<UTG_Node>(ReferenceObj))
@@ -33,10 +33,10 @@ FMixerErrorReport FTG_EditorErrorReporter::Report(int32 ErrorId, const FString&
 			}
 		}
 	}
-	TArray<FMixerErrorReport>& ErrorEntry = CompilationErrors.FindOrAdd(ErrorId);
+	TArray<FTextureGraphErrorReport>& ErrorEntry = CompilationErrors.FindOrAdd(ErrorId);
 
 	// check to see if this error on this object is already reported
-	const bool ReportAlreadyExists = ErrorEntry.ContainsByPredicate([&](const FMixerErrorReport& InReport) -> bool
+	const bool ReportAlreadyExists = ErrorEntry.ContainsByPredicate([&](const FTextureGraphErrorReport& InReport) -> bool
 	{
 		return InReport.ReferenceObj == ReferenceObj;
 	});
@@ -48,12 +48,12 @@ FMixerErrorReport FTG_EditorErrorReporter::Report(int32 ErrorId, const FString&
 	return Report;
 }
 
-FMixerErrorReport FTG_EditorErrorReporter::ReportWarning(int32 ErrorId, const FString& ErrorMsg, UObject* ReferenceObj /*= nullptr*/)
+FTextureGraphErrorReport FTG_EditorErrorReporter::ReportWarning(int32 ErrorId, const FString& ErrorMsg, UObject* ReferenceObj /*= nullptr*/)
 {
 	return Report(ErrorId, ErrorMsg, ReferenceObj, EMessageSeverity::Warning);
 }
 
-FMixerErrorReport FTG_EditorErrorReporter::ReportError(int32 ErrorId, const FString& ErrorMsg, UObject* ReferenceObj /*= nullptr*/)
+FTextureGraphErrorReport FTG_EditorErrorReporter::ReportError(int32 ErrorId, const FString& ErrorMsg, UObject* ReferenceObj /*= nullptr*/)
 {
 	return Report(ErrorId, ErrorMsg, ReferenceObj, EMessageSeverity::Error);
 }
@@ -62,7 +62,7 @@ void FTG_EditorErrorReporter::Clear()
 {
 	for(auto ErrorEntries : CompilationErrors)
 	{
-		for(const FMixerErrorReport& Error : ErrorEntries.Value)
+		for(const FTextureGraphErrorReport& Error : ErrorEntries.Value)
 		{
 			// Expect the reference object to be a TG_Node
 			const UTG_Node* NodeObj = Cast<UTG_Node>(Error.ReferenceObj);
@@ -81,6 +81,6 @@ void FTG_EditorErrorReporter::Clear()
 		}
 	}
 	
-	FMixerErrorReporter::Clear();
+	FTextureGraphErrorReporter::Clear();
 }
 
diff --git a/Engine/Plugins/Experimental/TextureGraph/Source/TextureGraphEditor/Private/EdGraph/TG_EditorErrorReporter.h b/Engine/Plugins/Experimental/TextureGraph/Source/TextureGraphEditor/Private/EdGraph/TG_EditorErrorReporter.h
--- a/Engine/Plugins/Experimental/TextureGraph/Source/TextureGraphEditor/Private/EdGraph/TG_EditorErrorReporter.h
+++ b/Engine/Plugins/Experimental/TextureGraph/Source/TextureGraphEditor/Private/EdGraph/TG_EditorErrorReporter.h
@@ -1,30 +1,30 @@
 // Copyright Epic Games, Inc. All Rights Reserved.
 #pragma once
-#include "MixerErrorReporter.h"
+#include "TextureGraphErrorReporter.h"
 #include "TG_Editor.h"
 #include "Logging/TokenizedMessage.h"
 
 
-class FTG_EditorErrorReporter: public FMixerErrorReporter
+class FTG_EditorErrorReporter: public FTextureGraphErrorReporter
 {
 public:
 	FTG_EditorErrorReporter(FTG_Editor* InEditor)
-	: FMixerErrorReporter()
+	: FTextureGraphErrorReporter()
 	, Editor(InEditor)
 	{
 	}
 	virtual ~FTG_EditorErrorReporter() { }
 
-	virtual FMixerErrorReport ReportLog(int32 ErrorId, const FString& ErrorMsg, UObject* ReferenceObj = nullptr) override;
+	virtual FTextureGraphErrorReport ReportLog(int32 ErrorId, const FString& ErrorMsg, UObject* ReferenceObj = nullptr) override;
 	
-	virtual FMixerErrorReport ReportWarning(int32 ErrorId, const FString& ErrorMsg, UObject* ReferenceObj = nullptr) override;
+	virtual FTextureGraphErrorReport ReportWarning(int32 ErrorId, const FString& ErrorMsg, UObject* ReferenceObj = nullptr) override;
 
-	virtual FMixerErrorReport ReportError(int32 ErrorId, const FString& ErrorMsg, UObject* ReferenceObj = nullptr) override;
+	virtual FTextureGraphErrorReport ReportError(int32 ErrorId, const FString& ErrorMsg, UObject* ReferenceObj = nullptr) override;
 
 	virtual void Clear() override;
 
 private:
-	FMixerErrorReport Report(int32 ErrorId, const FString& ErrorMsg, UObject* ReferenceObj, EMessageSeverity::Type ErrorType);
+	FTextureGraphErrorReport Report(int32 ErrorId, const FString& ErrorMsg, UObject* ReferenceObj, EMessageSeverity::Type ErrorType);
 
 	FTG_Editor* Editor;
 };
diff --git a/Engine/Plugins/Experimental/TextureGraph/Source/TextureGraphEditor/Private/STG_EditorViewport.cpp b/Engine/Plugins/Experimental/TextureGraph/Source/TextureGraphEditor/Private/STG_EditorViewport.cpp
--- a/Engine/Plugins/Experimental/TextureGraph/Source/TextureGraphEditor/Private/STG_EditorViewport.cpp
+++ b/Engine/Plugins/Experimental/TextureGraph/Source/TextureGraphEditor/Private/STG_EditorViewport.cpp
@@ -123,7 +123,6 @@ void FTG_EditorViewportClient::Tick(float DeltaSeconds)
 void FTG_EditorViewportClient::Draw(FViewport* InViewport, FCanvas* Canvas)
 {
 	FEditorViewportClient::Draw(InViewport, Canvas);
-	//_mixerAssetEditorPtr.Pin()->DrawMessages(InViewport, Canvas);
 }
 
 bool FTG_EditorViewportClient::ShouldOrbitCamera() const
diff --git a/Engine/Plugins/Experimental/TextureGraph/Source/TextureGraphEditor/Private/TG_Editor.cpp b/Engine/Plugins/Experimental/TextureGraph/Source/TextureGraphEditor/Private/TG_Editor.cpp
--- a/Engine/Plugins/Experimental/TextureGraph/Source/TextureGraphEditor/Private/TG_Editor.cpp
+++ b/Engine/Plugins/Experimental/TextureGraph/Source/TextureGraphEditor/Private/TG_Editor.cpp
@@ -185,7 +185,7 @@ void FTG_Editor::InitEditor(const EToolkitMode::Type Mode, const TSharedPtr< cla
 	ViewportSettings.OnViewportMaterialChangeEvent.AddSP(this, &FTG_Editor::OnViewportSettingsChanged);
 	ViewportSettings.OnMaterialMappingChangedEvent.AddSP(this, &FTG_Editor::OnMaterialMappingChanged);
 
-	MixerEngine::RegisterErrorReporter(EditedTextureGraph, std::make_shared<FTG_EditorErrorReporter>(this));
+	TextureGraphEngine::RegisterErrorReporter(EditedTextureGraph, std::make_shared<FTG_EditorErrorReporter>(this));
 
 	TG_EdGraph = NewObject<UTG_EdGraph>(EditedTextureGraph, UTG_EdGraph::StaticClass(), NAME_None, RF_Transactional);
 	TG_EdGraph->Schema = UTG_EdGraphSchema::StaticClass();
@@ -1492,19 +1492,19 @@ void FTG_Editor::RefreshErrors()
 {
 	TArray< TSharedRef<class FTokenizedMessage> > Messages;
 
-	FMixerErrorReporter* ErrorReporter = MixerEngine::GetErrorReporter(EditedTextureGraph);
+	FTextureGraphErrorReporter* ErrorReporter = TextureGraphEngine::GetErrorReporter(EditedTextureGraph);
 
 	if (ErrorReporter)
 	{
 		// loop through array of messages/nodeRef and highlight them
-		const auto CompileErrors = MixerEngine::GetErrorReporter(EditedTextureGraph)->GetCompilationErrors();
+		const auto CompileErrors = TextureGraphEngine::GetErrorReporter(EditedTextureGraph)->GetCompilationErrors();
 
 		// compute error crc
 		FString NewErrorHash;
 
 		for (auto ErrorEntry : CompileErrors)
 		{
-			for (const FMixerErrorReport& Error : ErrorEntry.Value)
+			for (const FTextureGraphErrorReport& Error : ErrorEntry.Value)
 			{
 				NewErrorHash += FMD5::HashAnsiString(*Error.ErrorMsg);
 			}
@@ -1516,7 +1516,7 @@ void FTG_Editor::RefreshErrors()
 			for (auto ErrorEntry : CompileErrors)
 			{
 				// report all errors for this type
-				for (FMixerErrorReport& Error : ErrorEntry.Value)
+				for (FTextureGraphErrorReport& Error : ErrorEntry.Value)
 				{
 					FString ErrorString = Error.GetFormattedMessage();
 					TSharedRef<FTokenizedMessage> Line = FTokenizedMessage::Create(EMessageSeverity::Error);
diff --git a/Engine/Plugins/Experimental/TextureGraph/Source/TextureGraphEditor/Private/TG_RenderModeManager.cpp b/Engine/Plugins/Experimental/TextureGraph/Source/TextureGraphEditor/Private/TG_RenderModeManager.cpp
--- a/Engine/Plugins/Experimental/TextureGraph/Source/TextureGraphEditor/Private/TG_RenderModeManager.cpp
+++ b/Engine/Plugins/Experimental/TextureGraph/Source/TextureGraphEditor/Private/TG_RenderModeManager.cpp
@@ -2,7 +2,7 @@
 #include "TG_RenderModeManager.h"
 #include "FxMat/MaterialManager.h"
 #include "FxMat/RenderMaterial_BP.h"
-#include "MixerEngine.h"
+#include "TextureGraphEngine.h"
 #include "TextureGraph.h"
 #include "Model/Mix/Mix.h"
 #include "Device/FX/Device_FX.h"
@@ -105,11 +105,11 @@ void TG_RenderModeManager::InitializeDefaultMaterials(int totalTargets)
 	for (int targetId = 0; targetId < totalTargets; targetId++)
 	{
 		TMap<FName, RenderMaterial_BPPtr> TargetMaterialsMap;
-		TargetMaterialsMap.Add(ViewportMaterial.GetFName(), MixerEngine::GetMaterialManager()->CreateMaterial_BP(ViewportMaterial.GetName(), ViewportMaterial));
+		TargetMaterialsMap.Add(ViewportMaterial.GetFName(), TextureGraphEngine::GetMaterialManager()->CreateMaterial_BP(ViewportMaterial.GetName(), ViewportMaterial));
 
 		for(const FMaterialMappingInfo& MaterialMappingInfo : MaterialMappingInfos)
 		{
-			TargetMaterialsMap.Add(MaterialMappingInfo.MaterialInput, MixerEngine::GetMaterialManager()->CreateMaterial_BP(MaterialMappingInfo.MaterialInput.ToString(), TEXT("Scene/RendermodeSource")));
+			TargetMaterialsMap.Add(MaterialMappingInfo.MaterialInput, TextureGraphEngine::GetMaterialManager()->CreateMaterial_BP(MaterialMappingInfo.MaterialInput.ToString(), TEXT("Scene/RendermodeSource")));
 		}
 		
 		_renderModeMaterials.Add(targetId, TargetMaterialsMap);
diff --git a/Engine/Plugins/Experimental/TextureGraph/Source/TextureGraphEditor/Private/TextureGraphEditorModule.cpp b/Engine/Plugins/Experimental/TextureGraph/Source/TextureGraphEditor/Private/TextureGraphEditorModule.cpp
--- a/Engine/Plugins/Experimental/TextureGraph/Source/TextureGraphEditor/Private/TextureGraphEditorModule.cpp
+++ b/Engine/Plugins/Experimental/TextureGraph/Source/TextureGraphEditor/Private/TextureGraphEditorModule.cpp
@@ -5,7 +5,7 @@
 #include "AssetToolsModule.h"
 #include "EdGraph/TG_EditorGraphNodeFactory.h"
 #include "EdGraph/TG_EditorGraphPanelPinFactory.h"
-#include "MixerEngine.h"
+#include "TextureGraphEngine.h"
 #include "TG_AssetTypeActions.h"
 #include "TG_Editor.h"
 #include "TG_EditorCommands.h"
@@ -49,7 +49,7 @@ void FTextureGraphEditorModule::StartupModule()
 	GraphPanelPinFactory = MakeShared<FTG_EditorGraphPanelPinFactory>();
 	FEdGraphUtilities::RegisterVisualPinFactory(GraphPanelPinFactory);
 	
-	StartMixerEngine();
+	StartTextureGraphEngine();
 }
 
 void FTextureGraphEditorModule::ShutdownModule()
@@ -74,33 +74,33 @@ void FTextureGraphEditorModule::ShutdownModule()
 	// Unregister slate style overrides
 	FTG_Style::Unregister();
 	
-	ShutdownMixerEngine();
+	ShutdownTextureGraphEngine();
 }
-void FTextureGraphEditorModule::StartMixerEngine()
+void FTextureGraphEditorModule::StartTextureGraphEngine()
 {
-	if (!MixerEngine::GetInstance())
+	if (!TextureGraphEngine::GetInstance())
 	{
-		// In case of editor, this is the place where we create a new mixer engine.
+		// In case of editor, this is the place where we create a new Texture Graph engine.
 		// This is done only once and will get destroyed when we shutdown unreal.
-		MixerEngine::Create(false);
-		check(MixerEngine::GetInstance());
+		TextureGraphEngine::Create(false);
+		check(TextureGraphEngine::GetInstance());
 	}
 	TickDelegate = FTickerDelegate::CreateRaw(this, &FTextureGraphEditorModule::Tick);
 	TickDelegateHandle = FTSTicker::GetCoreTicker().AddTicker(TickDelegate);
 }
-void FTextureGraphEditorModule::ShutdownMixerEngine()
+void FTextureGraphEditorModule::ShutdownTextureGraphEngine()
 {
-	if (MixerEngine::GetInstance())
+	if (TextureGraphEngine::GetInstance())
 	{
-		MixerEngine::Destroy();
-		check(!MixerEngine::GetInstance());
+		TextureGraphEngine::Destroy();
+		check(!TextureGraphEngine::GetInstance());
 		FTSTicker::GetCoreTicker().RemoveTicker(TickDelegateHandle);
 	}
 }
 bool FTextureGraphEditorModule::Tick(float deltaTime)
 {
-	if (MixerEngine::GetInstance())
-		MixerEngine::Update(deltaTime);
+	if (TextureGraphEngine::GetInstance())
+		TextureGraphEngine::Update(deltaTime);
 	return true;
 }
 void FTextureGraphEditorModule::RegisterAssetTypeAction(IAssetTools& AssetTools, TSharedRef<IAssetTypeActions> Action)
diff --git a/Engine/Plugins/Experimental/TextureGraph/Source/TextureGraphEditor/Public/TextureGraphEditorModule.h b/Engine/Plugins/Experimental/TextureGraph/Source/TextureGraphEditor/Public/TextureGraphEditorModule.h
--- a/Engine/Plugins/Experimental/TextureGraph/Source/TextureGraphEditor/Public/TextureGraphEditorModule.h
+++ b/Engine/Plugins/Experimental/TextureGraph/Source/TextureGraphEditor/Public/TextureGraphEditorModule.h
@@ -18,9 +18,9 @@ class TEXTUREGRAPHEDITOR_API FTextureGraphEditorModule : public IModuleInterface
 private:
 	/** All created asset type actions.  Cached here so that we can unregister them during shutdown. */
 	AssetTypeActionsArray			CreatedAssetTypeActions;
-	/** Delegate to run the Tick method in charge of running MixerEngine update*/
+	/** Delegate to run the Tick method in charge of running TextureGraphEngine update*/
 	FTickerDelegate					TickDelegate;
-	/** Handle of the delegate to run the Tick method in charge of running MixerEngine update*/
+	/** Handle of the delegate to run the Tick method in charge of running TextureGraphEngine update*/
 	FTSTicker::FDelegateHandle		TickDelegateHandle;
 	
 
@@ -31,8 +31,8 @@ public:
 	/** IModuleInterface implementation */
 	virtual void					StartupModule() override;
 	virtual void					ShutdownModule() override;
-	virtual void					StartMixerEngine();
-	virtual void					ShutdownMixerEngine();
+	virtual void					StartTextureGraphEngine();
+	virtual void					ShutdownTextureGraphEngine();
 	bool							Tick(float deltaTime);
 	void							RegisterAssetTypeAction(IAssetTools& AssetTools, TSharedRef<IAssetTypeActions> Action);
 	void							UnRegisterAllAssetTypeActions();
diff --git a/Engine/Plugins/Experimental/TextureGraph/Source/TextureGraphEngine/2D/TargetTextureSet.h b/Engine/Plugins/Experimental/TextureGraph/Source/TextureGraphEngine/2D/TargetTextureSet.h
new file mode 100644
--- /dev/null
+++ b/Engine/Plugins/Experimental/TextureGraph/Source/TextureGraphEngine/2D/TargetTextureSet.h
@@ -0,0 +1,61 @@
+// Copyright Epic Games, Inc. All Rights Reserved.
+#pragma once
+
+#include "CoreMinimal.h"
+#include "TextureSet.h"
+#include "2D/TextureType.h"
+#include "FxMat/RenderMaterial_BP.h"
+
+struct FMaterialMappingInfo;
+typedef TFunction<void(TiledBlobPtr)> TextureReadyCallback;
+
+class RenderMesh;
+typedef std::shared_ptr<RenderMesh>	RenderMeshPtr;
+
+
+class TEXTUREGRAPHENGINE_API TargetTextureSet : public TextureSet
+{
+protected:
+
+	mutable TMap<FName, TextureReadyCallback> Callbacks;		/// The callbacks to call when texture is set
+
+	int32							Id = -1;					/// The ID/index into the list of target texture sets
+	FString							Name;						/// Display name of this target set
+
+	RenderMeshPtr					Mesh = nullptr;				/// The render mesh that this this texture set is targetting. 
+																/// One texture set can only target one render mesh at a time.
+	
+
+	TMap<FName, TiledBlobRef>		BoundTexturesMap;			///Textures currently bound and dont want deallocation until new textures assigned
+	int32							RenderCount = 0;			/// The number of times this layer has been rendered to
+
+	virtual void					InitTex(int32 InTypeIndex) override;
+	virtual void					BindOnTextureUpdate(RenderMaterial_BPPtr InMaterial, FMaterialMappingInfo MaterialMappingInfo) const;
+	virtual void					RegisterCallback(TextureReadyCallback Callback, FMaterialMappingInfo MaterialMappingInfo) const;
+	virtual void					UnRegisterCallback(FName TextureName) const;
+public:
+									TargetTextureSet(int32 InId, const FString& InName, RenderMeshPtr InMesh, int32 InWidth, int32 InHeight);
+	virtual							~TargetTextureSet() override;
+
+	void							SetMesh(RenderMeshPtr InMesh);
+	
+	AsyncBufferResultPtr			BindTo(RenderMaterial_BPPtr Material, TArray<struct FMaterialMappingInfo> MaterialMappingInfos);
+	AsyncBufferResultPtr			BindTo(RenderMaterial_BPPtr Material, struct FMaterialMappingInfo MaterialMappingInfo);
+	AsyncBufferResultPtr			UnbindFrom(RenderMaterial_BPPtr Material);
+	
+	virtual void					SetTexture(FName TextureName, TiledBlobRef InTexture) override;
+	virtual void					FreeAt(FName TextureName) override;
+	virtual void					Free() override;
+	//////////////////////////////////////////////////////////////////////////
+	/// Inline functions
+	//////////////////////////////////////////////////////////////////////////
+	FORCEINLINE int32				GetId() const { return Id; }
+	FORCEINLINE FString				GetName() const { return Name; }
+	FORCEINLINE RenderMeshPtr		GetMesh() const { return Mesh; }
+	FORCEINLINE int32				GetRenderCount() const { return RenderCount; }
+	
+};
+
+typedef std::unique_ptr<TargetTextureSet>	TargetTextureSetPtr;
+typedef std::vector<TargetTextureSetPtr>	TargetTextureSetPtrVec;
+
diff --git a/Engine/Plugins/Experimental/TextureGraph/Source/TextureGraphEngine/2D/Tex.h b/Engine/Plugins/Experimental/TextureGraph/Source/TextureGraphEngine/2D/Tex.h
new file mode 100644
--- /dev/null
+++ b/Engine/Plugins/Experimental/TextureGraph/Source/TextureGraphEngine/2D/Tex.h
@@ -0,0 +1,207 @@
+// Copyright Epic Games, Inc. All Rights Reserved.
+#pragma once
+
+#include "CoreMinimal.h"
+#include "Data/Blob.h"
+#include "Data/TiledBlob.h"
+#include "Engine/Texture.h"
+#include "Engine/TextureRenderTarget2D.h"
+#include "FxMat/RenderMaterial_FX.h"
+#include "PixelFormat.h" 
+THIRD_PARTY_INCLUDES_START
+#include "continuable/continuable.hpp"
+THIRD_PARTY_INCLUDES_END
+#include "Helper/Promise.h"
+
+class UTexture2D;
+class UTextureRenderTarget2D;
+class UMaterial;
+struct DesiredImageProperties;
+class RenderMaterial;
+
+class Tex;
+
+typedef std::shared_ptr<Tex>		TexPtr;
+
+struct DesiredImageProperties
+{
+	FString							Name;
+	bool							bIsLinear = false;
+	bool							bForceSRGB = false;
+	bool							bIsGrayscale = false;
+	bool							bMipMaps = false;
+	bool							bPremultiplyAlpha = false;
+	bool							bCalculateDisplacementRange = false;
+	int32							Width = -1;
+	int32							Height = -1;
+	EPixelFormat					PixelFormat = static_cast<EPixelFormat>(-1);
+};
+
+//////////////////////////////////////////////////////////////////////////
+/// TexDescriptor: The Texture descriptor
+//////////////////////////////////////////////////////////////////////////
+struct TexDescriptor
+{
+	FString							Name;						/// [Optional] Name of the Texture
+	uint32							Width = 0;					/// Width of the Texture in pixels
+	uint32							Height = 0;					/// Height of the Texture in pixels
+	EPixelFormat					Format = EPixelFormat::PF_R8G8B8A8;	/// The pixel format
+	uint32							NumChannels = 4;				/// Channels per-pixels
+	bool							bMipMaps = false;			/// Whether to use mipmaps or not
+	bool							bAutoGenerateMipMaps = true;	/// Whether to auto-generate mipmaps or not
+	bool							bIsSRGB = false;				/// Whether the Texture is in SRGB space or not
+	bool							bCompress = false;			/// Whether to use compression on the Texture
+	bool							bUAV = false;				/// Whether the resource can create a UAV
+	ETextureCompressionQuality		CompressionQuality = ETextureCompressionQuality::TCQ_High;	/// What quality to use on the compression
+	TextureCompressionSettings		CompressionSettings = TextureCompressionSettings::TC_Default;/// Compression setting
+	FLinearColor					ClearColor = FLinearColor::Transparent;
+
+									TexDescriptor();
+									TexDescriptor(uint32 InWidth, uint32 InHeight, EPixelFormat InFormat);
+									TexDescriptor(uint32 InWidth, uint32 InHeight, EPixelFormat InFormat, bool bInSRGB);
+									TexDescriptor(uint32 InWidth, uint32 InHeight, EPixelFormat InFormat, bool bInMipMaps, bool bInAutoGenerateMipMaps, bool bInSRGB);
+	explicit						TexDescriptor(UTexture2D* TextureObj);
+	explicit						TexDescriptor(UTextureRenderTarget2D* RT);
+	explicit						TexDescriptor(const BufferDescriptor& BufferDesc);
+
+	BufferDescriptor				ToBufferDescriptor(uint32 NewWidth = 0, uint32 NewHeight = 0) const;
+	size_t							GetPitch() const;
+	FORCEINLINE size_t				Size() const { return GetPitch() * Height; }
+	HashType						HashValue() const;
+	HashType						Format_HashValue() const;
+};
+
+//////////////////////////////////////////////////////////////////////////
+
+/**
+ * This is a high level Texture class that wraps a 2D/1D Texture as well as a RenderTarget.
+ * Since we do a lot of texturing operations, we need to have some advanced requirements
+ * and that is why this class wraps a 2D Texture as well as a RenderTarget. It only has
+ * one state at a time though i.e. its either a UTexture or a RenderTarget. It can't be both.
+ * This works closely with the RCManager to efficiently manage the resources available to us
+ * within the system.
+ */
+class TEXTUREGRAPHENGINE_API Tex : public FGCObject
+{
+	friend class Device_FX;
+
+	//////////////////////////////////////////////////////////////////////////
+	/// FGCObject
+	//////////////////////////////////////////////////////////////////////////
+	virtual void					AddReferencedObjects(FReferenceCollector& Collector) override;
+	virtual FString					GetReferencerName() const override;
+
+public:
+	
+
+protected: 
+	TObjectPtr<UTexture2D>			Texture = nullptr;				/// The optional Texture
+	TObjectPtr<UTextureRenderTarget2D> RT = nullptr;				/// The render target
+	TexDescriptor					Desc;							/// The descriptor for this Texture
+	TextureFilter					Filter = TextureFilter::TF_Bilinear;/// The filtering of the Texture/RT
+	RenderMaterial_FXPtr			CopyMat;						/// Copy material
+	CHashPtr						HashObj;						/// The hash for this Texture. This is totally optional
+
+	void							InitTexture(const uint8* Pixels, size_t Length);
+	UTexture2D*						CreateTexture(uint32 Width, uint32 Height, EPixelFormat PixelFormat, bool sRGB, UObject* Package);
+	UTexture2D*						InitTextureDefault(int32 Width, int32 Height, EPixelFormat PixelFormat, const uint8* uncompressedData, size_t uncompressedDataSize, UObject* Package);
+
+	bool							CopyImageToBuffer(EPixelFormat& PixelFormat, int32& Width, int32& Height, TArray<uint8>& inputBuffer, TArray<uint8>& OutputBuffer);
+
+	UTexture2D*						InitTextureHDR(const TArray<uint8>& Buffer, UPackage* Package);
+
+	
+	void							InvalidateCached();
+	RenderMaterialPtr				GetDefaultMaterial();
+	void							UpdateRaw(const uint8* SrcPixels, size_t Length);
+	bool							IsValidVirtualTexture();
+public:
+	static EObjectFlags				Flags;
+	
+									Tex(int32 Width, int32 Height, EPixelFormat PixelFormat);
+									explicit Tex(const TexDescriptor& InDesc);
+									explicit Tex(UTexture2D* Texture);
+									explicit Tex(RawBufferPtr RawObj);
+									explicit Tex(UTextureRenderTarget2D* RT);
+									Tex();
+	virtual							~Tex() override;
+
+	//////////////////////////////////////////////////////////////////////////
+	/// Others
+	//////////////////////////////////////////////////////////////////////////
+	static void						FreeGenericTexture(UTexture** Texture);
+	static void						FreeTexture(UTexture2D** Texture);
+	static void						FreeRT(UTextureRenderTarget2D** RT);
+	
+	void							InitRT(bool ForceFloat = false);
+
+	void							Release();
+	void							ReleaseRT();
+	void							ReleaseTexture();
+	void							TransferTextureToRT(FRHICommandListImmediate& RHI, UTexture2D** PrevTexture, bool FreeAfterUse);
+	void							TransferVirtualTextureToRT(FRHICommandListImmediate& RHI, UTexture2D** TextureToTransfer, bool FreeAfterUse);
+	void							LoadAsset(FSoftObjectPath& SoftPath, const DesiredImageProperties* Props = nullptr);
+	AsyncActionResultPtr			LoadAsync(const FString& Filename, const DesiredImageProperties* Props = nullptr);
+	AsyncActionResultPtr			LoadFlat();
+	AsyncTiledBlobRef				ToBlob(int32 XTiles, int32 YTiles, uint32 Width = 0, uint32 Height = 0, bool TransferToRT = false); ///There is a visual difference between having material transfer to RT or directly copy from Texture
+	AsyncTiledBlobRef				ToSingleBlob(CHashPtr Hash, bool TransferToRT = false, bool ResolveOnRenderThread = false, bool NoCache = false);
+	AsyncActionResultPtr			LoadRaw(RawBufferPtr RawObj);
+	void							UpdateRaw(RawBufferPtr RawObj);
+
+	void							SetFilter(TextureFilter FilterValue);
+	
+	virtual void					Bind(FName Name, std::shared_ptr<RenderMaterial> Material) const;
+	void							Clear();
+	void							Clear(FRHICommandList& RHI);
+	void							Clear(FRHICommandList& RHI, FLinearColor Color);
+
+	void							Free();
+	size_t							GetMemSize() const;
+	RawBufferPtr					Raw(const BufferDescriptor* SrcDesc = nullptr) const; // Create a RawBuffer from the Texture, generate the descriptor if none is proveided
+	void							GenerateMips();
+
+	virtual bool					IsArray()  { return false; }
+	virtual FRHITexture2D*			GetRHITexture() const;
+	virtual UTexture*				GetTexture() const;
+
+	virtual bool					IsNull() const;
+
+	TArray<FColor>					ReadPixels();
+
+	//////////////////////////////////////////////////////////////////////////
+	/// Static functions
+	//////////////////////////////////////////////////////////////////////////
+	/**
+	 * Saves render target to given path and filename (path, filename).
+	 *
+	 * @param	RenderTarget Input render target to save
+	 * @param	Path		 receives the value of the path where the file needs to be saved
+	 * @param	Filename	 filename with extension to save in particular format
+	 * @param	bIsHDR		whether to save RTF_16 in HDR or EXR (discarded when trying to save PF_B8G8R8A8)	 
+	 */
+	static	bool					SaveImage(UTextureRenderTarget2D* RT, const FString& Path, const FString& Filename, bool bIsHDR = true);
+
+	//////////////////////////////////////////////////////////////////////////
+	/// Inline functions
+	//////////////////////////////////////////////////////////////////////////
+
+	FORCEINLINE FLinearColor		ClearColor() const { return Desc.ClearColor; }
+	FORCEINLINE FLinearColor&		ClearColor() { return Desc.ClearColor; }
+	FORCEINLINE FIntPoint			GetSize() const { return FIntPoint(Desc.Width, Desc.Height); }
+	FORCEINLINE operator			FRHITexture2D*() const { return GetRHITexture(); }
+	//FORCEINLINE operator			FTextureRHIRef() const { return RHITextureRef(); }
+	FORCEINLINE operator			UTexture*() const { return GetTexture(); }
+	FORCEINLINE FString				GetName() const { return GetTexture()->GetName(); }
+	FORCEINLINE const TexDescriptor& GetDescriptor() const { return Desc; }
+	FORCEINLINE TextureFilter		GetFilter() const { return Filter; }
+	FORCEINLINE uint32				GetWidth() const { return Desc.Width; }
+	FORCEINLINE uint32				GetHeight() const { return Desc.Height; }
+	FORCEINLINE bool				IsRenderTarget() const { return !Texture && RT; }
+	FORCEINLINE UTexture2D*			GetTexture2D() const { return Texture; }
+	FORCEINLINE UTextureRenderTarget2D* GetRenderTarget() const { return RT; }
+	FORCEINLINE void				SetHash(CHashPtr Hash) { HashObj = Hash; }
+	FORCEINLINE CHashPtr			Hash() const { return HashObj; }
+};
+
+//typedef std::shared_ptr<Tex>		TexPtr;
+
diff --git a/Engine/Plugins/Experimental/TextureGraph/Source/TextureGraphEngine/2D/TexArray.h b/Engine/Plugins/Experimental/TextureGraph/Source/TextureGraphEngine/2D/TexArray.h
new file mode 100644
--- /dev/null
+++ b/Engine/Plugins/Experimental/TextureGraph/Source/TextureGraphEngine/2D/TexArray.h
@@ -0,0 +1,42 @@
+// Copyright Epic Games, Inc. All Rights Reserved.
+#pragma once
+
+#include "CoreMinimal.h"
+#include "Tex.h"
+#include "Helper/Promise.h"
+
+class UTexture2D;
+class UTextureRenderTarget2DArray;
+
+class TexArray;
+
+typedef std::shared_ptr<TexArray>		TexArrayPtr;
+
+class TEXTUREGRAPHENGINE_API TexArray : public Tex
+{
+	friend class Device_FX;
+
+protected:
+	TObjectPtr<UTextureRenderTarget2DArray>	RTArray = nullptr;			/// The render target array
+	int32							XTiles = 0;
+	int32							YTiles = 0;
+	int32							Index = 0;
+	void							InitRTArray(bool ForceFloat = false);
+
+public:
+									TexArray();
+									TexArray(Tex& TexObj, int32 XTilesTotal, int32 YTilesTotal);
+									TexArray(const TexDescriptor& Desc, UTextureRenderTarget2DArray* rtArray);
+									TexArray(const TexDescriptor& Desc,  int32 XTilesTotal, int32 YTilesTotal);
+	virtual							~TexArray() override;
+	virtual void					AddReferencedObjects(FReferenceCollector& Collector) override;
+
+	virtual bool					IsArray() override { return true; }
+	virtual UTexture*				GetTexture() const override;
+	virtual bool					IsNull() const override;
+
+	//////////////////////////////////////////////////////////////////////////
+	/// Inline functions
+	//////////////////////////////////////////////////////////////////////////
+	FORCEINLINE UTextureRenderTarget2DArray* GetRenderTargetArray() const { return RTArray; }
+};
diff --git a/Engine/Plugins/Experimental/TextureGraph/Source/TextureGraphEngine/2D/TextureHelper.h b/Engine/Plugins/Experimental/TextureGraph/Source/TextureGraphEngine/2D/TextureHelper.h
new file mode 100644
--- /dev/null
+++ b/Engine/Plugins/Experimental/TextureGraph/Source/TextureGraphEngine/2D/TextureHelper.h
@@ -0,0 +1,148 @@
+// Copyright Epic Games, Inc. All Rights Reserved.
+#pragma once
+
+#include "2D/Mask/MaskEnums.h"
+#include "CoreMinimal.h"
+#include "Data/RawBuffer.h"
+#include "Engine/TextureRenderTarget2D.h"
+#include "Helper/Promise.h"
+#include "Kismet/KismetRenderingLibrary.h"
+#include "PixelFormat.h" 
+#include "TextureContent.h"
+#include "TextureType.h"
+
+class Tex;
+class UTextureRenderTarget2D;
+class UTexture2D;
+struct BufferDescriptor;
+enum class ETG_TextureFormat : uint8;
+enum class ETSBufferFormat;
+class TiledBlob;
+typedef std::shared_ptr<TiledBlob>		TiledBlobPtr;
+
+typedef cti::continuable<bool>			AsyncBool;
+typedef cti::continuable<FLinearColor>	AsyncLinearColor;
+
+struct TEXTUREGRAPHENGINE_API TextureHelper
+{
+	static TiledBlobPtr					GBlack;					/// Black RGBA texture
+	static TiledBlobPtr					GWhite;					/// White RGBA texture
+	static TiledBlobPtr					GGray;					/// Gray RGBA texture
+	static TiledBlobPtr					GRed;					/// Red RGBA texture
+	static TiledBlobPtr					GGreen;					/// Green RGBA texture
+	static TiledBlobPtr					GBlue;					/// Blue RGBA texture
+	static TiledBlobPtr					GYellow;				/// Yellow RGBA texture
+	static TiledBlobPtr					GMagenta;				/// Magenta RGBA texture
+	static TiledBlobPtr					GDefaultNormal;			/// Default Normal RGBA texture
+	static TiledBlobPtr					GWhiteMask;				/// Default white mask
+	static TiledBlobPtr					GBlackMask;				/// Default black mask
+
+	static cti::continuable<bool>		InitSolidTexture(TiledBlobPtr* BlobObj, FLinearColor Color, FString Name, struct TexDescriptor Desc);
+	static void							InitStockTextures();
+	static void							FreeStockTextures();
+
+	static FORCEINLINE TiledBlobPtr		GetBlack() { check(GBlack); return GBlack; }
+	static FORCEINLINE TiledBlobPtr		GetWhite() { check(GWhite); return GWhite; }
+	static FORCEINLINE TiledBlobPtr		GetGray() { check(GGray);  return GGray; }
+	static FORCEINLINE TiledBlobPtr		GetRed() { check(GRed);  return GRed; }
+	static FORCEINLINE TiledBlobPtr		GetGreen() { check(GGreen);  return GGreen; }
+	static FORCEINLINE TiledBlobPtr		GetBlue() { check(GBlue);  return GBlue; }
+	static FORCEINLINE TiledBlobPtr		GetYellow() { check(GYellow);  return GYellow; }
+	static FORCEINLINE TiledBlobPtr		GetMagenta() { check(GMagenta);  return GMagenta; }
+	static FORCEINLINE TiledBlobPtr		GetErrorBlob() { return GetMagenta(); }
+	static FORCEINLINE TiledBlobPtr		GetDefaultNormal() { check(GDefaultNormal); return GDefaultNormal; }
+	static FORCEINLINE TiledBlobPtr		GetWhiteMask() { check(GWhiteMask); return GWhiteMask; }
+	static FORCEINLINE TiledBlobPtr		GetBlackMask() { check(GBlackMask); return GBlackMask; }
+
+	//////////////////////////////////////////////////////////////////////////
+	static uint32						GetChannelsFromPixelFormat(EPixelFormat InPixelFormat);
+	static uint32						GetBppFromPixelFormat(EPixelFormat InPixelFormat);
+	static ETextureRenderTargetFormat	GetRenderTargetFormatFromPixelFormat(EPixelFormat InPixelFormat);
+	static ETextureSourceFormat			GetSourceFormat(ETG_TextureFormat TGTextureFormat);
+	static bool							GetBufferFormatAndChannelsFromTGTextureFormat(ETG_TextureFormat Format, BufferFormat& OutBufferFormat, uint32& OutBufferChannels);
+	static EPixelFormat					GetPixelFormatFromRenderTargetFormat(ETextureRenderTargetFormat RTFormat);
+	static bool							GetPixelFormatFromTextureSourceFormat(ETextureSourceFormat SourceFormat, EPixelFormat& OutPixelFormat, uint32& OutNumChannels);
+	static ETextureSourceFormat			GetTextureSourceFormat(BufferFormat Format, uint32 ItemsPerPoint);
+	static ETG_TextureFormat			GetTGTextureFormatFromChannelsAndFormat(uint32 ItemsPerPoint, BufferFormat Format);
+	static uint32						GetNumChannelsFromTGTextureFormat(ETG_TextureFormat TextureFormat);
+	static FString						GetChannelsTextFromItemsPerPoint(const int32 InItemsPerPoint);	
+	static void							ClearRT(UTextureRenderTarget2D* RenderTarget, FLinearColor Color = FLinearColor::Transparent);
+	static void							ClearRT(FRHICommandList& RHI, UTextureRenderTarget2D* RenderTarget, FLinearColor Color = FLinearColor::Transparent);
+
+	static const char*					TextureTypeToString(TextureType Type);
+	static TextureType					TextureStringToType(const FString& TypeString);
+	static const char*					TextureTypeToMegascansType(TextureType Type);
+
+	static TextureType					TextureContentToTextureType(TextureContent Content);
+	static TextureContent				TextureTypeToTextureContent(TextureType Type);
+	static MaskType						TextureContentToMaskType(TextureContent Content);
+	static TextureContent				MaskTypeToTextureContent(MaskType Type);
+	static MaskModifierType				TextureContentToMaskModifier(TextureContent Content);
+	static TextureContent				MaskModifierToTextureContent(MaskModifierType Type);
+
+	static RawBufferPtr					RawFromRT(UTextureRenderTarget2D* RenderTarget, const BufferDescriptor& Desc);
+	static RawBufferPtr					RawFromTexture(UTexture2D* Texture, const BufferDescriptor& Desc);
+	static RawBufferPtr					RawFromResource(const FTexture2DRHIRef& ResourceRHI, const BufferDescriptor& Desc);
+	static BufferFormat					FindOptimalSupportedFormat(BufferFormat SrcFormat);
+
+	static void							RawFromRT_Tiled(UTextureRenderTarget2D* RenderTarget, const BufferDescriptor& Desc, size_t TileSizeX, size_t TileSizeY, RawBufferPtrTiles& Tiles);
+	static void							RawFromTexture_Tiled(UTexture2D* Texture, const BufferDescriptor& Desc, size_t TileSizeX, size_t TileSizeY, RawBufferPtrTiles& Tiles);
+	static void							RawFromResource_Tiled(FTexture2DRHIRef ResourceRHI, const BufferDescriptor& Desc, size_t TileSizeX, size_t TileSizeY, RawBufferPtrTiles& Tiles);
+	static void							RawFromMem_Tiled(const uint8* SrcData, size_t SrcDataLength, const BufferDescriptor& SrcDesc, size_t TileSizeX, size_t TileSizeY, RawBufferPtrTiles& Tiles);
+	static RawBufferPtr					CombineRaw_Tiles(const RawBufferPtrTiles& Tiles, CHashPtr HashValue = nullptr, bool bIsTransient = false);
+
+	static bool							IsFloatRT(UTextureRenderTarget2D* RenderTarget);
+	static FLinearColor					GetPixelValueFromRaw(RawBufferPtr RawObj, int32 Width, int32 Height, int32 X, int32 Y);
+	static AsyncBool					ExportRaw(RawBufferPtr RawObj, const FString& CompletePath);
+	static bool							CanSupportTexture(UTexture* Tex);
+	static bool							CanSplitToTiles(UTexture* Texture, int TilesX, int TilesY);
+	static size_t						RoundUpTo(size_t Size, size_t DesiredRounding);
+
+	//////////////////////////////////////////////////////////////////////////
+	/// Inline functions
+	//////////////////////////////////////////////////////////////////////////
+	static FORCEINLINE FString			CreateTileName(const FString& BaseName, int32 TileX, int32 TileY)
+	{
+		return FString::Printf(TEXT("%s-%d,%d"), *BaseName, TileX, TileY);
+	}
+
+	template <class T>
+	static FORCEINLINE bool				IsPowerOf2(T Value)
+	{
+		return Value > 0 && (Value & (Value - 1)) == 0;
+	}
+
+	template <class T>
+	static FORCEINLINE T				FloorToPowerOf2(T Value)
+	{
+		Value = Value | (Value >> 1);
+		Value = Value | (Value >> 2);
+		Value = Value | (Value >> 4);
+		Value = Value | (Value >> 8);
+		Value = Value | (Value >> 16);
+		return Value - (Value >> 1);;
+	}
+
+	template <class T>
+	static FORCEINLINE T				CeilToPowerOf2(T Value)
+	{
+		Value--;
+
+		Value |= Value >> 1;
+		Value |= Value >> 2;
+		Value |= Value >> 4;
+		Value |= Value >> 8;
+		Value |= Value >> 16;
+		Value++;
+
+		return Value;
+	}
+
+
+	// Calculate the total number of mip levels for a given base size.
+	// This is including the level 0 of said base size 
+	static FORCEINLINE uint32			CalcNumMipLevels(uint32 BaseSize)
+	{
+		return (uint32)(ceil(log2(double(BaseSize))) + 1.0);
+	}
+};
diff --git a/Engine/Plugins/Experimental/TextureGraph/Source/TextureGraphEngine/3D/MeshDetails/MeshDetails.h b/Engine/Plugins/Experimental/TextureGraph/Source/TextureGraphEngine/3D/MeshDetails/MeshDetails.h
new file mode 100644
--- /dev/null
+++ b/Engine/Plugins/Experimental/TextureGraph/Source/TextureGraphEngine/3D/MeshDetails/MeshDetails.h
@@ -0,0 +1,38 @@
+// Copyright Epic Games, Inc. All Rights Reserved.
+#pragma once
+
+#include "CoreMinimal.h" 
+#include <memory>
+#include "Helper/Promise.h"
+#include "Helper/DataUtil.h"
+#include <array>
+
+class MeshInfo;
+class MeshDetails;
+typedef cti::continuable<MeshDetails*>	MeshDetailsPAsync;
+
+class RawBuffer;
+typedef std::shared_ptr<RawBuffer>		RawBufferPtr;
+
+//////////////////////////////////////////////////////////////////////////
+class TEXTUREGRAPHENGINE_API MeshDetails
+{
+protected:
+	MeshInfo*					Mesh;					/// The mesh that is containing this detail
+	bool						bIsFinalised = false;		/// Whether this detail has been finalised or not
+	CHashPtr					HashValue;					/// What is the hash of this detail
+
+	virtual void				CalculateTri(size_t ti);
+	virtual void				CalculateVertex(size_t vi);
+
+public:
+								MeshDetails(MeshInfo* mesh);
+	virtual						~MeshDetails();
+
+	virtual MeshDetailsPAsync	Calculate();
+	virtual MeshDetailsPAsync	Finalise();
+	virtual void				RenderDebug();
+	virtual void				Release();
+};
+
+typedef std::shared_ptr<MeshDetails> MeshDetailsPtr;
diff --git a/Engine/Plugins/Experimental/TextureGraph/Source/TextureGraphEngine/3D/MeshDetails/MeshDetails_SpatialUV.h b/Engine/Plugins/Experimental/TextureGraph/Source/TextureGraphEngine/3D/MeshDetails/MeshDetails_SpatialUV.h
new file mode 100644
--- /dev/null
+++ b/Engine/Plugins/Experimental/TextureGraph/Source/TextureGraphEngine/3D/MeshDetails/MeshDetails_SpatialUV.h
@@ -0,0 +1,20 @@
+// Copyright Epic Games, Inc. All Rights Reserved.
+#pragma once
+
+#include "MeshDetails.h"
+#include "GraphicsDefs.h"
+#include <vector>
+#include <list>
+
+class TEXTUREGRAPHENGINE_API MeshDetails_SpatialUV : public MeshDetails
+{
+protected:
+	virtual void				CalculateTri(size_t ti) override;
+
+public:
+								MeshDetails_SpatialUV(MeshInfo* mesh);
+	virtual						~MeshDetails_SpatialUV();
+
+	virtual MeshDetailsPAsync	Calculate() override;
+	virtual void				Release() override;
+};
diff --git a/Engine/Plugins/Experimental/TextureGraph/Source/TextureGraphEngine/3D/MeshDetails/MeshDetails_Tri.h b/Engine/Plugins/Experimental/TextureGraph/Source/TextureGraphEngine/3D/MeshDetails/MeshDetails_Tri.h
new file mode 100644
--- /dev/null
+++ b/Engine/Plugins/Experimental/TextureGraph/Source/TextureGraphEngine/3D/MeshDetails/MeshDetails_Tri.h
@@ -0,0 +1,54 @@
+// Copyright Epic Games, Inc. All Rights Reserved.
+#pragma once
+
+#include "MeshDetails.h"
+#include "GraphicsDefs.h"
+
+class TEXTUREGRAPHENGINE_API MeshDetails_Tri : public MeshDetails
+{
+public:
+	struct Triangle
+	{
+		FBox					bounds;
+		FBox					uvBounds;
+		Vector3					normal;
+		Vector3					tangent;
+		Vector3					biTangent;
+		float					distance;
+		Vector3					centre;
+
+        int32					indices[3];
+
+								Triangle(int32 i0, int32 i1, int32 i2);
+		bool					HasIndex(int32 i);
+		std::array<int32, 2>	GetOtherVertices(int32 vi) const;
+	};
+
+protected:
+	static const size_t			s_maxBatch;
+	Triangle*					_triangles = nullptr;			/// The triangles that we have
+	size_t						_numTriangles = 0;				/// How many triangles do we have
+	RawBufferPtr				_raw;							/// The triangle buffer is indexed into this RawBuffer
+
+	virtual void				CalculateTri(size_t ti) override;
+
+public:
+								MeshDetails_Tri(MeshInfo* mesh);
+	virtual						~MeshDetails_Tri();
+
+	virtual MeshDetailsPAsync	Calculate() override;
+
+	int32*						GetIndices(size_t ti) const;
+	std::array<Vector3, 3>		GetVertices(size_t ti) const;
+	std::array<Vector2, 3>		GetUVs(size_t ti) const;
+	std::array<Vector3, 3>		GetNormals(size_t ti) const;
+	std::array<Tangent, 3>		GetTangents(size_t ti) const;
+
+	//////////////////////////////////////////////////////////////////////////
+	/// Inline functions
+	//////////////////////////////////////////////////////////////////////////
+	FORCEINLINE Triangle*		Triangles() const { return _triangles; }
+	FORCEINLINE size_t			NumTriangles() const { return _numTriangles; }
+};
+
+
diff --git a/Engine/Plugins/Experimental/TextureGraph/Source/TextureGraphEngine/3D/ProceduralMeshActor.h b/Engine/Plugins/Experimental/TextureGraph/Source/TextureGraphEngine/3D/ProceduralMeshActor.h
new file mode 100644
--- /dev/null
+++ b/Engine/Plugins/Experimental/TextureGraph/Source/TextureGraphEngine/3D/ProceduralMeshActor.h
@@ -0,0 +1,77 @@
+// Copyright Epic Games, Inc. All Rights Reserved.
+#pragma once
+
+#include "CoreMinimal.h"
+#include "GameFramework/Actor.h"
+//Set material
+#include "Materials/MaterialInterface.h" 
+#include "Materials/MaterialInstanceDynamic.h" 
+//Load texture from file
+#include "Kismet/KismetRenderingLibrary.h" 
+//FCoreMesh
+#include "CoreMesh.h"
+#include "ProceduralMeshComponent.h"
+//Debug
+#include "DrawDebugHelpers.h" 
+//For World normals
+#include "Data/TiledBlob.h"
+#include "Device/FX/DeviceBuffer_FX.h"
+#include "ProceduralMeshActor.generated.h"
+
+class USceneCaptureComponent2D;
+
+class Tex;
+typedef std::shared_ptr<Tex>		TexPtr;
+UCLASS()
+class TEXTUREGRAPHENGINE_API AProceduralMeshActor : public AActor
+{
+	GENERATED_BODY()
+private:
+	TArray<TArray<FVector>>				_debugPoints;						
+	bool								_isDebug = false;
+	FString								_meshName = FString("");
+
+	void								BlitInternal(UMaterialInterface* mat);
+
+protected:
+	// Called when the game starts or when spawned
+	virtual void						BeginPlay() override;
+	// Called every frame
+	virtual void						Tick(float DeltaTime) override;
+
+	
+	UPROPERTY(VisibleAnywhere, BlueprintReadWrite, Category = "Hierarchy")
+		TObjectPtr<UProceduralMeshComponent>		_meshObj;
+
+	UPROPERTY(VisibleAnywhere, BlueprintReadWrite, Category = "Hierarchy")
+		TObjectPtr<USceneCaptureComponent2D>		_sceneCaptureComp;
+public:	
+	// Sets default values for this actor's properties
+										AProceduralMeshActor();
+										~AProceduralMeshActor();
+
+										void InitSceneCaptureComponent();
+
+										void InitMeshObj();
+
+	void								SetMeshData(CoreMeshPtr mesh);
+	void 								DrawDebugLines(TArray<FVector>& positionArray, FVector length);
+	void 								SetMaterial(UMaterialInterface* mat);
+	UMaterialInterface*  				GetMaterial();
+	
+	void								BlitTo(UTextureRenderTarget2D* rt, UMaterialInterface* mat);
+	virtual void						BeginDestroy();
+
+	UFUNCTION(BlueprintCallable, Category = "Debug")
+		void							ToggleDebug(int32 debugType);
+
+	UFUNCTION(BlueprintCallable,Category = "Info")
+		const FString &					GetMeshName() { return _meshName; }
+
+	UFUNCTION(BlueprintCallable, Category = "Info")
+		void							SetMeshName(const FString& name) { _meshName = name; }
+
+	//////////////////////////////////////////////////////////////////////////
+	/// Inline functions
+	//////////////////////////////////////////////////////////////////////////
+};
diff --git a/Engine/Plugins/Experimental/TextureGraph/Source/TextureGraphEngine/3D/RenderMesh_Editor.cpp b/Engine/Plugins/Experimental/TextureGraph/Source/TextureGraphEngine/3D/RenderMesh_Editor.cpp
new file mode 100644
--- /dev/null
+++ b/Engine/Plugins/Experimental/TextureGraph/Source/TextureGraphEngine/3D/RenderMesh_Editor.cpp
@@ -0,0 +1,138 @@
+// Copyright Epic Games, Inc. All Rights Reserved.
+#if WITH_EDITOR
+#include "RenderMesh_Editor.h"
+#include "CoreMesh.h"
+#include "MeshInfo.h"
+#include "KismetProceduralMeshLibrary.h"
+#include "../Helper/MathUtils.h"
+#include <Engine/StaticMesh.h>
+#include <StaticMeshResources.h>
+
+RenderMesh_Editor::RenderMesh_Editor()
+{
+	_isPlane = false;
+}
+
+RenderMesh_Editor::RenderMesh_Editor(RenderMesh* parent, TArray<MeshInfoPtr> meshes, MaterialInfoPtr matInfo) : RenderMesh(parent, meshes, matInfo)
+{
+
+}
+
+RenderMesh_Editor::RenderMesh_Editor(UStaticMeshComponent* staticMeshComponent,UWorld* world)
+{
+	_world = world;
+	_meshComponents.Add(staticMeshComponent);
+}
+
+void RenderMesh_Editor::PrepareForRendering(UWorld* world, FVector scale)
+{
+	if (_parentMesh)
+		_parentMesh->SetViewScale(scale);
+
+	_viewScale = scale;
+}
+
+AsyncActionResultPtr RenderMesh_Editor::Load()
+{
+	// Need to figure out how we are going to support MTS and UDIM meshes.
+	_meshSplitType = MeshSplitType::Single;
+
+	//_meshActor->GetComponents<UStaticMeshComponent>(_meshComponents);
+	if (_meshComponents.Num() <= 0)
+		return cti::make_exceptional_continuable<ActionResultPtr>(std::make_exception_ptr(std::runtime_error("No mesh components found on actor")));
+
+	return cti::make_continuable<ActionResultPtr>([this](auto&& promise)
+		{
+			AsyncTask(ENamedThreads::AnyBackgroundThreadNormalTask, [this, promise = std::forward<decltype(promise)>(promise)]() mutable
+			{
+				try
+				{
+					LoadInternal();
+					promise.set_value(std::make_shared<ActionResult>());
+				}
+				catch (cti::exception_t ex)
+				{
+					promise.set_exception(ex);
+				}
+			});
+		});
+}
+
+void RenderMesh_Editor::SetMaterial(UMaterialInterface* material)
+{
+	for (int meshIndex = 0; meshIndex < _meshComponents.Num(); meshIndex++)
+	{
+		UStaticMeshComponent* meshComponent = _meshComponents[meshIndex];
+
+		MeshInfoPtr meshInfo = _meshes[meshIndex];
+		meshComponent->SetMaterial(meshInfo->GetMaterialIndex(), material);
+	}
+}
+
+void RenderMesh_Editor::LoadInternal()
+{
+	for (UStaticMeshComponent* meshComponent : _meshComponents)
+	{
+		LoadSingleMeshComponent(*meshComponent);
+	}
+}
+void RenderMesh_Editor::LoadSingleMeshComponent(const UStaticMeshComponent& meshComponent)
+{
+	UStaticMesh* staticMesh = meshComponent.GetStaticMesh();
+
+	//Set bAllowCPUAccess to true as it is required to use  GetSectionFromStaticMesh
+	staticMesh->bAllowCPUAccess = true;
+
+	const int32 LODindex = 0;
+	int32 numSections = staticMesh->GetNumSections(LODindex);
+
+	// Return if RenderData is invalid
+	if (!staticMesh->GetRenderData())
+		return;
+
+	// No valid mesh data on lod 0 (shouldn't happen)
+	if (!staticMesh->GetRenderData()->LODResources.IsValidIndex(LODindex))
+		return;
+
+	// load materials
+	//const TArray<UMaterialInterface*> materialInterfaces = meshComponent.GetMaterials();
+	for (int32 materialIndex = 0; materialIndex < meshComponent.GetNumMaterials(); ++materialIndex)
+	{
+		UMaterialInterface* materialInterface = meshComponent.GetMaterial(materialIndex);
+		if (materialInterface)
+		{
+			FString matName;
+			materialInterface->GetName(matName);
+			AddMaterialInfo(materialIndex, matName);
+		}
+	}
+	_currentMaterials.Empty();
+	_currentMaterials.Append(_originalMaterials);
+
+	TArray<FVector> Vertices;
+	TArray<int32> Triangles;
+	TArray<FVector> Normals;
+	TArray<FVector2D> UVs;
+	TArray<FProcMeshTangent> Tangents;
+
+	_originalBounds.Init();
+	
+	for (int32 sectionIndex = 0; sectionIndex < numSections; sectionIndex++)
+	{
+		CoreMesh* cmesh = new CoreMesh();
+		CoreMeshPtr cmeshPtr = CoreMeshPtr(cmesh);
+		MeshInfoPtr meshInfo = std::make_shared<MeshInfo>(cmeshPtr);
+
+		cmesh->bounds = staticMesh->GetBoundingBox();
+		meshComponent.GetName(cmesh->name);
+
+		FMeshSectionInfo sectionInfo = staticMesh->GetSectionInfoMap().Get(LODindex, sectionIndex);
+		UKismetProceduralMeshLibrary::GetSectionFromStaticMesh(staticMesh, LODindex, sectionIndex, cmesh->vertices, cmesh->triangles, cmesh->normals, cmesh->uvs, cmesh->tangents);
+		cmesh->materialIndex = sectionInfo.MaterialIndex;
+		
+		MathUtils::EncapsulateBound(_originalBounds, cmeshPtr->bounds);	
+
+		_meshes.Add(meshInfo);
+	}
+}
+#endif // WITH_EDITOR
\ No newline at end of file
diff --git a/Engine/Plugins/Experimental/TextureGraph/Source/TextureGraphEngine/3D/RenderMesh_Editor.h b/Engine/Plugins/Experimental/TextureGraph/Source/TextureGraphEngine/3D/RenderMesh_Editor.h
new file mode 100644
--- /dev/null
+++ b/Engine/Plugins/Experimental/TextureGraph/Source/TextureGraphEngine/3D/RenderMesh_Editor.h
@@ -0,0 +1,32 @@
+// Copyright Epic Games, Inc. All Rights Reserved.
+#pragma once
+#if WITH_EDITOR
+#include "CoreMinimal.h"
+#include "RenderMesh.h"
+#include "Components/StaticMeshComponent.h"
+
+struct CoreMesh;
+
+class TEXTUREGRAPHENGINE_API RenderMesh_Editor: public RenderMesh
+{
+
+protected:
+	TArray<UStaticMeshComponent*>			_meshComponents;
+	UWorld*									_world = nullptr;
+	virtual void							LoadInternal() override;
+	void									LoadSingleMeshComponent(const UStaticMeshComponent& mesh);
+
+public:
+											RenderMesh_Editor();
+											RenderMesh_Editor(RenderMesh* parent, TArray<MeshInfoPtr> meshes, MaterialInfoPtr matInfo);
+											RenderMesh_Editor(UStaticMeshComponent* staticMeshComponent, UWorld* world);
+
+	FORCEINLINE								TArray<UStaticMeshComponent*> GetMeshComponents() { return _meshComponents; }
+	virtual AsyncActionResultPtr			Load() override;
+	virtual FMatrix							LocalToWorldMatrix() const override { return _meshComponents[0]->GetComponentTransform().ToMatrixWithScale(); }
+	virtual void							PrepareForRendering(UWorld* world, FVector scale) override;
+	virtual void							SetMaterial(UMaterialInterface* material) override;
+};
+
+typedef std::shared_ptr<RenderMesh_Editor> RenderMesh_EditorPtr;
+#endif
\ No newline at end of file
diff --git a/Engine/Plugins/Experimental/TextureGraph/Source/TextureGraphEngine/3D/RenderMesh_EditorScene.h b/Engine/Plugins/Experimental/TextureGraph/Source/TextureGraphEngine/3D/RenderMesh_EditorScene.h
new file mode 100644
--- /dev/null
+++ b/Engine/Plugins/Experimental/TextureGraph/Source/TextureGraphEngine/3D/RenderMesh_EditorScene.h
@@ -0,0 +1,26 @@
+// Copyright Epic Games, Inc. All Rights Reserved.
+#pragma once
+#if WITH_EDITOR
+#include "CoreMinimal.h"
+#include "RenderMesh_Editor.h"
+
+struct CoreMesh;
+
+class TEXTUREGRAPHENGINE_API RenderMesh_EditorScene: public RenderMesh_Editor
+{
+private:
+	AActor*									_meshActor;				/// Actor that this rendermesh is associated to
+
+public:
+											RenderMesh_EditorScene();
+											RenderMesh_EditorScene(AActor* actor);
+											RenderMesh_EditorScene(RenderMesh* parent, TArray<MeshInfoPtr> meshes, MaterialInfoPtr matInfo);
+
+	AsyncActionResultPtr					Load();
+	void									PrepareForRendering(UWorld* world, FVector scale) override;
+	void									RemoveActors() override;
+
+};
+
+typedef std::shared_ptr<RenderMesh_EditorScene> RenderMesh_EditorScenePtr;
+#endif
\ No newline at end of file
diff --git a/Engine/Plugins/Experimental/TextureGraph/Source/TextureGraphEngine/3D/RenderMesh_Procedural.h b/Engine/Plugins/Experimental/TextureGraph/Source/TextureGraphEngine/3D/RenderMesh_Procedural.h
new file mode 100644
--- /dev/null
+++ b/Engine/Plugins/Experimental/TextureGraph/Source/TextureGraphEngine/3D/RenderMesh_Procedural.h
@@ -0,0 +1,32 @@
+// Copyright Epic Games, Inc. All Rights Reserved.
+#pragma once
+#include "CoreMinimal.h"
+#include "3D/RenderMesh.h"
+#include "Helper/Promise.h"
+
+struct CoreMesh;
+typedef std::shared_ptr<CoreMesh>	CoreMeshPtr;
+
+
+class RenderMesh;
+
+class TEXTUREGRAPHENGINE_API RenderMesh_Procedural: public RenderMesh
+{
+private:
+	int										_tesselation = 32;
+	FVector2D								_dimension= FVector2D::UnitVector;
+	FVector									_offSet;
+	
+public:
+	explicit								RenderMesh_Procedural(const MeshLoadInfo loadInfo);
+	virtual AsyncActionResultPtr			Load() override;
+	virtual void							LoadInternal() override;
+	void									GenerateProcedural(int tesselation, FVector2D dimension, CoreMeshPtr cmesh);
+
+	//////////////////////////////////////////////////////////////////////////
+	////Inline Functions
+	//////////////////////////////////////////////////////////////////////////
+	FORCEINLINE	int&						Tesselation() { return _tesselation; }
+};
+
+typedef std::shared_ptr<RenderMesh_Procedural> RenderMesh_ProceduralPtr;
\ No newline at end of file
