Rename MixerEngine and other files with the Mixer branding in them.

diff --git a/Engine/Plugins/Experimental/TextureGraph/Source/TextureGraphEngine/Data/Blob.h b/Engine/Plugins/Experimental/TextureGraph/Source/TextureGraphEngine/Data/Blob.h
new file mode 100644
--- /dev/null
+++ b/Engine/Plugins/Experimental/TextureGraph/Source/TextureGraphEngine/Data/Blob.h
@@ -0,0 +1,198 @@
+// Copyright Epic Games, Inc. All Rights Reserved.
+#pragma once
+
+#include "CoreMinimal.h"
+#include "RawBuffer.h"
+#include "Device/DeviceBuffer.h"
+#include "BlobRef.h"
+#include "Helper/DataUtil.h"
+
+class Device;
+class RawBuffer;
+class BlobTransform;
+class Blobber;
+struct ResourceBindInfo;
+
+typedef TFunction<void(const Blob*)> BlobReadyCallback;
+typedef cti::continuable<const Blob*> AsyncBlobResultPtr;
+
+typedef std::shared_ptr<Blob>	BlobPtr;
+typedef std::unique_ptr<Blob>	BlobUPtr;
+typedef std::weak_ptr<Blob>		BlobPtrW;
+typedef std::vector<BlobPtr>	BlobPtrVec;
+typedef std::vector<BlobPtrW>	BlobPtrWVec;
+
+#if DEBUG_BLOB_REF_KEEPING == 1
+struct DebugBlobLock
+{
+public:
+	DebugBlobLock();
+	~DebugBlobLock();
+};
+#endif 
+
+class TEXTUREGRAPHENGINE_API Blob : public std::enable_shared_from_this<Blob>
+{
+	friend class Blobber;
+	friend class TiledBlob;
+	friend class TiledBlob_Promise;
+
+public:
+	static const char*				LODTransformName;	
+	typedef std::vector<BlobPtrW>	OwnerList;
+	typedef std::vector<std::weak_ptr<Blob>> LinkedBlobsVec;
+
+protected:
+	DeviceBufferRef					Buffer;							/// Device buffer that is attached to this blob
+	bool							bIsFinalised = false;			/// Whether the promise has been finalised or not
+	int32							ReplayCount = 0;				/// For debug purpose, Incremented when replaying a job and regenerating the content of the blob.
+	FDateTime						FinaliseTS = FDateTime::Now();	/// Time when the blob was finalised
+
+	BlobPtrWVec						LODLevels;						/// The mip chain for this blob (for textures). Can also be Mesh data with fewer triangles etc. Doesn't matter to us
+
+	BlobPtr							MinMax;							/// The blob that's used for calculating min/max
+	std::shared_ptr<float>			MinValue;						/// Extracted min value (this is not always present)
+	std::shared_ptr<float>			MaxValue;						/// Extracted min value (this is not always present)
+
+	BlobPtr							Histogram;						/// The blob that's used for storing histogram data. not always have a histogram only created when required.
+
+	BlobPtrW						LODParent;						/// The parent of this blob. Essentially the blob that was used to generate THIS lod Level
+	BlobPtrW						LODSource;						/// The original source of this blob, the mip Level 0
+
+	bool							bIsLODLevel = false;			/// Whether this blob already represents an LOD'd blob
+	LinkedBlobsVec					LinkedBlobs;					/// Blobs that are linked to this tiled blob. These are objects that are 
+																	/// uniquely created within the system but eventually resolve to the same
+																	/// tiled blob. We need to ensure that these are resolved correctly
+																	/// so that if some object is holding a pointer to a particular TiledBlob_Promise
+																	/// it gets the same view upon the completion.
+																	/// This is usually very uncommon with Blobs but more of a possibility with TiledBlob_Promise
+																	/// but we need to ensure that it works correctly for all blobs
+
+#if DEBUG_BLOB_REF_KEEPING == 1
+	OwnerList						_owners;						/// The TiledBlobs that contain this blob as a tile
+	virtual void					AddTiledOwner(BlobPtrW owner);
+#endif 
+
+									Blob();
+	virtual void					OnFinaliseInternal(BlobReadyCallback callback) const;
+	virtual void					ResetBuffer();
+	virtual void					Touch(uint64 BatchId);
+	virtual void					UpdateAccessInfo(uint64 batchId);
+	virtual void					SetHash(CHashPtr Hash);
+
+	virtual void					UpdateLinkedBlobs(bool bDoFinalise);
+	virtual void					AddLinkedBlob(BlobPtr LinkedBlob);
+
+	virtual void					FinaliseFrom(const Blob* RHS);
+
+public:
+									Blob(DeviceBufferRef InBuffer);
+									Blob(const BufferDescriptor& InDesc, CHashPtr InHash);
+									Blob(const Blob& InBlob) = delete; /// non-copyable
+
+	virtual							~Blob();
+
+	virtual AsyncBufferResultPtr	Bind(const BlobTransform* Transform, const ResourceBindInfo& BindInfo);
+	virtual AsyncBufferResultPtr	Unbind(const BlobTransform* Transform, const ResourceBindInfo& BindInfo);
+	virtual bool					IsValid() const;
+	virtual bool					IsNull() const;
+	virtual CHashPtr				Hash() const;
+	//virtual void					SetHash(CHashPtr hash);
+	virtual AsyncBufferResultPtr	Flush(const ResourceBindInfo& BindInfo);
+	virtual AsyncPrepareResult		PrepareForWrite(const ResourceBindInfo& BindInfo);
+
+#if 0 /// TODO /// Do we really need this?
+	virtual AsyncBufferResultPtr	PrepareForWrite(const ResourceBindInfo& BindInfo);
+#endif 
+
+	virtual FString					DisplayName() const;
+	virtual bool					IsFinalised() const { return bIsFinalised; }
+	virtual bool					IsTiled() const { return false; }
+	virtual bool					CanCalculateHash() const { return IsFinalised(); }
+
+	virtual void					FinaliseNow(bool bNoCalcHash, CHashPtr FixedHash);
+	virtual AsyncBufferResultPtr	Finalise(bool bNoCalcHash, CHashPtr FixedHash);
+
+	virtual AsyncRawBufferPtr		Raw();			/// SLOW: Read doc for DeviceBuffer::Raw() in DeviceBuffer.h
+	virtual AscynCHashPtr			CalcHash();		/// SLOW: Because calls Raw() underneath. Don't use directly
+
+	virtual FString&				Name() { check(Buffer); return Buffer->GetName(); }
+	virtual const FString&			Name() const { check(Buffer); return Buffer->GetName(); }
+	virtual const BufferDescriptor& GetDescriptor() const { /*check(_buffer);*/ return Buffer->Descriptor(); }
+	
+	virtual AsyncBlobResultPtr		OnFinalise() const;
+
+	virtual AsyncDeviceBufferRef	TransferTo(Device* TargetDevice);
+	virtual DeviceBufferRef			GetBufferRef() const;
+
+	//////////////////////////////////////////////////////////////////////////
+	/// Min/Max related
+	//////////////////////////////////////////////////////////////////////////
+	virtual BlobPtr					GetMinMaxBlob();
+	virtual void					SetMinMax(BlobPtr InMinMax);
+	virtual float					GetMinValue() const;
+	virtual float					GetMaxValue() const;
+	virtual bool					HasMinMax() const;
+	//////////////////////////////////////////////////////////////////////////
+	/// LOD/MipMaps related
+	//////////////////////////////////////////////////////////////////////////
+	virtual BlobPtrW				GetLODLevel(int32 Level);
+	virtual bool					HasLODLevels() const;
+	virtual bool					HasLODLevel(int32 Index) const;
+	virtual int32					NumLODLevels() const;
+	virtual void					SetLODLevel(int32 Level, BlobPtr LODBlob, BlobPtrW LODParentBlob, BlobPtrW LODSourceBlob, bool bAddToBlobber); 
+
+	//////////////////////////////////////////////////////////////////////////
+	/// Ownership related
+	//////////////////////////////////////////////////////////////////////////
+#if DEBUG_BLOB_REF_KEEPING == 1
+	virtual void					AddOwner(BlobPtrW owner);
+	virtual void					RemoveOwner(BlobPtrW owner);
+	void							RemoveOwner(Blob* owner);
+	virtual bool					HasOwner(BlobPtrW owner);
+	bool							HasOwner(Blob* owner);
+	virtual OwnerList::iterator		FindOwner(BlobPtrW owner);
+	OwnerList::iterator				FindOwner(Blob* owner);
+	virtual bool					HasBlobAsTile(Blob* blob) const { return false; };
+#endif
+
+	//////////////////////////////////////////////////////////////////////////
+	/// Inline function
+	//////////////////////////////////////////////////////////////////////////
+	FORCEINLINE uint32				GetWidth() const { return GetDescriptor().Width; }
+	FORCEINLINE uint32				GetHeight() const { return GetDescriptor().Height; }
+	//FORCEINLINE const DeviceBufferRef& Buffer() const { return _buffer; }
+
+	FORCEINLINE	DeviceType			GetDeviceType() const { return Buffer.GetDeviceType(); }
+
+	FORCEINLINE int32				GetReplayCount() const { return ReplayCount; } /// For debug purpose, report the replay number
+	FORCEINLINE const FDateTime&	FinaliseTimestamp() const { return FinaliseTS; }
+	FORCEINLINE bool				IsLateBound() const { return GetDescriptor().Format == BufferFormat::LateBound; }
+	FORCEINLINE BlobPtrW			GetLODParent() const { return LODParent; }
+	FORCEINLINE BlobPtrW			GetLODSource() const { return LODSource; }
+	FORCEINLINE bool				IsLODLevel() const { return bIsLODLevel; }
+	FORCEINLINE bool				IsTransient() const 
+	{ 
+		bool TransientBuffer = false;
+		if (Buffer && !Buffer.IsValid())
+			TransientBuffer = Buffer->IsTransient();
+		return TransientBuffer || GetDescriptor().bIsTransient; 
+	}
+
+	//Sets the histogram of the blob 
+	FORCEINLINE void				SetHistogram(BlobPtr InHistogram) { check(!Histogram); Histogram = InHistogram; }
+	//Check if blob already has a histogram
+	FORCEINLINE bool				HasHistogram() const { return Histogram != nullptr; }
+
+	FORCEINLINE BlobPtr				GetHistogram() { check(HasHistogram()); return Histogram; }
+
+	//////////////////////////////////////////////////////////////////////////
+	/// Static functions
+	//////////////////////////////////////////////////////////////////////////
+
+	static CHashPtr					CalculateMipHash(CHashPtr MainHash, int32 Level);
+	static CHashPtrVec				CalculateMipHashes(CHashPtr MainHash, CHashPtr ParentHash, int32 Level);
+};
+
+typedef T_BlobRef<Blob>				BlobRef;
+typedef T_Tiles<BlobPtr>			BlobPtrTiles;
diff --git a/Engine/Plugins/Experimental/TextureGraph/Source/TextureGraphEngine/Data/TiledBlob.h b/Engine/Plugins/Experimental/TextureGraph/Source/TextureGraphEngine/Data/TiledBlob.h
new file mode 100644
--- /dev/null
+++ b/Engine/Plugins/Experimental/TextureGraph/Source/TextureGraphEngine/Data/TiledBlob.h
@@ -0,0 +1,196 @@
+// Copyright Epic Games, Inc. All Rights Reserved.
+#pragma once
+
+#include "CoreMinimal.h"
+#include "Blob.h"
+
+class Job;
+typedef std::weak_ptr<Job>			JobPtrW;
+typedef T_BlobRef<Blob>				BlobRef;
+typedef T_Tiles<BlobPtr>			BlobPtrTiles;
+
+//////////////////////////////////////////////////////////////////////////
+/// TiledBlob: Represents a BlobObj that is made up of many Tiles. These blobs
+/// are always late bound i.e. they will not allocate any buffers 
+/// unless they absolutely need to. 
+//////////////////////////////////////////////////////////////////////////
+class TEXTUREGRAPHENGINE_API TiledBlob : public Blob
+{
+	friend class Blobber;
+
+protected:
+	mutable BlobPtrTiles			Tiles;					/// The Tiles that make up the larger BlobObj
+	mutable CHashPtr				HashValue;				/// The HashValue of tiled BlobObj
+	mutable BufferDescriptor		Desc;					/// We keep the combined descriptor separately
+	bool							bReady = false;			/// Whether the Buffer is already ready
+	bool							bTiledTarget = true;	/// Whether we're rendering to this as a tile by tile Target or not
+	JobPtrW							JobObj;					/// The job that is potentially generating this tiled BlobObj
+	BlobPtr							SingleBlob;			/// Single BlobObj pointer that we keep to prevent going out of Ref
+
+	void							CalcHash_Now() const;
+	virtual void					ResetBuffer() override;
+
+	virtual void					TouchTiles(uint64 BatchId);
+	virtual void					Touch(uint64 BatchId) override;
+	virtual void					MakeSingleBlob_Internal();
+	TiledBlob&						operator = (const TiledBlob& RHS);
+	virtual void					SetHash(CHashPtr Hash) override;
+
+	virtual void					AddLinkedBlob(BlobPtr LinkedBlob) override;
+	virtual void					FinaliseFrom(const Blob* RHS) override;
+
+public:
+	static AsyncBufferResultPtr		TileBuffer(DeviceBufferRef Buffer, BlobPtrTiles& Tiles);
+	AsyncBufferResultPtr			CombineTiles(bool bTouch, bool bIsArray, uint64 BatchId = 0);
+
+									TiledBlob(const BufferDescriptor& Desc, const BlobPtrTiles& Tiles);
+
+public:
+									TiledBlob(const BufferDescriptor& Desc, size_t NumTilesX, size_t NumTilesY, CHashPtr HashValue);
+									TiledBlob(DeviceBufferRef Buffer);
+									TiledBlob(BlobRef BlobObj);
+	virtual							~TiledBlob() override;
+
+	static std::shared_ptr<TiledBlob> InitFromTiles(const BufferDescriptor& Desc, BlobPtrTiles& Tiles);
+
+	virtual AsyncBufferResultPtr	Bind(const BlobTransform* Transform, const ResourceBindInfo& BindInfo) override;
+	virtual AsyncBufferResultPtr	Unbind(const BlobTransform* Transform, const ResourceBindInfo& BindInfo) override;
+
+	virtual bool					IsValid() const override;
+	virtual bool					IsNull() const override;
+	virtual AscynCHashPtr			CalcHash() override;	/// SLOW: Because calls Raw() underneath. Don't use directly
+
+	virtual CHashPtr				Hash() const override;
+	virtual AsyncBufferResultPtr	Flush(const ResourceBindInfo& BindInfo) override;
+	virtual AsyncPrepareResult		PrepareForWrite(const ResourceBindInfo& BindInfo) override;
+	virtual FString					DisplayName() const override;
+	virtual bool					IsTiled() const override { return true; }
+	virtual bool					IsPromise() const { return false; }
+	virtual bool					CanCalculateHash() const override;
+	virtual void					SetTile(int32 X, int32 Y, BlobRef Tile);
+	virtual void					SetTiles(const BlobPtrTiles& InTiles);
+
+	virtual FString&				Name() override { return Desc.Name; }
+	virtual const FString&			Name() const override { return Desc.Name; }
+	virtual const BufferDescriptor& GetDescriptor() const override { return Desc; }
+
+	virtual AsyncBufferResultPtr	MakeSingleBlob();
+
+	virtual AsyncDeviceBufferRef	TransferTo(Device* Target) override;
+	virtual DeviceBufferRef			GetBufferRef() const override;
+
+	virtual void					ResolveLateBound(BlobPtr Ref);
+	virtual void					ResolveLateBound(const BufferDescriptor& Desc, bool bOverrideExisting = false);
+	virtual void					CopyResolveLateBound(BlobPtr RHS);
+	virtual void					SetTransient();
+
+	//////////////////////////////////////////////////////////////////////////
+	/// Min/Max related
+	//////////////////////////////////////////////////////////////////////////
+	virtual void					SetMinMax(BlobPtr MinMax_) override;
+
+	//////////////////////////////////////////////////////////////////////////
+	/// LOD/MipMaps related
+	//////////////////////////////////////////////////////////////////////////
+	virtual void					SetLODLevel(int32 Level_, BlobPtr BlobObj, BlobPtrW LODParent_, BlobPtrW LODSource_, bool AddToBlobber) override;
+
+	//////////////////////////////////////////////////////////////////////////
+	/// Ownership related
+	//////////////////////////////////////////////////////////////////////////
+#if DEBUG_BLOB_REF_KEEPING == 1
+	virtual bool					HasBlobAsTile(Blob* BlobObj) const override;
+#endif 
+
+	//////////////////////////////////////////////////////////////////////////
+	/// Static functions
+	//////////////////////////////////////////////////////////////////////////
+
+	/// Takes a BlobObj and returns a TiledBlob. If the passed BlobObj is already tiled, then
+	/// it does nothing. Otherwise it creates a new TiledBlob with 1x1 tile that contains
+	/// the passed BlobObj.
+	/// Just a simple helper function, used in places that always require a TiledBlob.
+	static std::shared_ptr<TiledBlob>	AsTiledBlob(BlobPtr BlobObj);
+
+	//////////////////////////////////////////////////////////////////////////
+	/// Inline functions
+	//////////////////////////////////////////////////////////////////////////
+	FORCEINLINE size_t				Rows() const { return Tiles.Rows(); }
+	FORCEINLINE size_t				Cols() const { return Tiles.Cols(); }
+	FORCEINLINE const BlobPtrTiles&	GetTiles() const { return Tiles; }
+	FORCEINLINE uint32				GetWidth() const { return Desc.Width; }
+	FORCEINLINE uint32				GetHeight() const { return Desc.Height; }
+	FORCEINLINE bool				TiledTarget() const { return bTiledTarget; }
+	FORCEINLINE bool&				TiledTarget() { return bTiledTarget; }
+	FORCEINLINE bool				IsValidTileIndex(int32 TileX, int32 TileY) const
+	{
+		return TileX < static_cast<int32>(Tiles.Rows()) && TileY < static_cast<int32>(Tiles.Cols());
+	}
+
+	FORCEINLINE BlobRef				GetTile(int32 TileX, int32 TileY) const 
+	{
+		if (Tiles.Rows() == 1 && Tiles.Cols() == 1)
+			return BlobRef(Tiles[0][0], false);
+
+		check(IsValidTileIndex(TileX, TileY));
+		return BlobRef(Tiles[TileX][TileY], false);
+	}
+
+	JobPtrW							Job() const { return JobObj; }
+	JobPtrW&						Job() { return JobObj; }
+};
+
+//////////////////////////////////////////////////////////////////////////
+
+typedef std::shared_ptr<TiledBlob>	TiledBlobPtr;
+typedef std::unique_ptr<TiledBlob>	TiledBlobUPtr;
+typedef std::weak_ptr<TiledBlob>	TiledBlobPtrW;
+
+//////////////////////////////////////////////////////////////////////////
+
+//////////////////////////////////////////////////////////////////////////
+/// This is a promised TiledBlob that will be the result of the job.
+/// These are lazy evaluated.
+//////////////////////////////////////////////////////////////////////////
+class TEXTUREGRAPHENGINE_API TiledBlob_Promise : public TiledBlob
+{
+protected:
+	mutable std::vector<BlobReadyCallback> Callbacks;		/// The callbacks to call when the promised tiled BlobObj is ready
+	bool							bMakeSingleBlob = false;/// Make single BlobObj on finalise
+	std::weak_ptr<TiledBlob_Promise> CachedBlob;			/// The cached blob that this is derived off
+
+	virtual void					OnFinaliseInternal(BlobReadyCallback Callback) const override;
+	TiledBlob_Promise&				operator = (const TiledBlob_Promise& RHS);
+	virtual void					NotifyCallbacks();
+
+	virtual void					AddLinkedBlob(BlobPtr LinkedBlob) override;
+	virtual void					FinaliseFrom(const Blob* RHS) override;
+
+public:
+									TiledBlob_Promise(TiledBlobPtr Source);
+									TiledBlob_Promise(const BufferDescriptor& Desc, size_t NumTilesX, size_t NumTilesY, CHashPtr HashValue);
+									TiledBlob_Promise(DeviceBufferRef Buffer);
+	virtual							~TiledBlob_Promise();
+
+	virtual AsyncBufferResultPtr	Bind(const BlobTransform* Transform, const ResourceBindInfo& BindInfo) override;
+	virtual AsyncBufferResultPtr	Unbind(const BlobTransform* Transform, const ResourceBindInfo& BindInfo) override;
+	virtual AsyncBufferResultPtr	Flush(const ResourceBindInfo& BindInfo) override;
+	virtual bool					IsFinalised() const override { return bIsFinalised; }
+	virtual void					SetTile(int32 TileX, int32 TileY, BlobRef Tile) override;
+	virtual AsyncBufferResultPtr	MakeSingleBlob() override;
+	virtual void					CopyResolveLateBound(BlobPtr RHS) override;
+	virtual bool					IsPromise() const override { return true; }
+
+	virtual void					FinaliseNow(bool bNoCalcHash, CHashPtr FixedHash) override;
+	virtual AsyncBufferResultPtr	Finalise(bool bNoCalcHash, CHashPtr FixedHash) override;
+	AsyncBufferResultPtr			FinaliseFrom(std::shared_ptr<TiledBlob_Promise> RHS);
+
+	void							ResetForReplay(); /// For debug purpose, reset the state of the tile as a promise to NOT finalised, increment the replayCount
+
+	//////////////////////////////////////////////////////////////////////////
+	FORCEINLINE bool				IsSingleBlob() const { return bMakeSingleBlob; }
+};
+
+typedef std::shared_ptr<TiledBlob_Promise>	TiledBlob_PromisePtr;
+typedef std::weak_ptr<TiledBlob_Promise>	TiledBlob_PromisePtrW;
+typedef T_BlobRef<TiledBlob>				TiledBlobRef;
+typedef cti::continuable<TiledBlobRef>		AsyncTiledBlobRef;
diff --git a/Engine/Plugins/Experimental/TextureGraph/Source/TextureGraphEngine/Device/DeviceNativeTask.h b/Engine/Plugins/Experimental/TextureGraph/Source/TextureGraphEngine/Device/DeviceNativeTask.h
new file mode 100644
--- /dev/null
+++ b/Engine/Plugins/Experimental/TextureGraph/Source/TextureGraphEngine/Device/DeviceNativeTask.h
@@ -0,0 +1,177 @@
+// Copyright Epic Games, Inc. All Rights Reserved.
+#pragma once
+
+#include "CoreMinimal.h"
+#include "Helper/DataUtil.h"
+#include <vector>
+#include <unordered_set>
+#include "Helper/Util.h"
+#include "Async/TaskGraphInterfaces.h"
+
+class Device;
+class DeviceNativeTask;
+typedef cti::continuable<DeviceNativeTask>	AsyncDeviceNativeTask;
+typedef cti::continuable<DeviceNativeTask*>	AsyncDeviceNativeTaskP;
+typedef std::shared_ptr<DeviceNativeTask>	DeviceNativeTaskPtr;
+typedef std::vector<DeviceNativeTaskPtr>	DeviceNativeTaskPtrVec;
+
+#define DEVICE_NATIVE_TASKS_CHECK_CYCLES 0
+
+//////////////////////////////////////////////////////////////////////////
+/// Represents a Task that runs natively on a Device (e.g. Rendering 
+/// thread on an FX device or CPU and GPU threads on an OpenCL device)
+//////////////////////////////////////////////////////////////////////////
+class TEXTUREGRAPHENGINE_API DeviceNativeTask
+{
+private:
+	static std::atomic<int64_t>		GTaskId;				/// Task ID
+	static int64_t					GetNewTaskId();
+
+	std::promise<int32>				Promise;				/// Our promise
+	std::shared_future<int32>		Future;					/// The future for our promise
+
+protected:
+	DeviceNativeTaskPtrVec			Prev;					/// The previous Tasks that this is dependent on
+	std::atomic<bool>				bTerminate = false;		/// Terminate this process
+	size_t							ThreadId = 0;			/// What is the actual thread ID
+	int64_t							TaskId;					/// The ID of the Task
+	int32							Priority;				/// The TaskPriority of the job. Must be in range [E_Priority::kHighest, E_Priority::kLowest]
+	std::atomic<int>				ErrorCode = 0;			/// The error code for this native Task. This will be set in PostExec. 
+															/// Derived implementations must set this correcly to the error code of the job/Task
+
+	FString							Name;					/// The TaskName of this DeviceNativeTask
+
+	bool							bExecOnly = false;		/// Whether its an exec only Task with no pre and post exec. This is to avoid waiting for system Tasks
+	bool							bIsCulled = false;		/// Whether the Task is culled or not
+	bool							bIsDone = false;		/// Whether the Task is done or not
+
+	int64_t							BatchId = 0;			/// What's the batch id of this native Task
+
+	void							Reset();
+	void							WaitSelf();
+	bool							WaitSelfFor(int32 Seconds);
+
+	void							SetPromise(int32 Value);
+	FString							DebugWaitChain();
+
+	virtual cti::continuable<int32>	GenericExecAsync(ENamedThreads::Type CallingThread, ENamedThreads::Type ReturnThread);
+
+	void							AdjustPrevPriority(DeviceNativeTask* Ref, DeviceNativeTask* Parent);
+	void							AdjustPrevPriority(DeviceNativeTask* Ref, DeviceNativeTask* Parent, DeviceNativeTask* PrevDeps);
+	virtual void					SetPriority(DeviceNativeTask* Ref);
+
+public:
+									DeviceNativeTask(int32 TaskPriority, FString TaskName);
+									DeviceNativeTask(DeviceNativeTaskPtrVec& PrevTasks, int32 TaskPriority, FString TaskName);
+	virtual							~DeviceNativeTask();
+
+	virtual bool					WaitFor(int32 Seconds);
+	virtual void					Wait();
+	virtual cti::continuable<int32>	PreExecAsync(ENamedThreads::Type CallingThread, ENamedThreads::Type ReturnThread);
+	virtual void					PreExec();
+	virtual int32					Exec() = 0;
+	virtual cti::continuable<int32>	ExecAsync(ENamedThreads::Type CallingThread, ENamedThreads::Type ReturnThread);
+	virtual void					PostExec();
+	virtual cti::continuable<int32>	WaitAsync(ENamedThreads::Type ReturnThread = ENamedThreads::UnusedAnchor);
+
+	virtual Device*					GetTargetDevice() const = 0;
+
+#if (DEVICE_NATIVE_TASKS_CHECK_CYCLES == 1)
+	void							CheckQueued(DeviceNativeTask* Ref);
+	void							DebugVerifyDepsQueued(const DeviceNativeTaskPtrVec& Queue, const DeviceNativeTaskPtrVec& AllItems);
+#endif 
+
+	void							FixPriorities();
+
+	/// Waits for the _prev Task to finish before executing. Its the responsibility of the system
+	/// to ensure that both the _prev
+	virtual ENamedThreads::Type		GetExecutionThread() const;
+	virtual bool					IsAsync() const;
+	virtual FString					GetName() const;
+	virtual FString					GetDebugName() const;
+	virtual void					Terminate();
+
+	virtual void					AddPrev(DeviceNativeTaskPtr PrevTasks);
+	virtual void					AddPrev(const DeviceNativeTaskPtrVec& Tasks);
+	virtual void					SetPrev(const DeviceNativeTaskPtrVec& PrevTasks);
+	virtual void					SetPriority(int32 TaskPriority, bool AdjustPrev);
+	bool							IsHigherPriorityThan(const DeviceNativeTask& RHS) const;
+
+	void							CheckCyclicalDependency(std::unordered_set<DeviceNativeTask*>& TransferChain);
+	virtual bool					DebugCompleteCheck();
+
+	//////////////////////////////////////////////////////////////////////////
+	/// Inline functions
+	//////////////////////////////////////////////////////////////////////////
+	FORCEINLINE std::shared_future<int32> GetFuture() { return Future; }
+	FORCEINLINE size_t				GetActualThreadId() const { return ThreadId; }
+	FORCEINLINE int64_t				GetTaskId() const { return TaskId; }
+	FORCEINLINE int64_t				GetBatchId() const { return BatchId; }
+	FORCEINLINE void				SetBatchId(int64_t NewBatchId) { BatchId = NewBatchId; }
+	FORCEINLINE int32				GetPriority() const { return Priority; }
+	FORCEINLINE bool				operator < (const DeviceNativeTask& RHS) const { return Priority < RHS.Priority; }
+	FORCEINLINE bool				IsExecOnly() const { return bExecOnly; }
+	FORCEINLINE bool				IsCulled() const { return bIsCulled; }
+	FORCEINLINE bool				IsDone() const { return bIsDone || IsCulled(); }
+	FORCEINLINE const DeviceNativeTaskPtrVec& GetPrev() const { return Prev; }
+};
+
+//////////////////////////////////////////////////////////////////////////
+/// Below are some generic helper implementations
+/// Note that if this starts getting too big, then we'll need to move
+/// these to their individual files under a Device/Task/ sub-directory
+//////////////////////////////////////////////////////////////////////////
+
+//////////////////////////////////////////////////////////////////////////
+class TEXTUREGRAPHENGINE_API DeviceNativeTask_Generic : public DeviceNativeTask
+{
+	friend class Device;
+
+protected:
+	Device*							Dev = nullptr;			/// The target device that this is associated with
+	ENamedThreads::Type				ThreadId = ENamedThreads::AnyThread; /// What thread is this going to run on
+	bool							bIsAsync = false;		/// Is this asynchronous or not
+
+public:
+									DeviceNativeTask_Generic(Device* TargetDevice, FString TaskName, ENamedThreads::Type TaskThreadId = ENamedThreads::AnyThread, 
+										bool IsTaskAsync = false, int32 TaskPriority = (int32)E_Priority::kNormal);
+									DeviceNativeTask_Generic(DeviceNativeTaskPtrVec& PrevDeps, Device* TargetDevice, FString TaskName, ENamedThreads::Type TaskThreadId = ENamedThreads::AnyThread, 
+										bool IsTAskAsync = false, int32 TaskPriority = (int32)E_Priority::kNormal);
+
+	virtual Device*					GetTargetDevice() const override { return Dev; }
+
+	/// Waits for the _prev Task to finish before executing. Its the responsibility of the system
+	/// to ensure that both the _prev
+	virtual ENamedThreads::Type		GetExecutionThread() const override { return ThreadId; }
+	virtual bool					IsAsync() const override { return bIsAsync; }
+
+	//////////////////////////////////////////////////////////////////////////
+	/// Inline functions
+	//////////////////////////////////////////////////////////////////////////
+	FORCEINLINE ENamedThreads::Type& GetThreadId() { return ThreadId; }
+	FORCEINLINE bool&				IsAsync() { return bIsAsync; }
+};
+
+//////////////////////////////////////////////////////////////////////////
+typedef TFunction<int32()>			DeviceNativeTask_Lambda_Func;
+class TEXTUREGRAPHENGINE_API DeviceNativeTask_Lambda : public DeviceNativeTask_Generic
+{
+protected:
+	DeviceNativeTask_Lambda_Func	Callback = nullptr;	/// The lambda callback
+
+	virtual int32					Exec() override;
+
+public:
+									DeviceNativeTask_Lambda(DeviceNativeTask_Lambda_Func LambdaCallback, Device* Dev, FString TaskName, 
+										ENamedThreads::Type TaskThreadId = ENamedThreads::AnyThread, bool IsTaskAsync = false, int32 TaskPriority = (int32)E_Priority::kNormal);
+
+	//////////////////////////////////////////////////////////////////////////
+	/// Static functions
+	//////////////////////////////////////////////////////////////////////////
+	static std::shared_ptr<DeviceNativeTask> Create(Device* device, int32 TaskPriority, FString TaskName, DeviceNativeTask_Lambda_Func callback);
+	static std::shared_ptr<DeviceNativeTask> Create(Device* device, ENamedThreads::Type threadId, bool isAsync, int32 TaskPriority, FString TaskName, DeviceNativeTask_Lambda_Func callback);
+
+	//static cti::continuable<int32>	Create_Promise(Device* device, int32 TaskPriority, DeviceNativeTask_Lambda_Func callback);
+	//static cti::continuable<int32>	Create_Promise(Device* device, ENamedThreads::Type threadId, bool isAsync, int32 TaskPriority, DeviceNativeTask_Lambda_Func callback);
+};
+
diff --git a/Engine/Plugins/Experimental/TextureGraph/Source/TextureGraphEngine/Device/Disk/DeviceBuffer_Disk.h b/Engine/Plugins/Experimental/TextureGraph/Source/TextureGraphEngine/Device/Disk/DeviceBuffer_Disk.h
new file mode 100644
--- /dev/null
+++ b/Engine/Plugins/Experimental/TextureGraph/Source/TextureGraphEngine/Device/Disk/DeviceBuffer_Disk.h
@@ -0,0 +1,44 @@
+// Copyright Epic Games, Inc. All Rights Reserved.
+#pragma once
+
+#include "CoreMinimal.h"
+#include "Device/DeviceBuffer.h"
+#include "Data/RawBuffer.h"
+
+class Device;
+class Device_Disk;
+
+//////////////////////////////////////////////////////////////////////////
+/// This is simply a wrapper around RawBuffer. Must be owned by the 
+/// CPU device (though that can change in the future)
+//////////////////////////////////////////////////////////////////////////
+class TEXTUREGRAPHENGINE_API DeviceBuffer_Disk : public DeviceBuffer
+{
+protected:
+	RawBufferPtr					RawObj;				/// The underlying raw buffer
+
+	virtual DeviceBuffer*			CreateFromRaw(RawBufferPtr NewRawObj) override;
+	virtual DeviceBuffer*			CreateFromDesc(BufferDescriptor NewDesc, CHashPtr NewHashValue) override;
+
+	virtual AsyncBufferResultPtr	TransferFrom(DeviceBufferRef& Source) override;
+
+public:
+									DeviceBuffer_Disk(Device_Disk* Dev, BufferDescriptor Desc, CHashPtr HashValue);
+									DeviceBuffer_Disk(Device_Disk* Dev, RawBufferPtr RawObj);
+	virtual							~DeviceBuffer_Disk() override;
+
+	//////////////////////////////////////////////////////////////////////////
+	/// DeviceBuffer implementation
+	//////////////////////////////////////////////////////////////////////////
+	virtual RawBufferPtr			Raw_Now() override;
+	virtual size_t					MemSize() const override;
+	virtual AsyncBufferResultPtr	Bind(const BlobTransform* Transform, const ResourceBindInfo& BindInfo) override;
+
+	/// With base implementations
+	//virtual void					Release() override;
+	virtual bool					IsCompatible(Device* Dev) const override;
+
+	//////////////////////////////////////////////////////////////////////////
+	/// Inline functions
+	//////////////////////////////////////////////////////////////////////////
+};
diff --git a/Engine/Plugins/Experimental/TextureGraph/Source/TextureGraphEngine/Device/Disk/Device_Disk.h b/Engine/Plugins/Experimental/TextureGraph/Source/TextureGraphEngine/Device/Disk/Device_Disk.h
new file mode 100644
--- /dev/null
+++ b/Engine/Plugins/Experimental/TextureGraph/Source/TextureGraphEngine/Device/Disk/Device_Disk.h
@@ -0,0 +1,30 @@
+// Copyright Epic Games, Inc. All Rights Reserved.
+#pragma once
+
+#include "CoreMinimal.h"
+#include "Device/Mem/Device_Mem.h"
+
+class TEXTUREGRAPHENGINE_API Device_Disk : public Device_Mem
+{
+protected:
+	FString							BaseDir;				/// The base directory for 
+
+public:
+									Device_Disk();
+	virtual							~Device_Disk() override;
+
+	virtual FString					Name() const override { return "Device_Disk"; }
+	virtual void					Update(float Delta) override;
+	virtual FString					SetBaseDirectory(const FString& InBaseDir, bool MigrateExisting = true);
+	FString							GetCacheFilename(HashType HashValue) const;
+
+	//////////////////////////////////////////////////////////////////////////
+	/// Static functions
+	//////////////////////////////////////////////////////////////////////////
+	static Device_Disk*				Get();
+
+	//////////////////////////////////////////////////////////////////////////
+	/// Inline function
+	//////////////////////////////////////////////////////////////////////////
+	FORCEINLINE const FString&		GetBaseDirectory() const { return BaseDir; }
+};
diff --git a/Engine/Plugins/Experimental/TextureGraph/Source/TextureGraphEngine/Device/FX/DeviceBuffer_FX.h b/Engine/Plugins/Experimental/TextureGraph/Source/TextureGraphEngine/Device/FX/DeviceBuffer_FX.h
new file mode 100644
--- /dev/null
+++ b/Engine/Plugins/Experimental/TextureGraph/Source/TextureGraphEngine/Device/FX/DeviceBuffer_FX.h
@@ -0,0 +1,59 @@
+// Copyright Epic Games, Inc. All Rights Reserved.
+#pragma once
+
+#include "CoreMinimal.h"
+#include "Device/DeviceBuffer.h"
+
+class Device_FX;
+class DeviceBuffer_Mem;
+
+class Tex;
+typedef std::shared_ptr<Tex>		TexPtr;
+
+class TEXTUREGRAPHENGINE_API DeviceBuffer_FX : public DeviceBuffer
+{
+	friend class Device_FX;
+protected: 
+	TexPtr							Texture;				/// The underlying texture object
+
+	AsyncBufferResultPtr			AllocateTexture();
+	AsyncBufferResultPtr			AllocateRenderTarget();
+	void							UpdateRenderTarget(TexPtr TextureObj);
+
+	virtual DeviceBuffer*			CreateFromRaw(RawBufferPtr RawObj) override;
+	virtual DeviceBuffer*			CreateFromDesc(BufferDescriptor NewDesc, CHashPtr NewHashValue) override;
+	virtual DeviceBuffer*			CopyFrom(DeviceBuffer* RHS) override;
+	virtual AsyncBufferResultPtr	UpdateRaw(RawBufferPtr RawObj) override;
+
+	virtual CHashPtr				CalcHash() override;
+
+	virtual AsyncBufferResultPtr	TransferFrom(DeviceBufferRef& Source) override;
+
+public:
+									DeviceBuffer_FX(Device_FX* Dev, BufferDescriptor Desc, CHashPtr NewHashValue);
+									DeviceBuffer_FX(Device_FX* Dev, TexPtr TextureObj, RawBufferPtr RawObj);
+									DeviceBuffer_FX(Device_FX* Dev, TexPtr TextureObj, BufferDescriptor Desc, CHashPtr NewHashValue);
+									DeviceBuffer_FX(Device_FX* Dev, RawBufferPtr RawObj);
+	virtual							~DeviceBuffer_FX() override;
+
+	//////////////////////////////////////////////////////////////////////////
+	/// DeviceBuffer Implementation
+	//////////////////////////////////////////////////////////////////////////
+	virtual RawBufferPtr			Raw_Now() override;
+	virtual size_t					MemSize() const override;
+	virtual size_t					DeviceNative_MemSize() const override;
+	virtual AsyncBufferResultPtr	Bind(const BlobTransform* Transform, const ResourceBindInfo& BindInfo) override;
+	virtual AsyncBufferResultPtr	Unbind(const BlobTransform* Transform, const ResourceBindInfo& BindInfo) override;
+	virtual AsyncBufferResultPtr	Flush(const ResourceBindInfo& BindInfo) override;
+	virtual AsyncPrepareResult		PrepareForWrite(const ResourceBindInfo& BindInfo) override;
+	virtual bool					IsNull() const override;
+
+	virtual void					ReleaseNative() override;
+
+	//////////////////////////////////////////////////////////////////////////
+	/// Inline functions
+	//////////////////////////////////////////////////////////////////////////
+	FORCEINLINE TexPtr				GetTexture() { return Texture; }
+};
+
+typedef std::shared_ptr<DeviceBuffer_FX> DeviceBuffer_FXPtr;
diff --git a/Engine/Plugins/Experimental/TextureGraph/Source/TextureGraphEngine/Device/FX/Device_FX.h b/Engine/Plugins/Experimental/TextureGraph/Source/TextureGraphEngine/Device/FX/Device_FX.h
new file mode 100644
--- /dev/null
+++ b/Engine/Plugins/Experimental/TextureGraph/Source/TextureGraphEngine/Device/FX/Device_FX.h
@@ -0,0 +1,107 @@
+// Copyright Epic Games, Inc. All Rights Reserved.
+#pragma once
+
+#include "CoreMinimal.h"
+#include "Data/Blob.h"
+#include "Device/Device.h"
+#include "DeviceBuffer_FX.h"
+#include "RenderTargetPool.h"
+#include "Engine/TextureRenderTarget2D.h"
+#include <list>
+#include <unordered_map>
+
+typedef cti::continuable<TexPtr>	AsyncTexPtr;
+class RenderMaterial_FX;
+typedef std::shared_ptr<RenderMaterial_FX> RenderMaterial_FXPtr;
+typedef std::shared_ptr<class TexArray>	TexArrayPtr;
+class UTextureRenderTarget;
+class IRendererModule;
+
+class TEXTUREGRAPHENGINE_API Device_FX : public Device,  public FGCObject
+{
+private: 
+	static size_t					s_maxRenderBatch;
+
+	typedef std::list<TexPtr>		TextureNodeList;
+
+	mutable FCriticalSection		GCLock;							/// Lock to the GC lists
+	TextureNodeList					GCTargetTextures;				/// Textures that are GC targets right now
+
+	typedef std::list<TObjectPtr<UTextureRenderTarget2D>> RTList;
+	typedef std::unordered_map<HashType, RTList*> RenderTargetCache;
+
+	mutable FCriticalSection		GC_RTCache;						/// Lock to the GC lists
+
+	RenderTargetCache				RTCache;						/// The render target cache that we are currently maintaining
+	RenderTargetCache				RTArrayCache;					/// The render target array cache that we are currently maintaining
+
+	RTList							RTUsed;							/// Used render targets
+
+	IRendererModule*				RendererModule;					/// Useful rendering functionalities wrapped in this module
+
+	void							GCTextures();
+	virtual void					Free() override;
+	virtual void					ClearCache() override;
+
+	//////////////////////////////////////////////////////////////////////////
+	/// FGCObject
+	//////////////////////////////////////////////////////////////////////////
+	virtual void					AddReferencedObjects(FReferenceCollector& Collector) override;
+	virtual FString					GetReferencerName() const override;
+
+	void							FreeCacheInternal(RenderTargetCache& TargetRTCache);
+	void							FreeRTList(RTList& RTList);
+	void							FreeRenderTarget(HashType HashValue, UTextureRenderTarget2D* RT);
+
+	static UTextureRenderTarget2D*	CreateRenderTarget(const BufferDescriptor& Desc);
+	AsyncDeviceBufferRef			SplitToTiles_Internal(const CombineSplitArgs& SplitArgs);
+
+	int32							AllocateRTArrayResource(TexArrayPtr TexArrayObj);
+	int32							AllocateRTResource(TexPtr TextureObj);
+	int32							InitRTResource(TexPtr TextureObj, UTextureRenderTarget* RT);
+
+protected:
+	virtual void					PrintStats() override;
+
+public:
+									Device_FX();
+	virtual							~Device_FX() override;
+
+	virtual cti::continuable<int32>	Use() const override;
+	void							GetStatArray(TArray<FString>& ResourceArrayTiled, TArray<FString>& ResourceArrayUnTiled, TArray<FString>& TooltipListTiled, TArray<FString>& TooltipListUnTiled, FString& TotalStats);
+	/// Must be called from the rendering thread
+	DeviceBufferRef					CreateFromRT(UTextureRenderTarget2D* RT, const BufferDescriptor& Desc);
+	DeviceBufferRef					CreateFromTexture(UTexture2D* TextureObj, const BufferDescriptor& Desc);
+	DeviceBufferRef					CreateFromTex(TexPtr TextureObj, bool bInitRaw);
+	DeviceBufferRef					CreateFromTexAndRaw(TexPtr TextureObj, RawBufferPtr RawObj);
+
+	virtual AsyncDeviceBufferRef	CombineFromTiles(const CombineSplitArgs& CombineArgs) override;
+
+	AsyncDeviceBufferRef			DrawTilesToBuffer_Deferred(DeviceBufferRef buffer, const T_Tiles<DeviceBufferRef>& tiles);
+	AsyncDeviceBufferRef			FillTextureArray_Deferred(DeviceBufferRef buffer, const T_Tiles<DeviceBufferRef>& tiles);
+
+	virtual AsyncDeviceBufferRef	SplitToTiles(const CombineSplitArgs& splitArgs) override;
+
+	TexPtr							AllocateRenderTarget(const BufferDescriptor& Desc);
+	TexPtr							AllocateRenderTargetArray(const BufferDescriptor& Desc, int32 NumTilesX, int32 NumTilesY);
+
+
+
+	virtual void					Update(float Delta) override;
+	void							MarkForCollection(TexPtr TextureObj);
+	virtual FString					Name() const override { return "Device_FX"; }
+	virtual void					AddNativeTask(DeviceNativeTaskPtr task) override;
+
+	//////////////////////////////////////////////////////////////////////////
+	/// Inline functions
+	//////////////////////////////////////////////////////////////////////////
+	FORCEINLINE FRHICommandListImmediate& RHI() const { check(IsInRenderingThread()); return GRHICommandList.GetImmediateCommandList(); }
+
+	//////////////////////////////////////////////////////////////////////////
+	/// Static functions
+	//////////////////////////////////////////////////////////////////////////
+	static void						InitTiles_Texture(BlobUPtr* tiles, size_t NumRows, size_t NumCols, const BufferDescriptor& InTileDesc, bool bInitRaw);
+	static void						InitTiles_RenderTargets(BlobUPtr* Tiles, size_t NumRows, size_t NumCols, const BufferDescriptor& InTileDesc, bool bInitRaw);
+
+	static Device_FX*				Get();
+};
diff --git a/Engine/Plugins/Experimental/TextureGraph/Source/TextureGraphEngine/Device/Mem/DeviceBuffer_Mem.h b/Engine/Plugins/Experimental/TextureGraph/Source/TextureGraphEngine/Device/Mem/DeviceBuffer_Mem.h
new file mode 100644
--- /dev/null
+++ b/Engine/Plugins/Experimental/TextureGraph/Source/TextureGraphEngine/Device/Mem/DeviceBuffer_Mem.h
@@ -0,0 +1,43 @@
+// Copyright Epic Games, Inc. All Rights Reserved.
+#pragma once
+
+#include "CoreMinimal.h"
+#include "Device/DeviceBuffer.h"
+#include "Data/RawBuffer.h"
+
+class Device;
+class Device_Mem;
+
+//////////////////////////////////////////////////////////////////////////
+/// This is simply a wrapper around RawBuffer. Must be owned by the 
+/// CPU device (though that can change in the future)
+//////////////////////////////////////////////////////////////////////////
+class TEXTUREGRAPHENGINE_API DeviceBuffer_Mem : public DeviceBuffer
+{
+protected:
+	virtual DeviceBuffer*			CreateFromRaw(RawBufferPtr RawObj) override;
+	virtual DeviceBuffer*			CreateFromDesc(BufferDescriptor NewDesc, CHashPtr NewHashValue) override;
+
+	void							Allocate();
+	virtual AsyncBufferResultPtr	TransferFrom(DeviceBufferRef& Source) override;
+
+public:
+									DeviceBuffer_Mem(Device_Mem* Device, BufferDescriptor Desc, CHashPtr HashValue);
+									DeviceBuffer_Mem(Device_Mem* Device, RawBufferPtr RawObj);
+	virtual							~DeviceBuffer_Mem() override;
+
+	//////////////////////////////////////////////////////////////////////////
+	/// DeviceBuffer implementation
+	//////////////////////////////////////////////////////////////////////////
+	virtual RawBufferPtr			Raw_Now() override;
+	virtual size_t					MemSize() const override;
+	virtual AsyncBufferResultPtr	Bind(const BlobTransform* Transform, const ResourceBindInfo& BindInfo) override;
+
+	/// With base implementations
+	//virtual void					Release() override;
+	virtual bool					IsCompatible(Device* Dev) const override;
+
+	//////////////////////////////////////////////////////////////////////////
+	/// Inline functions
+	//////////////////////////////////////////////////////////////////////////
+};
diff --git a/Engine/Plugins/Experimental/TextureGraph/Source/TextureGraphEngine/Device/Mem/Device_Mem.h b/Engine/Plugins/Experimental/TextureGraph/Source/TextureGraphEngine/Device/Mem/Device_Mem.h
new file mode 100644
--- /dev/null
+++ b/Engine/Plugins/Experimental/TextureGraph/Source/TextureGraphEngine/Device/Mem/Device_Mem.h
@@ -0,0 +1,24 @@
+// Copyright Epic Games, Inc. All Rights Reserved.
+#pragma once
+
+#include "CoreMinimal.h"
+#include "Device/Device.h"
+
+class TEXTUREGRAPHENGINE_API Device_Mem : public Device
+{
+protected:
+									Device_Mem(DeviceType Type, DeviceBuffer* BufferFactory);
+
+public:
+									Device_Mem();
+	virtual							~Device_Mem() override;
+
+	virtual FString					Name() const override { return "Device_Mem"; }
+	virtual void					Update(float Delta) override;
+	virtual void					AddNativeTask(DeviceNativeTaskPtr Task) override;
+
+	//////////////////////////////////////////////////////////////////////////
+	/// Static functions
+	//////////////////////////////////////////////////////////////////////////
+	static Device_Mem*				Get();
+};
diff --git a/Engine/Plugins/Experimental/TextureGraph/Source/TextureGraphEngine/Device/MemCM/DeviceBuffer_MemCM.h b/Engine/Plugins/Experimental/TextureGraph/Source/TextureGraphEngine/Device/MemCM/DeviceBuffer_MemCM.h
new file mode 100644
--- /dev/null
+++ b/Engine/Plugins/Experimental/TextureGraph/Source/TextureGraphEngine/Device/MemCM/DeviceBuffer_MemCM.h
@@ -0,0 +1,41 @@
+// Copyright Epic Games, Inc. All Rights Reserved.
+#pragma once
+
+#include "CoreMinimal.h"
+#include "Device/Mem/DeviceBuffer_Mem.h"
+#include "Data/RawBuffer.h"
+
+class Device;
+class Device_MemCM;
+
+//////////////////////////////////////////////////////////////////////////
+/// This is simply a wrapper around RawBuffer. Must be owned by the 
+/// CPU device (though that can change in the future)
+//////////////////////////////////////////////////////////////////////////
+class TEXTUREGRAPHENGINE_API DeviceBuffer_MemCM : public DeviceBuffer_Mem
+{
+protected:
+	virtual DeviceBuffer*			CreateFromRaw(RawBufferPtr InRawObj) override;
+	virtual DeviceBuffer*			CreateFromDesc(BufferDescriptor InDesc, CHashPtr InHashValue) override;
+
+	virtual AsyncBufferResultPtr	TransferFrom(DeviceBufferRef& Source) override;
+public:
+									DeviceBuffer_MemCM(Device_MemCM* Dev, BufferDescriptor InDesc, CHashPtr InHashValue);
+									DeviceBuffer_MemCM(Device_MemCM* Dev, RawBufferPtr InRawObj);
+	virtual							~DeviceBuffer_MemCM() override;
+
+	//////////////////////////////////////////////////////////////////////////
+	/// DeviceBuffer implementation
+	//////////////////////////////////////////////////////////////////////////
+	virtual RawBufferPtr			Raw_Now() override;
+	virtual size_t					MemSize() const override;
+	virtual AsyncBufferResultPtr	Bind(const BlobTransform* Transform, const ResourceBindInfo& BindInfo) override;
+
+	/// With base implementations
+	//virtual void					Release() override;
+	virtual bool					IsCompatible(Device* Dev) const override;
+
+	//////////////////////////////////////////////////////////////////////////
+	/// Inline functions
+	//////////////////////////////////////////////////////////////////////////
+};
diff --git a/Engine/Plugins/Experimental/TextureGraph/Source/TextureGraphEngine/Device/MemCM/Device_MemCM.h b/Engine/Plugins/Experimental/TextureGraph/Source/TextureGraphEngine/Device/MemCM/Device_MemCM.h
new file mode 100644
--- /dev/null
+++ b/Engine/Plugins/Experimental/TextureGraph/Source/TextureGraphEngine/Device/MemCM/Device_MemCM.h
@@ -0,0 +1,20 @@
+// Copyright Epic Games, Inc. All Rights Reserved.
+#pragma once
+
+#include "CoreMinimal.h"
+#include "Device/Mem/Device_Mem.h"
+
+class TEXTUREGRAPHENGINE_API Device_MemCM : public Device_Mem
+{
+public:
+									Device_MemCM();
+	virtual							~Device_MemCM() override;
+
+	virtual FString					Name() const override { return "Device_MemCM"; }
+	virtual void					Update(float Delta) override;
+
+	//////////////////////////////////////////////////////////////////////////
+	/// Static functions
+	//////////////////////////////////////////////////////////////////////////
+	static Device_MemCM*			Get();
+};
diff --git a/Engine/Plugins/Experimental/TextureGraph/Source/TextureGraphEngine/Device/Null/DeviceBuffer_Null.h b/Engine/Plugins/Experimental/TextureGraph/Source/TextureGraphEngine/Device/Null/DeviceBuffer_Null.h
new file mode 100644
--- /dev/null
+++ b/Engine/Plugins/Experimental/TextureGraph/Source/TextureGraphEngine/Device/Null/DeviceBuffer_Null.h
@@ -0,0 +1,27 @@
+// Copyright Epic Games, Inc. All Rights Reserved.
+#pragma once
+
+#include "CoreMinimal.h"
+#include "Device/DeviceBuffer.h"
+
+class Device_Null;
+
+class TEXTUREGRAPHENGINE_API DeviceBuffer_Null : public DeviceBuffer
+{
+protected:
+	virtual DeviceBuffer*			CreateFromRaw(RawBufferPtr rawBuffer) override;
+	virtual DeviceBuffer*			CreateFromDesc(BufferDescriptor desc, CHashPtr hash) override;
+
+	virtual AsyncBufferResultPtr	TransferFrom(DeviceBufferRef& source);
+
+public:
+									DeviceBuffer_Null(Device_Null* device, BufferDescriptor desc, CHashPtr hash);
+	virtual							~DeviceBuffer_Null();
+
+	virtual RawBufferPtr			Raw_Now();
+	virtual size_t					MemSize() const;
+	virtual AsyncBufferResultPtr	Bind(const BlobTransform* transform, const ResourceBindInfo& bindInfo) override;
+	virtual bool					IsValid() const;
+	virtual bool					IsNull() const;
+	virtual bool					IsTransient() const override;
+};
diff --git a/Engine/Plugins/Experimental/TextureGraph/Source/TextureGraphEngine/Device/Null/Device_Null.h b/Engine/Plugins/Experimental/TextureGraph/Source/TextureGraphEngine/Device/Null/Device_Null.h
new file mode 100644
--- /dev/null
+++ b/Engine/Plugins/Experimental/TextureGraph/Source/TextureGraphEngine/Device/Null/Device_Null.h
@@ -0,0 +1,23 @@
+// Copyright Epic Games, Inc. All Rights Reserved.
+#pragma once
+
+#include "CoreMinimal.h"
+#include "Device/Device.h"
+
+class TEXTUREGRAPHENGINE_API Device_Null : public Device
+{
+protected:
+
+	virtual void					Collect(DeviceBuffer* Buffer) override;
+
+public:
+									Device_Null();
+	virtual							~Device_Null() override;
+
+	virtual FString					Name() const override { return "Device_Null"; }
+
+	//////////////////////////////////////////////////////////////////////////
+	/// Static functions
+	//////////////////////////////////////////////////////////////////////////
+	static Device_Null*				Get();
+};
diff --git a/Engine/Plugins/Experimental/TextureGraph/Source/TextureGraphEngine/Export/T_ExportAsUAsset.h b/Engine/Plugins/Experimental/TextureGraph/Source/TextureGraphEngine/Export/T_ExportAsUAsset.h
new file mode 100644
--- /dev/null
+++ b/Engine/Plugins/Experimental/TextureGraph/Source/TextureGraphEngine/Export/T_ExportAsUAsset.h
@@ -0,0 +1,42 @@
+// Copyright Epic Games, Inc. All Rights Reserved.
+#pragma once
+
+#include "CoreMinimal.h" 
+#include "Job/Job.h"
+#include "TextureExporter.h"
+
+//////////////////////////////////////////////////////////////////////////
+class TEXTUREGRAPHENGINE_API Job_ExportAsUAsset : public Job
+{
+private:
+	FExportMapSettings				Setting;
+	FString							OutFolder;
+
+public:
+									Job_ExportAsUAsset(UMixInterface* Mix, int32 TargetId, FExportMapSettings InSetting, FString OutputFolder, UObject* InErrorOwner = nullptr, uint16 Priority = (uint16)E_Priority::kHigh, uint64 id = 0);
+
+protected:
+	virtual void					GetDependencies(JobPtrVec& prior, JobPtrVec& after, JobRunInfo runInfo) override { RunInfo = runInfo; Prev.clear(); }
+	virtual AsyncPrepareResult		PrepareTargets(JobBatch* batch) override { return cti::make_ready_continuable(0); }
+	virtual AsyncPrepareResult		PrepareResources(JobBatch* batch) override { return cti::make_ready_continuable(0); }
+
+	virtual int32					Exec() override { return 0; }
+	virtual cti::continuable<int32>	BindArgs_All(JobRunInfo runInfo) override { return cti::make_ready_continuable(0); }
+	virtual cti::continuable<int32>	UnbindArgs_All(JobRunInfo runInfo) override { return cti::make_ready_continuable(0); }
+	virtual bool					CanHandleTiles() const override { return false; }
+
+	virtual cti::continuable<int32>	PreExecAsync(ENamedThreads::Type execThread, ENamedThreads::Type returnThread) override;
+	virtual cti::continuable<int32>	ExecAsync(ENamedThreads::Type execThread, ENamedThreads::Type returnThread) override;
+};
+
+class TEXTUREGRAPHENGINE_API T_ExportAsUAsset
+{
+public:
+									T_ExportAsUAsset();
+	virtual							~T_ExportAsUAsset();
+
+	//////////////////////////////////////////////////////////////////////////
+	/// Static functions
+	//////////////////////////////////////////////////////////////////////////
+	static void						ExportAsUAsset(MixUpdateCyclePtr cycle, int32 targetId, FExportMapSettings exportMap, FString outFolder);
+};
diff --git a/Engine/Plugins/Experimental/TextureGraph/Source/TextureGraphEngine/Export/TextureExporter.h b/Engine/Plugins/Experimental/TextureGraph/Source/TextureGraphEngine/Export/TextureExporter.h
new file mode 100644
--- /dev/null
+++ b/Engine/Plugins/Experimental/TextureGraph/Source/TextureGraphEngine/Export/TextureExporter.h
@@ -0,0 +1,101 @@
+// Copyright Epic Games, Inc. All Rights Reserved.
+#pragma once
+#include "2D/TargetTextureSet.h"
+#include "CoreMinimal.h"
+#include "Misc/AutomationTest.h"
+#include "Model/Mix/MixInterface.h"
+THIRD_PARTY_INCLUDES_START
+#include "continuable/continuable.hpp"
+THIRD_PARTY_INCLUDES_END
+#include <vector>
+
+#if WITH_EDITOR
+#include "Editor.h"
+#endif
+
+#include "CoreGlobals.h"
+#include "Framework/Application/SlateApplication.h"
+
+DECLARE_LOG_CATEGORY_EXTERN(ExportLogs, Log, Verbose);
+
+struct TEXTUREGRAPHENGINE_API FExportMapSettings
+{
+	bool bIsEnabled = true;
+	FName Name;
+	FString Path;
+	bool UseOverridePath;
+	bool OverwriteTextures;
+	TiledBlobPtr Map;
+	bool IsSRGB;
+	bool UseMipMaps;
+	int32 Width;
+	int32 Height;
+	TextureCompressionSettings Compression;
+
+	DECLARE_DELEGATE_OneParam(OnExportMapDone, FExportMapSettings MapSettings);
+	OnExportMapDone OnDone;
+};
+
+typedef std::vector<std::pair<FName, FExportMapSettings>> ExportMaps;
+
+class TEXTUREGRAPHENGINE_API FExportSettings
+{
+public:
+	FString PresetName;
+	ExportMaps ExportPreset;
+
+	DECLARE_DELEGATE(OnExportDone);
+	OnExportDone OnDone;
+
+	int MapsExported = 0;
+	
+	int GetMapCount()
+	{
+		int MapCount = 0;
+
+		
+		if(ExportPreset.size() > 0)
+		{
+			for (const auto& Map : ExportPreset)
+			{
+				if(Map.second.bIsEnabled)
+				{
+					MapCount++;
+				}
+			}
+		}
+
+		return MapCount;
+	}
+
+	void Reset()
+	{
+		MapsExported = 0;
+		ExportPreset.clear();
+	}
+
+	void OnMapExportDone(FExportMapSettings MapSetting)
+	{
+		if (!FSlateApplication::IsInitialized() || !FApp::CanEverRender())
+			return;
+
+		MapsExported++;
+		if (GetMapCount() == MapsExported)
+		{
+			OnDone.ExecuteIfBound();
+		}
+
+		if(MapSetting.OnDone.IsBound())
+			MapSetting.OnDone.Unbind();
+	}
+};
+
+struct TEXTUREGRAPHENGINE_API TextureExporter
+{
+	static FExportMapSettings		GetExportSettingsForTarget(FExportSettings& ExportSettings, TiledBlobPtr BlobObj, FName Name);
+	static AsyncBool				ExportAllMapsAsUAsset(UMixInterface* MixObj, const FString& OutputFolder, const FExportSettings& ExportSettings);
+	static AsyncBool				ExportAsUAsset(UMixInterface* MixObj, const FExportSettings& ExportSettings, FString C);
+	static void						ExportAsUAsset(UMixInterface* MixObj, FName AssetName, FString CurrentPath, const FExportSettings& Setting);
+	static AsyncInt					ExportRawAsUAsset(RawBufferPtr RawObj,const FExportMapSettings& Setting, const FString& CompletePath, const FName& TextureName);
+	static bool						IsPackageNameValid(FString Path, FString AssetName);
+};
diff --git a/Engine/Plugins/Experimental/TextureGraph/Source/TextureGraphEngine/FxMat/FxMaterial_DrawMaterial.h b/Engine/Plugins/Experimental/TextureGraph/Source/TextureGraphEngine/FxMat/FxMaterial_DrawMaterial.h
new file mode 100644
--- /dev/null
+++ b/Engine/Plugins/Experimental/TextureGraph/Source/TextureGraphEngine/FxMat/FxMaterial_DrawMaterial.h
@@ -0,0 +1,224 @@
+// Copyright Epic Games, Inc. All Rights Reserved.
+#pragma once
+
+#include "CoreMinimal.h"
+
+#include <Materials/Material.h>
+#include <Materials/MaterialRenderProxy.h>
+#include <MaterialDomain.h>
+#include <TextureResource.h>
+#include <SceneManagement.h>
+#include <PostProcess/DrawRectangle.h>
+
+#include <DataDrivenShaderPlatformInfo.h>
+#include <MaterialShader.h>
+#include <MaterialShared.h>
+#include "Engine/TextureRenderTarget2D.h"
+
+#include "FxMaterial.h"
+#include "TileInfo_FX.h"
+
+#include "Device/Device.h"
+
+class TEXTUREGRAPHENGINE_API FTextureGraphMaterialShaderVS : public FMaterialShader
+{
+public:
+	DECLARE_SHADER_TYPE(FTextureGraphMaterialShaderVS, Material)
+
+	BEGIN_SHADER_PARAMETER_STRUCT(FParameters, )
+	END_SHADER_PARAMETER_STRUCT()
+
+	SHADER_USE_PARAMETER_STRUCT_WITH_LEGACY_BASE(FTextureGraphMaterialShaderVS, FMaterialShader);
+
+	static bool ShouldCompilePermutation(const FMaterialShaderPermutationParameters& Parameters)
+	{
+		return IsPCPlatform(Parameters.Platform) && EnumHasAllFlags(Parameters.Flags, EShaderPermutationFlags::HasEditorOnlyData);
+	}
+
+	static void ModifyCompilationEnvironment(const FMaterialShaderPermutationParameters& Parameters, FShaderCompilerEnvironment& OutEnvironment)
+	{
+		FMaterialShader::ModifyCompilationEnvironment(Parameters, OutEnvironment);
+	}
+
+	static bool ShouldCache(EShaderPlatform Platform)
+	{
+		// Only cache the shader for the given platform if it is supported
+		return IsFeatureLevelSupported(Platform, ERHIFeatureLevel::SM5);
+	}
+
+	void SetParameters(FRHIBatchedShaderParameters& BatchedParameters, const FSceneView& View)
+	{
+		UE::Renderer::PostProcess::SetDrawRectangleParameters(BatchedParameters, this, View);
+		FMaterialShader::SetViewParameters(BatchedParameters, View, View.ViewUniformBuffer);
+	}
+};
+
+class TEXTUREGRAPHENGINE_API FTextureGraphMaterialShaderPS : public FMaterialShader
+{
+public:
+	DECLARE_SHADER_TYPE(FTextureGraphMaterialShaderPS, Material)
+
+	BEGIN_SHADER_PARAMETER_STRUCT(FParameters, )
+		SHADER_PARAMETER(FLinearColor, PSControl)
+		SHADER_PARAMETER_STRUCT(FTileInfo, TileInfo)
+	END_SHADER_PARAMETER_STRUCT()
+
+	SHADER_USE_PARAMETER_STRUCT_WITH_LEGACY_BASE(FTextureGraphMaterialShaderPS, FMaterialShader);
+
+	static FName PSCONTROL_ARG;
+
+	static bool ShouldCompilePermutation(const FMaterialShaderPermutationParameters& Parameters)
+	{
+		return IsPCPlatform(Parameters.Platform) && EnumHasAllFlags(Parameters.Flags, EShaderPermutationFlags::HasEditorOnlyData);
+	}
+
+	static void ModifyCompilationEnvironment(const FMaterialShaderPermutationParameters& Parameters, FShaderCompilerEnvironment& OutEnvironment)
+	{
+		FMaterialShader::ModifyCompilationEnvironment(Parameters, OutEnvironment);
+	}
+
+	static bool ShouldCache(EShaderPlatform Platform)
+	{
+		// Only cache the shader for the given platform if it is supported
+		return IsFeatureLevelSupported(Platform, ERHIFeatureLevel::SM5);
+	}
+
+	void SetParameters(FRHIBatchedShaderParameters& BatchedParameters, const FSceneView& View, const FMaterialRenderProxy* MaterialProxy)
+	{
+		const FMaterial& Material = MaterialProxy->GetMaterialWithFallback(View.GetFeatureLevel(), MaterialProxy);
+		FMaterialShader::SetViewParameters(BatchedParameters, View, View.ViewUniformBuffer);
+		FMaterialShader::SetParameters(BatchedParameters, MaterialProxy, Material, View);
+	}
+};
+
+class FxMaterial_DrawMaterialBase : public FxMaterial
+{
+public:
+	FxMaterial_DrawMaterialBase()
+		: FxMaterial(TEXT("VSH_TypeVSH_TypeVSH_TypeVSH_Type"), TEXT("FSH_TypeFSH_TypeFSH_TypeFSH_Type"))
+	{
+	}
+
+	FSceneView* CreateSceneView(UTextureRenderTarget2D* RenderTarget, const FIntPoint& TargetSizeXY);
+
+
+	static bool ValidateMaterialShaderMap(UMaterial* InMaterial, FShaderType* InPixelShaderType);
+};
+
+template <typename VSH_Type, typename FSH_Type>
+class FxMaterial_DrawMaterial : public FxMaterial_DrawMaterialBase
+{
+public:
+
+
+
+	// type name for the permutation domain
+	using VSHPermutationDomain = typename VSH_Type::FPermutationDomain;
+	using FSHPermutationDomain = typename FSH_Type::FPermutationDomain;
+
+	TObjectPtr<UMaterialInterface>	Material;
+
+protected:
+	VSHPermutationDomain			VSHPermDomain;		/// Vertex shader Permutation Domain value
+	FSHPermutationDomain			FSHPermDomain;		/// Fragment shader Permutation Domain value
+	typename VSH_Type::FParameters	VSHParams;			/// Params for the vertex shader
+	typename FSH_Type::FParameters	FSHParams;			/// Params for the fragment shader
+
+
+public:
+
+	FxMaterial_DrawMaterial()
+		: FxMaterial_DrawMaterialBase()
+	{
+	}
+
+	FxMaterial_DrawMaterial(UMaterialInterface* InMaterial, const VSHPermutationDomain& InVSHPermutationDomain, const FSHPermutationDomain& InFSHPermutationDomain)
+		: 	FxMaterial_DrawMaterialBase(),
+		Material(InMaterial),
+		VSHPermDomain(InVSHPermutationDomain),
+		FSHPermDomain(InFSHPermutationDomain)
+	{
+	}
+
+	virtual std::shared_ptr<FxMaterial> Clone() override
+	{
+		return std::static_pointer_cast<FxMaterial>(std::make_shared<FxMaterial_DrawMaterial<VSH_Type, FSH_Type>>(Material, VSHPermDomain, FSHPermDomain));
+	}
+
+	virtual FxMetadataSet GetMetadata() const override
+	{
+		return {
+			{ VSH_Type::FParameters::FTypeInfo::GetStructMetadata(), (char*)&VSHParams }, 
+			{ FSH_Type::FParameters::FTypeInfo::GetStructMetadata(), (char*)&FSHParams }
+		};
+	}
+
+	virtual void Blit(FRHICommandListImmediate& RHI, FRHITexture2D* Target, const RenderMesh* MeshObj, int32 TargetId, FGraphicsPipelineStateInitializer* InPSO = nullptr) override
+	{
+		
+	}
+
+	void MyBlit(FRHICommandListImmediate& RHI, UTextureRenderTarget2D* RenderTarget, FRHITexture2D* Target, const RenderMesh* MeshObj, int32 TargetId, FGraphicsPipelineStateInitializer* InPSO = nullptr)
+	{
+		BindTexturesForBlitting();
+
+		// Make a Scene view for this pass
+		FSceneView* SceneView = CreateSceneView(RenderTarget, Target->GetSizeXY());
+
+		//const typename FSH_Type::FParameters* params = reinterpret_cast<typename FSH_Type::FParameters*>(params_);
+		QUICK_SCOPE_CYCLE_COUNTER(STAT_ShaderPlugin_PixelShader); // Used to gather CPU profiling data for the UE4 session frontend
+		SCOPED_DRAW_EVENT(RHI, ShaderPlugin_Pixel); // Used to profile GPU activity and add metadata to be consumed by for example RenderDoc
+
+
+		// Pixel shader combined with Material
+		const FMaterialRenderProxy* MaterialProxy = Material->GetRenderProxy();
+		MaterialProxy->UpdateUniformExpressionCacheIfNeeded(GMaxRHIFeatureLevel);
+		FMaterialRenderProxy::UpdateDeferredCachedUniformExpressions();
+
+		const FMaterial& RenderMaterial = MaterialProxy->GetMaterialWithFallback(GMaxRHIFeatureLevel, MaterialProxy);
+		const FMaterialShaderMap* const MaterialShaderMap = RenderMaterial.GetRenderingThreadShaderMap();
+
+		TShaderRef<VSH_Type> VertexShader = MaterialShaderMap->GetShader<VSH_Type>();
+		FRHIVertexShader* RHIVertexShader = VertexShader.GetVertexShader();
+		TShaderRef<FSH_Type> PixelShader = MaterialShaderMap->GetShader<FSH_Type>();
+		FRHIPixelShader* RHIPixelShader = PixelShader.GetPixelShader();
+
+
+		FRHIRenderPassInfo passInfo(Target, ERenderTargetActions::Clear_Store);
+		RHI.BeginRenderPass(passInfo, TEXT("FxMaterial_UMaterial"));
+
+	
+		// Set the graphic pipeline state.
+		FGraphicsPipelineStateInitializer PSO;
+		
+		InitPSO_Default(PSO, RHIVertexShader, RHIPixelShader);
+
+		RHI.ApplyCachedRenderTargets(PSO);
+
+		SetGraphicsPipelineState(RHI, PSO, 0);
+		SetShaderParametersLegacyVS(RHI, VertexShader, *SceneView);
+		SetShaderParametersLegacyPS(RHI, PixelShader, *SceneView, MaterialProxy);
+		
+		SetShaderParameters(RHI, VertexShader, RHIVertexShader, VSHParams);
+		SetShaderParameters(RHI, PixelShader, RHIPixelShader, FSHParams);
+
+		RHI.SetStreamSource(0, GQuadBuffer.VertexBufferRHI, 0);
+		RHI.DrawPrimitive(0, 2, 1);
+
+		RHI.EndRenderPass();
+	}
+
+	static bool ValidateMaterial(UMaterialInterface* InMaterial)
+	{
+		check(IsInRenderingThread());
+
+		UMaterial* Material = InMaterial->GetMaterial();
+
+		return ValidateMaterialShaderMap(Material, &FSH_Type::GetStaticType());
+	}
+};
+
+typedef FxMaterial_DrawMaterial< FTextureGraphMaterialShaderVS, FTextureGraphMaterialShaderPS> FxMaterial_QuadDrawMaterial;
+typedef std::shared_ptr<FxMaterial_QuadDrawMaterial>	FxMaterialQuadDrawMaterialPtr;
+
+
diff --git a/Engine/Plugins/Experimental/TextureGraph/Source/TextureGraphEngine/FxMat/MaterialManager.cpp b/Engine/Plugins/Experimental/TextureGraph/Source/TextureGraphEngine/FxMat/MaterialManager.cpp
new file mode 100644
--- /dev/null
+++ b/Engine/Plugins/Experimental/TextureGraph/Source/TextureGraphEngine/FxMat/MaterialManager.cpp
@@ -0,0 +1,69 @@
+// Copyright Epic Games, Inc. All Rights Reserved.
+#include "MaterialManager.h"
+#include "Model/ModelObject.h"
+#include "2D/Tex.h"
+#include "Engine/Canvas.h"
+#include "Materials/Material.h"
+
+UMaterialManager::UMaterialManager()
+{
+}
+
+UMaterialManager::~UMaterialManager()
+{
+}
+
+void UMaterialManager::InitEssentialMaterials()
+{
+	LoadMaterial(TEXT("Util/CopyTexture"));
+}
+
+UMaterial* UMaterialManager::LoadMaterial(const FString& Path)
+{
+	TObjectPtr<UMaterial>* ExistingMat = Materials.Find(Path);
+
+	if (ExistingMat)
+		return ExistingMat->Get();
+
+	check(IsInGameThread());
+
+	FString FullPath = TEXT("Materials/") + Path;
+	UMaterial* Mat = Cast<UMaterial>(LoadObjectFromPath(FullPath));
+
+	if (Mat)
+	{
+		Materials.Add(FullPath, Mat);
+		Materials.Add(Path, Mat);
+	}
+
+	return Mat;
+}
+
+RenderMaterial_BPPtr UMaterialManager::CreateMaterial_BP(FString Name, const FString& Path)
+{
+	UMaterial* Mat = LoadMaterial(Path);
+	return CreateMaterial_BP(Name, Mat);
+}
+
+RenderMaterial_BPPtr UMaterialManager::CreateMaterial_BP(FString InName, UMaterial* InMaterial)
+{
+	return std::make_shared<RenderMaterial_BP>(InName, InMaterial, nullptr);
+}
+
+RenderMaterial_BP_NoTilePtr UMaterialManager::CreateMaterial_BP_NoTile(FString Name, const FString& Path)
+{
+	UMaterial* Mat = LoadMaterial(Path);
+	return std::make_shared<RenderMaterial_BP_NoTile>(Name, Mat, nullptr);
+}
+
+RenderMaterial_ThumbPtr UMaterialManager::CreateMaterial_Thumbnail(FString Name, const FString& Path)
+{
+	UMaterial* Mat = LoadMaterial(Path);
+	return std::make_shared<RenderMaterial_Thumbnail>(Name, Mat, nullptr);
+}
+
+RenderMaterial_BP_TileArgsPtr UMaterialManager::CreateMaterialOfType_BP_TileArgs(FString Name, const FString& Path)
+{
+	UMaterial* Mat = LoadMaterial(Path);
+	return std::make_shared<RenderMaterial_BP_TileArgs>(Name, Mat, nullptr);
+}
diff --git a/Engine/Plugins/Experimental/TextureGraph/Source/TextureGraphEngine/FxMat/RenderMaterial.h b/Engine/Plugins/Experimental/TextureGraph/Source/TextureGraphEngine/FxMat/RenderMaterial.h
new file mode 100644
--- /dev/null
+++ b/Engine/Plugins/Experimental/TextureGraph/Source/TextureGraphEngine/FxMat/RenderMaterial.h
@@ -0,0 +1,54 @@
+// Copyright Epic Games, Inc. All Rights Reserved.
+#pragma once
+
+#include "CoreMinimal.h"
+#include "Transform/BlobTransform.h"
+#include "Engine/TextureRenderTarget2D.h"
+
+
+class Tex;
+class Device;
+class RenderMesh;
+class FRHICommandListImmediate;
+class UTextureRenderTarget2D;
+class UTexture;
+
+class TEXTUREGRAPHENGINE_API RenderMaterial : public BlobTransform
+{
+protected:
+	bool							bCanHandleTiles = true;		/// Whether this material can handle tiles or not
+
+public:
+	explicit						RenderMaterial(FString Name);
+	virtual							~RenderMaterial() override;
+
+	virtual void					BlitTo(FRHICommandListImmediate& RHI, UTextureRenderTarget2D* Dst, const RenderMesh* MeshObj, int32 TargetId) const = 0;
+	void							BlitTo(FRHICommandListImmediate& RHI, UTextureRenderTarget2D* Dst) const;
+
+	//////////////////////////////////////////////////////////////////////////
+	/// BlobTransform overrides
+	//////////////////////////////////////////////////////////////////////////
+	virtual Device*					TargetDevice(size_t Index) const override;
+	virtual AsyncTransformResultPtr	Exec(const TransformArgs& Args) override;
+	virtual ENamedThreads::Type		ExecutionThread() const override;
+	virtual bool					CanHandleTiles() const override;
+
+	virtual void					SetTexture(FName Name, const UTexture* Texture) const = 0;
+	virtual void					SetArrayTexture(FName Name, const std::vector<const UTexture*>& Textures) const = 0;
+	virtual void					SetInt(FName Name, int32 Value) const = 0;
+	virtual void					SetFloat(FName Name, float Value) const = 0;
+	virtual void					SetColor(FName Name, const FLinearColor& Value) const = 0;
+	virtual void					SetIntVector4(FName Name, const FIntVector4& Value) const = 0;
+	virtual void					SetMatrix(FName Name, const FMatrix& Value) const = 0;
+	void							SetSourceTexture(const Tex* Texture) const;
+	void							SetTexture(FName InName, std::shared_ptr<Tex> TexObj) const;
+	void							SetArrayTexture(FName InName, const std::vector<std::shared_ptr<Tex>>& TexObj) const;
+	void							BlitTo(FRHICommandListImmediate& RHI, Tex* Dst, const RenderMesh* MeshObj = nullptr, int32 TargetId = -1) const;
+	//////////////////////////////////////////////////////////////////////////
+	/// Inline functions
+	//////////////////////////////////////////////////////////////////////////
+	FORCEINLINE void				SetSourceTexture(const UTexture* Texture) const { SetTexture(FName("SourceTexture"), Texture); }
+	FORCEINLINE bool&				CanHandleTiles() { return bCanHandleTiles; }
+};
+
+typedef std::shared_ptr<RenderMaterial>	RenderMaterialPtr;
diff --git a/Engine/Plugins/Experimental/TextureGraph/Source/TextureGraphEngine/FxMat/RenderMaterial_BP.h b/Engine/Plugins/Experimental/TextureGraph/Source/TextureGraphEngine/FxMat/RenderMaterial_BP.h
new file mode 100644
--- /dev/null
+++ b/Engine/Plugins/Experimental/TextureGraph/Source/TextureGraphEngine/FxMat/RenderMaterial_BP.h
@@ -0,0 +1,72 @@
+// Copyright Epic Games, Inc. All Rights Reserved.
+#pragma once
+
+#include "CoreMinimal.h"
+#include "RenderMaterial.h"
+#include "Materials/MaterialInstanceDynamic.h"
+#include "UObject/StrongObjectPtr.h"
+
+class FxMaterial;
+typedef std::shared_ptr< FxMaterial> FxMaterialPtr;
+
+class TEXTUREGRAPHENGINE_API RenderMaterial_BP : public RenderMaterial
+{
+protected:
+	UMaterial*						Material = nullptr;			/// The base material that is part 
+	TStrongObjectPtr<UMaterialInstanceDynamic> 
+									MaterialInstance;			/// An instance of the material
+	CHashPtr						HashValue;					/// The hash for this material
+	bool							RequestMaterialValidation = true;		/// 
+	bool							MaterialInstanceValidated = false;		/// 
+
+	FxMaterialPtr					FXMaterialObj;
+
+	UCanvas*						Canvas = nullptr;
+
+	void							DrawMaterial(UMaterialInterface* RenderMaterial, FVector2D ScreenPosition, FVector2D ScreenSize,
+									FVector2D CoordinatePosition, FVector2D CoordinateSize=FVector2D::UnitVector, float Rotation=0.f,
+									FVector2D PivotPoint=FVector2D(0.5f,0.5f)) const;
+public:
+									RenderMaterial_BP(FString Name, UMaterial* InMaterial, UMaterialInstanceDynamic* InMaterialInstance = nullptr);
+
+	virtual							~RenderMaterial_BP() override;
+
+	static bool						ValidateMaterialCompatible(UMaterialInterface* InMaterial);
+
+
+	virtual AsyncPrepareResult		PrepareResources(const TransformArgs& Args) override;
+
+	//////////////////////////////////////////////////////////////////////////
+	/// BlobTransform Implementation
+	//////////////////////////////////////////////////////////////////////////
+	virtual void					Bind(int32 Value, const ResourceBindInfo& BindInfo) override;
+	virtual void					Bind(float Value, const ResourceBindInfo& BindInfo) override;
+	virtual void					Bind(const FLinearColor& Value, const ResourceBindInfo& BindInfo) override;
+	virtual void					Bind(const FIntVector4& Value, const ResourceBindInfo& BindInfo) override;
+	virtual void					Bind(const FMatrix& Value, const ResourceBindInfo& BindInfo) override;
+	virtual void					BindStruct(const char* ValueAddress, size_t StructSize, const ResourceBindInfo& BindInfo) override;
+	virtual CHashPtr				Hash() const override;
+	virtual std::shared_ptr<BlobTransform> DuplicateInstance(FString InName) override;
+
+	//////////////////////////////////////////////////////////////////////////
+	/// RenderMaterial Implementation
+	//////////////////////////////////////////////////////////////////////////
+	virtual void					BlitTo(FRHICommandListImmediate& RHI, UTextureRenderTarget2D* DstRT, const RenderMesh* MeshObj, int32 TargetId) const override;
+	virtual void					SetTexture(FName InName, const UTexture* Texture) const override;
+	virtual void					SetArrayTexture(FName InName, const std::vector<const UTexture*>& Textures) const override;
+	virtual void					SetInt(FName InName, int32 Value) const override;
+	virtual void					SetFloat(FName InName, float Value) const override;
+	virtual void					SetColor(FName InName, const FLinearColor& Value) const override;
+	virtual void					SetIntVector4(FName InName, const FIntVector4& Value) const override;
+	virtual void					SetMatrix(FName InName, const FMatrix& Value) const override;
+	
+	//////////////////////////////////////////////////////////////////////////
+	/// Inline functions
+	////////////////////////////////////////////////////////////////////////// 
+	FORCEINLINE UMaterial*			GetMaterial() { return Material; }
+
+	FORCEINLINE UMaterialInstanceDynamic* Instance() { return MaterialInstance.Get(); }
+	FORCEINLINE const UMaterialInstanceDynamic* Instance() const { return MaterialInstance.Get(); }
+};
+
+typedef std::shared_ptr<RenderMaterial_BP> RenderMaterial_BPPtr;
\ No newline at end of file
diff --git a/Engine/Plugins/Experimental/TextureGraph/Source/TextureGraphEngine/FxMat/RenderMaterial_BP_NoTile.h b/Engine/Plugins/Experimental/TextureGraph/Source/TextureGraphEngine/FxMat/RenderMaterial_BP_NoTile.h
new file mode 100644
--- /dev/null
+++ b/Engine/Plugins/Experimental/TextureGraph/Source/TextureGraphEngine/FxMat/RenderMaterial_BP_NoTile.h
@@ -0,0 +1,33 @@
+// Copyright Epic Games, Inc. All Rights Reserved.
+#pragma once
+
+#include "CoreMinimal.h"
+#include "RenderMaterial_BP.h"
+
+/**
+ * 
+ */
+class TEXTUREGRAPHENGINE_API RenderMaterial_BP_NoTile : public RenderMaterial_BP
+{
+
+public:
+	RenderMaterial_BP_NoTile(FString InName, UMaterial* InMaterial, UMaterialInstanceDynamic* InMaterialInstance = nullptr) : RenderMaterial_BP(InName, InMaterial, InMaterialInstance) {}
+	virtual bool					CanHandleTiles() const override { return false; };
+};
+
+typedef std::shared_ptr<RenderMaterial_BP_NoTile> RenderMaterial_BP_NoTilePtr;
+
+class TEXTUREGRAPHENGINE_API RenderMaterial_BP_TileArgs : public RenderMaterial_BP
+{
+private:
+	FName							MemberName = "TileInfo";
+
+public:
+	RenderMaterial_BP_TileArgs(FString InName, UMaterial* InMaterial, UMaterialInstanceDynamic* InMaterialInstance = nullptr) : RenderMaterial_BP(InName, InMaterial, InMaterialInstance) {}
+	void							AddTileArgs(TransformArgs& Args);
+	virtual AsyncPrepareResult		PrepareResources(const TransformArgs& Args) override;
+	virtual bool					CanHandleTiles() const override { return true; };
+	virtual std::shared_ptr<BlobTransform> DuplicateInstance(FString NewName) override;
+};
+
+typedef std::shared_ptr<RenderMaterial_BP_TileArgs> RenderMaterial_BP_TileArgsPtr;
\ No newline at end of file
diff --git a/Engine/Plugins/Experimental/TextureGraph/Source/TextureGraphEngine/FxMat/RenderMaterial_FX.h b/Engine/Plugins/Experimental/TextureGraph/Source/TextureGraphEngine/FxMat/RenderMaterial_FX.h
new file mode 100644
--- /dev/null
+++ b/Engine/Plugins/Experimental/TextureGraph/Source/TextureGraphEngine/FxMat/RenderMaterial_FX.h
@@ -0,0 +1,55 @@
+// Copyright Epic Games, Inc. All Rights Reserved.
+#pragma once
+
+#include "CoreMinimal.h"
+#include "RenderMaterial.h"
+#include "SamplerStates_FX.h"
+
+class FxMaterial;
+typedef std::shared_ptr<FxMaterial>	FxMaterialPtr;
+
+class TEXTUREGRAPHENGINE_API RenderMaterial_FX : public RenderMaterial
+{
+protected:
+	mutable FxMaterialPtr			FXMaterial;			/// The FX material from the shader
+	CHashPtr						HashValue;			/// The hash for the material
+	
+public:
+									RenderMaterial_FX(FString InName, FxMaterialPtr InFXMaterial);
+	virtual							~RenderMaterial_FX() override;
+
+	//////////////////////////////////////////////////////////////////////////
+	/// BlobTransform implementation
+	//////////////////////////////////////////////////////////////////////////
+	virtual void					Bind(int32 Value, const ResourceBindInfo& BindInfo) override;
+	virtual void					Bind(float Value, const ResourceBindInfo& BindInfo) override;
+	virtual void					Bind(const FLinearColor& Value, const ResourceBindInfo& BindInfo) override;
+	virtual void					Bind(const FIntVector4& Value, const ResourceBindInfo& BindInfo) override;
+	virtual void					Bind(const FMatrix& Value, const ResourceBindInfo& BindInfo) override;
+	virtual void					BindStruct(const char* ValueAddress, size_t StructSize, const ResourceBindInfo& BindInfo) override;
+	virtual void					BindScalarArray(const char* StartingAddress, size_t TypeSize, size_t ArraySize, const ResourceBindInfo& BindInfo) override;
+
+	virtual CHashPtr				Hash() const override;
+	virtual std::shared_ptr<BlobTransform> DuplicateInstance(FString NewName) override;
+
+	//////////////////////////////////////////////////////////////////////////
+	/// RenderMaterial Implementation
+	//////////////////////////////////////////////////////////////////////////
+	virtual void					BlitTo(FRHICommandListImmediate& RHI, UTextureRenderTarget2D* RenderTarget, const RenderMesh* MeshObj, int32 TargetId) const override;
+	virtual void					SetTexture(FName InName, const UTexture* Texture) const override;
+	virtual void					SetArrayTexture(FName InName, const std::vector<const UTexture*>& Textures) const override;
+	virtual void					SetInt(FName InName, int32 Value) const override;
+	virtual void					SetFloat(FName InName, float Value) const override;
+	virtual void					SetColor(FName InName, const FLinearColor& Value) const override;
+	virtual void					SetIntVector4(FName InName, const FIntVector4& Value) const override;
+	virtual void					SetMatrix(FName InName, const FMatrix& Value) const override;
+
+	//////////////////////////////////////////////////////////////////////////
+	/// Inline functions
+	//////////////////////////////////////////////////////////////////////////
+	FORCEINLINE FxMaterialPtr		FxMaterial() const { return FXMaterial; }
+};
+
+
+
+typedef std::shared_ptr<RenderMaterial_FX> RenderMaterial_FXPtr;
diff --git a/Engine/Plugins/Experimental/TextureGraph/Source/TextureGraphEngine/FxMat/RenderMaterial_FX_Combined.h b/Engine/Plugins/Experimental/TextureGraph/Source/TextureGraphEngine/FxMat/RenderMaterial_FX_Combined.h
new file mode 100644
--- /dev/null
+++ b/Engine/Plugins/Experimental/TextureGraph/Source/TextureGraphEngine/FxMat/RenderMaterial_FX_Combined.h
@@ -0,0 +1,34 @@
+// Copyright Epic Games, Inc. All Rights Reserved.
+#pragma once
+
+#include "CoreMinimal.h"
+#include "RenderMaterial_FX.h"
+#include "Data/TiledBlob.h"
+
+DECLARE_LOG_CATEGORY_EXTERN(LogRenderMaterial_FX_NoTile, All, All);
+
+
+
+
+class TEXTUREGRAPHENGINE_API RenderMaterial_FX_Combined : public RenderMaterial_FX
+{
+private:
+	
+	TArray<TiledBlobPtr>					ToBeCombined;		/// To save under process blobs. This array ensures that multiple SRVs can be passed to the shader
+
+public:
+	RenderMaterial_FX_Combined(FString InName, FxMaterialPtr InFXMaterial);
+	virtual									~RenderMaterial_FX_Combined() override;
+
+	void									AddTileArgs(TransformArgs& Args);
+	virtual AsyncPrepareResult				PrepareResources(const TransformArgs& Args) override;
+	virtual bool							CanHandleTiles() const override { return true; };
+	virtual std::shared_ptr<BlobTransform>	DuplicateInstance(FString NewName) override;
+	virtual AsyncTransformResultPtr			Exec(const TransformArgs& Args) override;
+	TiledBlobPtr							AddBlobToCombine(TiledBlobPtr BlobToCombine);
+	void									FreeCombinedBlob();
+
+	static const FName						MemberName; /// Default struct parameter for tile
+};
+
+typedef std::shared_ptr<RenderMaterial_FX_Combined> RenderMaterial_FX_CombinedPtr;
diff --git a/Engine/Plugins/Experimental/TextureGraph/Source/TextureGraphEngine/FxMat/RenderMaterial_FX_MinMax.h b/Engine/Plugins/Experimental/TextureGraph/Source/TextureGraphEngine/FxMat/RenderMaterial_FX_MinMax.h
new file mode 100644
--- /dev/null
+++ b/Engine/Plugins/Experimental/TextureGraph/Source/TextureGraphEngine/FxMat/RenderMaterial_FX_MinMax.h
@@ -0,0 +1,31 @@
+// Copyright Epic Games, Inc. All Rights Reserved.
+#pragma once
+
+#include "CoreMinimal.h"
+#include "RenderMaterial_FX.h"
+#include "Data/TiledBlob.h"
+
+DECLARE_LOG_CATEGORY_EXTERN(LogRenderMaterial_FX_MinMax, All, All);
+
+class Tex;
+typedef std::shared_ptr<Tex>		TexPtr;
+
+class TEXTUREGRAPHENGINE_API RenderMaterial_FX_MinMax : public RenderMaterial_FX
+{
+private:
+	
+	FxMaterialPtr					SecondPassMaterial;
+	BufferDescriptor				SourceDesc;
+	TArray<TexPtr>					DownsampledResultTargets;		/// used for creating intermediate targets.
+	
+public:
+	RenderMaterial_FX_MinMax(FString InName, FxMaterialPtr InMaterial, FxMaterialPtr InSecondPassMaterial);
+	virtual ~RenderMaterial_FX_MinMax() override;
+
+	virtual AsyncPrepareResult		PrepareResources(const TransformArgs& Args) override;
+	virtual AsyncTransformResultPtr	Exec(const TransformArgs& Args) override;
+	
+	void							SetDescriptor(const BufferDescriptor& InDesc) { SourceDesc = InDesc;}
+};
+
+typedef std::shared_ptr<RenderMaterial_FX_MinMax> RenderMaterial_FX_MinMaxPtr;
diff --git a/Engine/Plugins/Experimental/TextureGraph/Source/TextureGraphEngine/FxMat/RenderMaterial_Thumbnail.h b/Engine/Plugins/Experimental/TextureGraph/Source/TextureGraphEngine/FxMat/RenderMaterial_Thumbnail.h
new file mode 100644
--- /dev/null
+++ b/Engine/Plugins/Experimental/TextureGraph/Source/TextureGraphEngine/FxMat/RenderMaterial_Thumbnail.h
@@ -0,0 +1,28 @@
+// Copyright Epic Games, Inc. All Rights Reserved.
+#pragma once
+
+#include "CoreMinimal.h"
+#include "RenderMaterial.h"
+#include "UObject/StrongObjectPtr.h"
+#include "RenderMaterial_BP.h"
+
+// Subset of RenderMaterial BP, This class will always render a thumbnail from its material instance
+class TEXTUREGRAPHENGINE_API RenderMaterial_Thumbnail : public RenderMaterial_BP
+{
+public:
+	static constexpr uint32	GThumbWidth = 128;			/// Width of standard thumbnail image
+	static constexpr uint32	GThumbHeight = 128;		/// Height of the standard thumbnail image
+
+private:
+	ERHIFeatureLevel::Type FeatureLevel;
+
+public:
+	RenderMaterial_Thumbnail(FString InName, UMaterial* InMaterial, UMaterialInstanceDynamic* instance = nullptr);
+
+	virtual void BlitTo(FRHICommandListImmediate& RHI, UTextureRenderTarget2D* DstRT, const RenderMesh* MeshObj, int32 TargetId) const override;
+	virtual bool CanHandleTiles() const override { return true; };
+	virtual std::shared_ptr<BlobTransform> DuplicateInstance(FString InName) override;
+	virtual AsyncPrepareResult PrepareResources(const TransformArgs& Args) override;
+};
+
+typedef std::shared_ptr<RenderMaterial_Thumbnail> RenderMaterial_ThumbPtr;
\ No newline at end of file
diff --git a/Engine/Plugins/Experimental/TextureGraph/Source/TextureGraphEngine/Helper/ColorUtil.h b/Engine/Plugins/Experimental/TextureGraph/Source/TextureGraphEngine/Helper/ColorUtil.h
new file mode 100644
--- /dev/null
+++ b/Engine/Plugins/Experimental/TextureGraph/Source/TextureGraphEngine/Helper/ColorUtil.h
@@ -0,0 +1,37 @@
+// Copyright Epic Games, Inc. All Rights Reserved.
+#pragma once
+
+#include "CoreMinimal.h"
+#include "2D/TextureType.h"
+
+class TEXTUREGRAPHENGINE_API ColorUtil
+{
+public:
+	//Converting from Array to Functions to avoid static fiasco problem with CPP
+	static const FLinearColor			DefaultColor(TextureType Type);
+	static const FLinearColor			DefaultDiffuse();
+	static const FLinearColor			DefaultSpecular();
+	static const FLinearColor			DefaultAlbedo();
+	static const FLinearColor			DefaultMetalness();
+	static const FLinearColor			DefaultNormal();
+	static const FLinearColor			DefaultDisplacement();
+	static const FLinearColor			DefaultOpacity();
+	static const FLinearColor			DefaultRoughness();
+	static const FLinearColor			DefaultAO();
+	static const FLinearColor			DefaultCurvature();
+	static const FLinearColor			DefaultPreview();
+
+	static float						GetHue(const FColor& Color);
+	static float						GetSquaredDistance(FLinearColor Current, FLinearColor Match);
+	static FString						GetColorName(FLinearColor Color);
+
+	static bool							IsColorBlack(const FLinearColor& Color, bool IgnoreAlpha = true);
+	static bool							IsColorWhite(const FLinearColor& Color, bool IgnoreAlpha = true);
+	static bool							IsColorGray(const FLinearColor& Color, bool IgnoreAlpha = true);
+	static bool							IsColorRed(const FLinearColor& Color, bool IgnoreAlpha = true);
+	static bool							IsColorGreen(const FLinearColor& Color, bool IgnoreAlpha = true);
+	static bool							IsColorBlue(const FLinearColor& Color, bool IgnoreAlpha = true);
+	static bool							IsColorYellow(const FLinearColor& Color, bool IgnoreAlpha = true);
+	static bool							IsColorMagenta(const FLinearColor& Color, bool IgnoreAlpha = true);
+	static bool							IsColorNear(const FLinearColor& Color, const FLinearColor& Ref, bool IgnoreAlpha = true);
+};
diff --git a/Engine/Plugins/Experimental/TextureGraph/Source/TextureGraphEngine/Helper/GraphicsUtil.h b/Engine/Plugins/Experimental/TextureGraph/Source/TextureGraphEngine/Helper/GraphicsUtil.h
new file mode 100644
--- /dev/null
+++ b/Engine/Plugins/Experimental/TextureGraph/Source/TextureGraphEngine/Helper/GraphicsUtil.h
@@ -0,0 +1,59 @@
+// Copyright Epic Games, Inc. All Rights Reserved.
+#pragma once
+
+#include "CoreMinimal.h"
+#include "2D/TextureType.h"
+#include "GraphicsDefs.h"
+#include <array>
+
+enum class BoundsCorners
+{
+	L0_BottomLeft,
+	L0_TopLeft,
+	L0_TopRight,
+	L0_BottomRight,
+
+	L1_BottomLeft,
+	L1_TopLeft,
+	L1_TopRight,
+	L1_BottomRight
+};
+
+class TEXTUREGRAPHENGINE_API GraphicsUtil
+{
+public:
+	static Vector2				CalculateBarycentricCoordinates(const Vector2* UV, const Vector2& P);
+	static Vector2				CalculateBarycentricCoordinates(const Vector3& A, const Vector3& B, const Vector3& C, const Vector3& P);
+
+	static bool					IsPointInTriangle(const Vector3& A, const Vector3& B, const Vector3& C, const Vector3& P);
+	static bool					IsPointInTriangle(float u, float v, float tolerance = 0.001f);
+	static bool					IsPointInTriangle(const Vector2& uv, float tolerance = 0.001f);
+
+	//////////////////////////////////////////////////////////////////////////
+	static Vector3				s_boundsAxis[6];
+	static float				SquaredDistance(const Vector3& p0, const Vector3& p1);
+	static float				SquaredLength(const Vector3& pt);
+
+	static std::array<Vector3, 8> GetBoundsCorners(const FBox& bounds);
+
+	static Vector3				GetBoundsPoint(const FBox& bounds, BoundsCorners pti);
+	static FBox					GetChildBounds_QuadTree(const FBox& bounds, BoundsCorners ci);
+	static FBox					GetChildBounds_Octree(const FBox& bounds, BoundsCorners ci);
+	static bool					CheckRayTriangleIntersection(const FRay& ray, const Vector3& p1, const Vector3& p2, const Vector3& p3, float& distance, Vector3& ip);
+
+	static FORCEINLINE Vector4	ConvertTangent(const Tangent& t) { return Vector4(t.TangentX.X, t.TangentX.Y, t.TangentX.Z, t.bFlipTangentY ? -1.0f : 1.0f); }
+	static FORCEINLINE FLinearColor TangentAsLinearColor(const Tangent& t) { return FLinearColor(ConvertTangent(t)); }
+	static FORCEINLINE FLinearColor Vec2AsLinearColor(const Vector2& v) { return FLinearColor(v.X, v.Y, 0.0f, 0.0f); }
+	static FORCEINLINE FLinearColor Vec3AsLinearColor(const Vector3& v) { return FLinearColor(v.X, v.Y, v.Z, 0.0f); }
+	static FORCEINLINE FLinearColor Vec4AsLinearColor(const Vector4& v) { return FLinearColor(v); }
+
+	//////////////////////////////////////////////////////////////////////////
+	template <typename VecType>
+	static VecType				Interpolate(const VecType& v0, const VecType& v1, const VecType& v2, Vector2 barycentric)
+	{
+		float u = barycentric.X;
+		float v = barycentric.Y;
+
+		return v0 + u * (v1 - v0) + v * (v2 - v0);
+	}
+};
diff --git a/Engine/Plugins/Experimental/TextureGraph/Source/TextureGraphEngine/Helper/MathUtils.h b/Engine/Plugins/Experimental/TextureGraph/Source/TextureGraphEngine/Helper/MathUtils.h
new file mode 100644
--- /dev/null
+++ b/Engine/Plugins/Experimental/TextureGraph/Source/TextureGraphEngine/Helper/MathUtils.h
@@ -0,0 +1,26 @@
+// Copyright Epic Games, Inc. All Rights Reserved.
+#pragma once
+
+#include "CoreMinimal.h"
+/**
+ * 
+ */
+class TEXTUREGRAPHENGINE_API MathUtils
+{
+public:
+	MathUtils() = delete;
+
+	static const float					GMeterToCm;
+
+	static float						MinFloat(); 
+	static float						MaxFloat(); 
+	
+	static FVector						MinFVector();
+	static FVector						MaxFVector();
+	static FVector2f					MinFVector2();
+	static FVector2f					MaxFVector2();
+	static void						UpdateBounds(FBox& bounds, const FVector& point);
+	static void						EncapsulateBound(FBox& bounds, FBox& otherBounds);
+	static FVector						GetDirection(float yzAngle, float xAngle, int xSign = 1);
+	static FBox						GetCombinedBounds(TArray<FBox> inputBounds);
+};
diff --git a/Engine/Plugins/Experimental/TextureGraph/Source/TextureGraphEngine/Job/BlobHasherService.h b/Engine/Plugins/Experimental/TextureGraph/Source/TextureGraphEngine/Job/BlobHasherService.h
new file mode 100644
--- /dev/null
+++ b/Engine/Plugins/Experimental/TextureGraph/Source/TextureGraphEngine/Job/BlobHasherService.h
@@ -0,0 +1,19 @@
+// Copyright Epic Games, Inc. All Rights Reserved.
+#pragma once
+
+#include "CoreMinimal.h"
+#include "BlobHelperService.h"
+#include <vector>
+
+class TEXTUREGRAPHENGINE_API BlobHasherService : public BlobHelperService
+{
+public:
+									BlobHasherService();
+	virtual							~BlobHasherService() override;
+
+	virtual AsyncJobResultPtr		Tick() override;
+	virtual void					Add(BlobRef BlobObj) override;
+};
+
+typedef std::shared_ptr<BlobHasherService>	BlobHasherServicePtr;
+typedef std::weak_ptr<BlobHasherService>	BlobHasherServicePtrW;
diff --git a/Engine/Plugins/Experimental/TextureGraph/Source/TextureGraphEngine/Job/BlobHelperService.h b/Engine/Plugins/Experimental/TextureGraph/Source/TextureGraphEngine/Job/BlobHelperService.h
new file mode 100644
--- /dev/null
+++ b/Engine/Plugins/Experimental/TextureGraph/Source/TextureGraphEngine/Job/BlobHelperService.h
@@ -0,0 +1,20 @@
+// Copyright Epic Games, Inc. All Rights Reserved.
+#pragma once
+
+#include "CoreMinimal.h"
+#include "IdleService.h"
+#include <vector>
+
+class TEXTUREGRAPHENGINE_API BlobHelperService : public IdleService
+{
+protected:
+	std::vector<BlobPtr>			Blobs;					/// The blobs that we want to hash
+
+public:
+	explicit						BlobHelperService(const FString& InName);
+	virtual							~BlobHelperService() override;
+
+	virtual void					Stop() override;
+
+	virtual void					Add(BlobRef BlobObj);
+};
diff --git a/Engine/Plugins/Experimental/TextureGraph/Source/TextureGraphEngine/Job/HistogramService.h b/Engine/Plugins/Experimental/TextureGraph/Source/TextureGraphEngine/Job/HistogramService.h
new file mode 100644
--- /dev/null
+++ b/Engine/Plugins/Experimental/TextureGraph/Source/TextureGraphEngine/Job/HistogramService.h
@@ -0,0 +1,33 @@
+// Copyright Epic Games, Inc. All Rights Reserved.
+#pragma once
+
+#include "CoreMinimal.h"
+#include "IdleService.h"
+#include "Model/Mix/MixUpdateCycle.h"
+
+class UModelObject;
+typedef std::shared_ptr<class Job>		JobPtr;
+typedef std::unique_ptr<class Job>		JobUPtr;
+typedef std::weak_ptr<class Job>		JobPtrW;
+
+class UMixInterface;
+
+class TEXTUREGRAPHENGINE_API HistogramService : public IdleService
+{
+private:
+	JobBatchPtr Batch;
+
+public:
+	HistogramService();
+	virtual	~HistogramService() override;
+
+	virtual AsyncJobResultPtr Tick() override;
+	virtual void Stop() override;
+
+	void AddHistogramJob(MixUpdateCyclePtr Cycle,JobUPtr JobToAdd, int32 TargetID, UMixInterface* Mix);
+	JobBatchPtr GetOrCreateNewBatch(UMixInterface* mix);
+};
+
+typedef std::shared_ptr<HistogramService>	HistogramServicePtr;
+typedef std::weak_ptr<HistogramService>	HistogramServicePtrW;
+
