Rename MixerEngine and other files with the Mixer branding in them.

diff --git a/Engine/Plugins/Experimental/TextureGraph/Source/TextureGraphEngine/Job/Job.h b/Engine/Plugins/Experimental/TextureGraph/Source/TextureGraphEngine/Job/Job.h
new file mode 100644
--- /dev/null
+++ b/Engine/Plugins/Experimental/TextureGraph/Source/TextureGraphEngine/Job/Job.h
@@ -0,0 +1,219 @@
+// Copyright Epic Games, Inc. All Rights Reserved.
+#pragma once
+
+#include "CoreMinimal.h"
+#include "Helper/Promise.h"
+#include "Device/DeviceNativeTask.h"
+#include "JobCommon.h"
+#include "JobArgs.h"
+#include "Data/TiledBlob.h"
+
+class Job;
+class Scheduler;
+class JobBatch;
+class MixUpdateCycle;
+
+DECLARE_LOG_CATEGORY_EXTERN(LogJob, Log, Verbose);
+
+class RenderMesh;
+
+//////////////////////////////////////////////////////////////////////////
+class TEXTUREGRAPHENGINE_API Job : public DeviceNativeTask
+{
+public:
+	struct JobStats
+	{
+		double						BeginNativeTime = 0.0;		/// The timestamp [ms] when this job started
+		double						EndNativeTime = 0.0;		/// The timestamp [ms] when this job ended
+		double						BeginRunTime = 0.0;			/// The timestamp [ms] when this job RUN is started	
+		double						EndRunTime = 0.0;			/// The timestamp [ms] when this job RUN has ended
+
+		double						TargetPrepStartTime = 0;	/// The timestamp [ms] when this job PrepareTargets was run
+		double						TargetPrepEndTime = 0;		/// The timestamp [ms] when this job PrepareTargets has ended
+		double						TargetPrepWaitStartTime = 0;/// The timestamp [ms] when this job started waiting for targets to be prepared (all promises to return from Device)
+		double						TargetPrepWaitEndTime = 0;	/// The timestamp [ms] when this job when the waiting of promises has finished
+	};
+
+protected:
+	typedef T_Tiles<JobResultPtr>	JobResultPtrTiles;
+
+	UMixInterface*					MixObj = nullptr;	/// The mix that this job belongs to
+	int64							Id = -1;			/// Unique ID for the job [-1 means it was unassigned]
+	int32							QueueId = -1;		/// ID of the queue that this job belongs to
+	TArray<JobArgPtr>				Args;				/// The arguments for the job
+	BlobTransformPtr				Transform;			/// The transformation that we'll be doing for the job
+	FWeakObjectPtr					ErrorOwner;			/// Error owner is an object associated with any error reported to the TextureGraphEngine::ErrorReporter during this job execution
+	mutable CHashPtr				HashValue = nullptr;/// What is the hash of this job
+	TiledBlob_PromisePtr			ResultOrg;			/// The original result that was created
+	TiledBlobRef					Result;				/// The rendering result of this job
+	BufferDescriptorPtr				DesiredResultDesc;	/// Desired descriptor for the result
+	JobResultPtrTiles				TileResults;		/// The results for all the tiles
+	JobResultPtr					FinalJobResult;		/// The overall result info for this job
+
+	JobPArray						Parents;			/// [upstream] Jobs that this job is dependent on
+	JobPtrArray						Children;			/// [downstream] Jobs that are dependent on this job
+
+	bool							bIsTiled = true;	/// Whether the job runs in tiled mode or not. You can override this and force the job to run in non-tiled mode
+	int32							TargetId = -1;		/// What is our target
+
+	TileInvalidateMatrix			TileInvalidationMatrix; /// The tiles that were actually invalidated
+
+	BufferDescriptor				TileDesc;			/// The descriptor for each tile
+
+	FCriticalSection				Mutex;				/// For preventing shared table writes
+	JobRunInfo						RunInfo;			/// The run info that was saved when it was passed as argument
+
+	JobStats						Stats;				/// Job timing stats collected during execution
+
+	const RenderMesh*				Mesh = nullptr;	/// The mesh that we're targeting 
+
+	int32							ReplayCount = 0;	/// Counting the number of debug replay of this job
+	bool							bIsNoCache = false;	/// Whether we should ignore the cache
+
+	JobPtrVec						JobsGeneratedPrior;	/// [Prior] The jobs that were generated from this job
+	JobPtrVec						JobsGeneratedAfter;	/// [After] The jobs that were generated from this job
+
+	JobPtrW							Generator;			/// The job that generated this job. This is useful in certain situations
+	FString							DebugJobName;		/// The name of the job
+
+
+	//std::vector<std::shared_ptr<cti::promise<TiledBlobPtr>>> _blobPreparedPromises;
+
+	typedef std::pair<int32, int32>	IndexPair;
+	typedef TMap<HashType, IndexPair> JobLUT;
+
+	/// Binds/Unbinds args for individual tiles or the entire blob if its a non-tiled job
+	virtual AsyncJobArgResultPtr	BindArgs(JobRunInfo NewRunInfo, BlobTransformPtr TransformObj, int32 TileX, int32 TileY);
+	virtual AsyncJobArgResultPtr	UnbindArgs(JobRunInfo NewRunInfo, BlobTransformPtr TransformObj, int32 TileX, int32 TileY);
+
+	/// Binds/Unbinds args for all the tiles at once and returns a promise
+	virtual AsyncInt				BindArgs_All(JobRunInfo RunInfo);
+	virtual AsyncInt				UnbindArgs_All(JobRunInfo RunInfo);
+
+	typedef std::function			<AsyncJobArgResultPtr(JobRunInfo, BlobTransformPtr, int32, int32)> Bind_Unbind_Func;
+	virtual AsyncInt				Bind_Or_Unbind_All_Generic(JobRunInfo NewRunInfo, Bind_Unbind_Func BFunc);
+
+	virtual AsyncJobResultPtr		RunTile(JobRunInfo RunInfo, int32 TileX, int32 TileY);
+	virtual AsyncJobResultPtr		RunSingle(JobRunInfo RunInfo);
+	virtual AsyncJobResultPtr		FinaliseTiles(JobRunInfo RunInfo);
+	void							SetTileResult(int32 TileX, int32 TileY, JobResultPtr TileResult);
+	virtual AsyncPrepareResult		PrepareTargets(JobBatch* batch);
+
+	AsyncTransformResultPtr			ExecTransform(JobRunInfo RunInfo, BlobTransformPtr TransformObj , int32 TileX, int32 TileY, CHashPtr jobHash);
+	virtual void					MarkJobDone();
+	virtual bool					IsDiscard() const;
+
+	//////////////////////////////////////////////////////////////////////////
+	/// These methods can be used to execute a Job then and there
+	//////////////////////////////////////////////////////////////////////////
+	virtual AsyncJobResultPtr		Run(JobRunInfo RunInfo);
+	virtual AsyncPrepareResult		PrepareResources(JobBatch* batch);
+
+	virtual bool					CheckCached();
+	BufferDescriptor				GetResultDesc() const;
+	BufferDescriptor				GetCombinedDesc(BufferDescriptor& ArgsDescCombined, size_t& Count) const;
+
+public:
+									Job(int32 TargetId, BlobTransformPtr InTransform, UObject* ErrorOwner = nullptr, uint16 InPriority = (uint16)E_Priority::kNormal, uint64 InId = -1);
+									Job(UMixInterface* InMix, int32 TargetId, BlobTransformPtr Transform, UObject* ErrorOwner = nullptr, uint16 priority = (uint16)E_Priority::kNormal, uint64 id = -1);
+	virtual							~Job() override;
+
+
+	virtual Job*					SetArgs(const TArray<JobArgPtr>& NewArgs);
+	virtual Job*					AddArg(JobArgPtr Arg);
+	uint32_t						NumArgs() const;
+	JobArgPtr						GetArg(uint32_t Index) const;
+	virtual void					AddResultToBlobber();
+
+	virtual CHashPtr				Hash() const;
+	virtual CHashPtr				TileHash(int32 TileX, int32 TileY) const;
+	virtual bool					CanHandleTiles() const;
+	virtual bool					CheckDefaultArgs() const;
+
+	virtual TiledBlobPtr			InitResult(FString NewName, const BufferDescriptor* InDesiredDesc, int32 NumTilesX = 0, int32 NumTilesY = 0);
+
+	void							ResetForReplay(bool noCache); /// For debug purpose, reset the job so it can be replayed
+
+	//////////////////////////////////////////////////////////////////////////
+	/// These methods are used to queue up the job as a DeviceNativeTask.
+	/// You cannot call the DeviceNativeTask implementation directly. That is
+	/// done by the Device itself.
+	//////////////////////////////////////////////////////////////////////////
+	virtual AsyncInt				BeginNative(JobRunInfo RunInfo);
+	virtual AsyncJobResultPtr		EndNative();
+
+	//////////////////////////////////////////////////////////////////////////
+	/// BEGIN: DeviceNativeTask implementation
+	//////////////////////////////////////////////////////////////////////////
+protected:
+	virtual TiledBlobPtr			InitLateBoundResult(FString NewName, BufferDescriptor DesiredDesc, size_t NumInputBlobs);
+
+	virtual AsyncInt				PreExecAsync(ENamedThreads::Type execThread, ENamedThreads::Type ReturnThread) override;
+	virtual int32					Exec() override;
+	virtual void					PostExec() override;
+	virtual AsyncInt				ExecAsync(ENamedThreads::Type execThread, ENamedThreads::Type ReturnThread) override;
+	virtual JobPtrVec				GetArgDependencies();
+
+public:
+	virtual Device*					GetTargetDevice() const override;
+	virtual ENamedThreads::Type		GetExecutionThread() const override;
+	virtual bool					IsAsync() const override;
+	bool							GetTileInvalidation(int32 TileX, int32 TileY) const;
+
+	virtual void					GetDependencies(JobPtrVec& Prior, JobPtrVec& After, JobRunInfo RunInfo);
+
+	virtual FString					GetName() const override;
+
+	virtual bool					DebugCompleteCheck() override;
+
+	//////////////////////////////////////////////////////////////////////////
+	/// END: DeviceNativeTask implementation
+	//////////////////////////////////////////////////////////////////////////
+	virtual bool					CheckCulled(JobRunInfo RunInfo);
+	virtual FString					GetRunTimings(double BatchStartTime) const;
+	virtual FString					GetDebugName() const override;
+
+	//////////////////////////////////////////////////////////////////////////
+	/// Inline functions
+	//////////////////////////////////////////////////////////////////////////
+	FORCEINLINE int64				GetJobId() const { return Id; }
+	FORCEINLINE void				SetJobId(int64 NewJobId) { Id = NewJobId; }
+	FORCEINLINE int32				GetQueueId() const { return QueueId; }
+	FORCEINLINE void				SetQueueId(int32 NewQueueId) { QueueId = NewQueueId; }
+	FORCEINLINE BlobTransformPtr	GetTransform() const { return Transform; }
+	FORCEINLINE TiledBlobRef		GetResult() const { return Result; }
+	
+
+	FORCEINLINE bool				GetTiled() const { return bIsTiled; }
+	FORCEINLINE void				SetTiled(bool bTiled) { bIsTiled = bTiled; }
+
+	FORCEINLINE JobStats			GetStats() const { return Stats; }
+
+	FORCEINLINE const JobRunInfo&	GetRunInfo() const { return RunInfo; }
+	FORCEINLINE const RenderMesh*	GetMesh() const { return Mesh; }
+	FORCEINLINE void				SetMesh(RenderMesh* NewMesh) { Mesh = NewMesh; }
+	FORCEINLINE JobResultPtr		GetTileResult(int32 TileX, int32 TileY) const
+	{
+		if (TileX >= 0 && TileY >= 0)
+			return TileResults[TileX][TileY];
+		return FinalJobResult;
+	}
+
+	FORCEINLINE int32				GetReplayCount() const { return ReplayCount; } /// For debug purpose, report the replay number
+	FORCEINLINE int32				GetTargetId() const { return TargetId; }
+
+	FORCEINLINE JobPtrVec&			GeneratedPriorJobs() { return JobsGeneratedPrior; }
+	FORCEINLINE JobPtrVec&			GeneratedAfterJobs() { return JobsGeneratedAfter; }
+	FORCEINLINE const JobPtrW&		GeneratorJob() const { return Generator; }
+
+	FORCEINLINE UMixInterface*		GetMix() const { return MixObj; }
+	FORCEINLINE void				SetMix(UMixInterface* mix) { MixObj = mix; }
+
+	FORCEINLINE BlobRef				GetResultRef() const { return BlobRef(std::static_pointer_cast<Blob>(Result.get()), false); }
+	FORCEINLINE TiledBlob_PromisePtr GetResultPromise() const { return std::static_pointer_cast<TiledBlob_Promise>(Result.get()); }
+
+	FORCEINLINE UObject*			GetErrorOwner() const { return ErrorOwner.Get(); }
+
+};
+
+//////////////////////////////////////////////////////////////////////////
diff --git a/Engine/Plugins/Experimental/TextureGraph/Source/TextureGraphEngine/Job/Scheduler.cpp b/Engine/Plugins/Experimental/TextureGraph/Source/TextureGraphEngine/Job/Scheduler.cpp
new file mode 100644
--- /dev/null
+++ b/Engine/Plugins/Experimental/TextureGraph/Source/TextureGraphEngine/Job/Scheduler.cpp
@@ -0,0 +1,402 @@
+// Copyright Epic Games, Inc. All Rights Reserved.
+#include "Scheduler.h"
+#include "Job.h"
+#include "TempHashService.h"
+#include "TextureGraphEngine.h"
+#include "Device/DeviceManager.h"
+#include "Device/Device.h"
+
+DECLARE_CYCLE_STAT(TEXT("Scheduler_Update"), STAT_Scheduler_Update, STATGROUP_TextureGraphEngine);
+
+const double Scheduler::IdleTimeInterval = 500.0;
+const double Scheduler::IdleBatchTimeLimit = 100.0;
+const double Scheduler::CurrentBatchWarningLimit = 2000.0;
+
+Scheduler::Scheduler()
+{
+	ObserverSource = std::make_shared<SchedulerObserverSource>();
+
+	TempHashServicePtr TempHashSPtr = std::make_shared<TempHashService>();
+	BlobHasherServicePtr BlobHasherSPtr = std::make_shared<BlobHasherService>();
+	DeviceTransferServicePtr DeviceTransferSPtr = std::make_shared<DeviceTransferService>();
+	ThumbnailsServicePtr ThumbnailsSPtr = std::make_shared<ThumbnailsService>();
+	MipMapServicePtr MipmapSPtr = std::make_shared<MipMapService>();
+	MinMaxServicePtr MinmaxSPtr = std::make_shared<MinMaxService>();
+	HistogramServicePtr HistogramSPtr = std::make_shared<HistogramService>();
+
+	BlobHasherServiceObj = BlobHasherSPtr;
+	DeviceTransferServiceObj = DeviceTransferSPtr;
+	ThumbnailsServiceObj = ThumbnailsSPtr;
+	MipMapServiceObj = MipmapSPtr;
+	MinMaxServiceObj = MinmaxSPtr;
+	HistogramServiceObj = HistogramSPtr;
+
+	AddIdleService(BlobHasherSPtr);
+	AddIdleService(DeviceTransferSPtr);
+	AddIdleService(ThumbnailsSPtr);
+	AddIdleService(MipmapSPtr);
+	AddIdleService(MinmaxSPtr);
+	AddIdleService(HistogramSPtr);
+}
+
+Scheduler::~Scheduler()
+{
+	bIsRunning = false;
+	StopServices();
+
+	/// Add a bit of a delay to allow proper clean-up
+	std::this_thread::sleep_for(std::chrono::milliseconds((int)IdleBatchTimeLimit * 2));
+
+	Stop();
+}
+
+void Scheduler::Update(float dt)
+{
+	SCOPE_CYCLE_COUNTER(STAT_Scheduler_Update);
+	check(IsInGameThread());
+
+	/// Make sure value is initialised correct
+	if (TimeSinceIdle < 0.001)
+		TimeSinceIdle = Util::Time();
+
+	if (!bIsRunning)
+		Start();
+
+	{
+		FScopeLock Lock(&CurrentBatchMutex);
+		if (CurrentBatch)
+		{
+			double currentTime = Util::Time();
+			double delta = currentTime - CurrentBatchStartTime;
+			double timeoutLimit = !TextureGraphEngine::IsTestMode() ? CurrentBatchWarningLimit : CurrentBatchWarningLimit * 4.0;
+
+			if (delta > timeoutLimit)
+			{
+				UE_LOG(LogBatch, Warning, TEXT("Current Batch: %llu been running for %f ms [Max threshold: %f, Jobs: %u/%u]!"), 
+					CurrentBatch->GetBatchId(), (float)delta, (float)CurrentBatchWarningLimit, CurrentBatch->GetNumJobsRunning(), 
+					(uint32)CurrentBatch->NumJobs());
+				CurrentBatch->DebugDumpUnfinishedJobs();
+				//_currentBatch->Terminate();
+			}
+
+			return;
+		}
+	}
+
+	JobBatchPtr batchToRun = nullptr;
+	{
+		FScopeLock Lock(&BatchMutex);
+		if (Batches.size())
+		{
+			/// TODO: Enable this at some point
+			if (Batches.size() == 1)
+			{
+				batchToRun = Batches.front();
+				Batches.pop_front();
+			}
+			else
+			{
+				check(!TextureGraphEngine::IsTestMode());
+
+				/// OK here, we're going to merge batches belonging to the same mix together
+				
+				// figure out the prioritized mix from the Batch
+				batchToRun = Batches.front();
+				UMixInterface* currentBatchMix = batchToRun->GetCycle()->GetMix();
+
+				// find the latest Batch of required mix
+				auto iter = Batches.begin();
+				JobBatchPtr lastMixbatch = batchToRun;
+
+				while (iter != Batches.end())
+				{
+					JobBatchPtr futureBatch = *iter;
+
+					if (futureBatch->GetCycle()->GetMix() == currentBatchMix)
+					{
+						lastMixbatch = *iter;
+					}
+
+					iter++;
+				}
+
+				// we get latest Batch to run for the mix
+				batchToRun = lastMixbatch;
+
+				// Merge invalidation details of the mix from multiple batches.
+				// Remove the batches from the list after merging.
+				iter = Batches.begin();
+				while (iter != Batches.end())
+				{
+					JobBatchPtr futureBatch = *iter;
+
+					if (futureBatch->GetCycle()->GetMix() == currentBatchMix)
+					{
+						batchToRun->GetCycle()->MergeDetails(futureBatch->GetCycle()->GetDetails());
+						iter = Batches.erase(iter);
+					}
+					else
+						iter++;
+				}
+			}
+		}
+	}
+
+	if (batchToRun)
+	{
+		FScopeLock Lock(&CurrentBatchMutex);
+		CurrentBatch = batchToRun;
+		CurrentBatchStartTime = Util::Time();
+
+		UE_LOG(LogBatch, Log, TEXT("Next Batch: %llu [Prev: %llu]"), CurrentBatch->GetBatchId(), PreviousBatch ? PreviousBatch->GetBatchId() : 0);
+
+		if (!CurrentBatch->WasGeneratedFromIdleService())
+			TimeSinceIdle = Util::Time();
+
+		if (bCaptureNextBatch)
+		{
+			CurrentBatch->SetCaptureRenderDoc(true);
+			bCaptureNextBatch = false;
+		}
+
+		CurrentBatch->Exec([this] (JobBatch*)	/// Instead of passing it as an argument, JobBatch should be a return type; this is to keep it from cyclic dependancy
+			{
+				/// This closure will be called when all the jobs execution is really over
+				/// tell the world through the observer
+				/// Clear out the current Batch
+				FScopeLock Lock(&CurrentBatchMutex);
+				PreviousBatch = CurrentBatch;
+				CurrentBatch = nullptr;
+
+				if (PreviousBatch && !PreviousBatch->WasGeneratedFromIdleService())
+					TimeSinceIdle = Util::Time();
+
+				UE_LOG(LogBatch, Log, TEXT("Scheduler, Batch fully queued: %llu. Triggering Observer::BatchJobsDone ..."), PreviousBatch ? PreviousBatch->GetBatchId() : -1);
+
+				ObserverSource->BatchJobsDone(PreviousBatch);
+
+				UE_LOG(LogBatch, Log, TEXT("Scheduler Observer::BatchJobsDone finished for Batch: %llu"), PreviousBatch ? PreviousBatch->GetBatchId() : -1);
+			})
+			.then([this]()
+				{
+					if (PreviousBatch && !PreviousBatch->WasGeneratedFromIdleService())
+						TimeSinceIdle = Util::Time();
+
+					UE_LOG(LogBatch, Log, TEXT("Scheduler triggering Observer::BatchDone for Batch: %llu ..."), PreviousBatch ? PreviousBatch->GetBatchId() : -1);
+
+					ObserverSource->BatchDone(CurrentBatch);
+
+					UE_LOG(LogBatch, Log, TEXT("Scheduler Observer::BatchDone finished for Batch: %llu"), PreviousBatch ? PreviousBatch->GetBatchId() : -1);
+				})
+				.then([this]()
+					{
+						// Trigger cycle->invalidateDetails.DoneCallBackLambda()
+						
+						//Adding logs to debug and safty checks here to avoid the crash
+						//Crash was happening in cases when we do not have a valid Cycle or Mix
+						//This could happen when we open a script run it and then delete it and then try to run another script.
+						// 
+						//TODO: Investigate why we need to call the PreviousBatch->BroadcastOnDone instead of calling Currentbatch->BroadcastOnDone
+						//As we are setting the PreviousBatch = CurrentBatch in the callback of Exec function and that gets called from EndBatch after the
+						//execution of this .then block 
+						//looks like a bug though 
+						//fix would be to remove this .then block and call PreviousBatch->GetCycle()->GetDetails().BroadcastOnDone(); in the callback of Exec function
+						if (!PreviousBatch)
+						{
+							UE_LOG(LogBatch, Log, TEXT("Scheduler Previous Batch is null"));
+							return;
+						}
+						if (!PreviousBatch->GetCycle())
+						{
+							UE_LOG(LogBatch, Warning, TEXT("Scheduler Previous Batch : %llu Cycle is null") , PreviousBatch->GetBatchId());
+							return;
+						}
+						if (!PreviousBatch->GetCycle()->GetDetails().Mix.IsValid())
+						{
+							UE_LOG(LogBatch, Warning, TEXT("Scheduler Previous Batch : %llu Cycle Details Mix is null"), PreviousBatch->GetBatchId());
+							return;
+						}
+
+						PreviousBatch->GetCycle()->GetDetails().BroadcastOnDone();
+					});
+	}
+	else
+	{
+		double Delta = Util::TimeDelta(TimeSinceIdle);
+		if (Delta > IdleTimeInterval)
+		{
+			UpdateIdle();
+
+			// Since we have time, notify the observer update
+			ObserverSource->UpdateIdle();
+		}
+	}
+}
+
+void Scheduler::SetCaptureRenderDocNextBatch(bool bCapture /* = true */)
+{
+	bCaptureNextBatch = bCapture;
+}
+
+AsyncJobResultPtr Scheduler::UpdateIdleBatch(size_t Index)
+{
+	check(Index < IdleServices.size());
+	IdleServicePtr Batch = IdleServices[Index];
+	double StartTime = Util::Time();
+
+	return Batch->Tick().then([this, Batch, StartTime](JobResultPtr result)
+		{
+			double EndTime = Util::Time();
+			double Duration = EndTime - StartTime;
+			bool bDidOffendTimeLimit = false;
+
+			if (Duration > IdleBatchTimeLimit)
+			{
+				bDidOffendTimeLimit = true;
+				UE_LOG(LogBatch, Warning, TEXT("Idle Batch [%s] time limit offense @ %f"), *Batch->GetName(), (float)Util::Time());
+			}
+
+			Batch->UpdateStats(EndTime, StartTime, bDidOffendTimeLimit);
+
+			return result;
+		});
+}
+
+void Scheduler::UpdateIdle()
+{
+	/// Currently we don't do the idle update in test mode
+	if (TextureGraphEngine::IsTestMode())
+		return;
+
+	// if no Batch to run, check if we have an idle Batch we should run
+	{
+		/// Shouldn't run idle services if a Batch is still being processed
+		FScopeLock Lock(&CurrentBatchMutex);
+		if (CurrentBatch)
+			return;
+	}
+
+	{
+		/// Shouldn't run idle services if we still have batches to process
+		FScopeLock Lock(&BatchMutex);
+		if (Batches.size())
+			return;
+	}
+
+	UE_LOG(LogIdle_Svc, Verbose, TEXT("Scheduler::UpdateIdle"));
+
+	if (TextureGraphEngine::IsDestroying())
+		return;
+
+	size_t NumIdleBatches = 0;
+
+	{
+		FScopeLock Lock(&IdlServiceMutex);
+		NumIdleBatches = IdleServices.size();
+	}
+
+	std::vector<AsyncJobResultPtr> Promises;
+
+	Promises.emplace_back(TextureGraphEngine::GetBlobber()->UpdateIdle());
+
+	/// Idle update for all the devices
+	for (size_t DeviceIndex = 0; DeviceIndex < TextureGraphEngine::GetDeviceManager()->GetNumDevices(); DeviceIndex++)
+	{
+		Device* Dev = TextureGraphEngine::GetDeviceManager()->GetDevice(DeviceIndex);
+		if (Dev)
+			Promises.emplace_back(Dev->UpdateIdle());
+	}
+
+	for (size_t IdleBatchIndex = 0; IdleBatchIndex < NumIdleBatches; IdleBatchIndex++)
+	{
+		auto Promise = UpdateIdleBatch(IdleBatchIndex);
+		Promises.push_back(std::move(Promise));
+	}
+
+	if (Promises.size())
+	{
+		cti::when_all(Promises.begin(), Promises.end()).then([this]()
+			{
+				/// TODO:
+			});
+	}
+}
+
+void Scheduler::AddIdleService(IdleServicePtr Service)
+{
+	FScopeLock Lock(&IdlServiceMutex);
+	IdleServices.push_back(Service);
+}
+
+void Scheduler::CaptureRenderDocLastRunBatch()
+{
+}
+
+void Scheduler::ClearCache()
+{
+	check(IsInGameThread() && TextureGraphEngine::IsTestMode());
+
+	check(!CurrentBatch);
+	check(!PreviousBatch || PreviousBatch->IsFinished());
+	check(Batches.empty());
+
+	PreviousBatch = nullptr;
+}
+
+void Scheduler::AddBatch(JobBatchPtr Batch)
+{
+	check(IsInGameThread());
+
+	FScopeLock Lock(&BatchMutex);
+
+	if (Batches.size())
+	{
+		check(Batches.back()->GetFrameId() <= Batch->GetFrameId());
+	}
+
+	check(Batch->GetFrameId())
+		UE_LOG(LogBatch, Log, TEXT("Adding Batch: %llu [FrameId: %llu, Num batches: %llu]"), Batch->GetBatchId(), Batch->GetFrameId(), Batches.size());
+
+	Batches.push_back(Batch);
+	ObserverSource->BatchAdded(Batch); // notify observer
+	TimeSinceIdle = Util::Time();
+}
+
+void Scheduler::Start()
+{
+	check(IsInGameThread());
+	bIsRunning = true;
+
+	if (ObserverSource)
+		ObserverSource->Start();
+}
+
+void Scheduler::StopServices()
+{
+	check(IsInGameThread());
+	for (auto Service : IdleServices)
+		Service->Stop();
+
+	IdleServices.clear();
+}
+
+void Scheduler::Stop()
+{
+	check(IsInGameThread());
+	bIsRunning = false;
+
+	if (ObserverSource)
+		ObserverSource->Stop();
+}
+
+void Scheduler::RegisterObserverSource(const SchedulerObserverSourcePtr& InObserverSource)
+{
+	if (InObserverSource)
+	{
+		ObserverSource = InObserverSource;
+	}
+	else
+	{
+		ObserverSource = std::make_shared<SchedulerObserverSource>();
+	}
+}
\ No newline at end of file
diff --git a/Engine/Plugins/Experimental/TextureGraph/Source/TextureGraphEngine/Job/TempHashService.h b/Engine/Plugins/Experimental/TextureGraph/Source/TextureGraphEngine/Job/TempHashService.h
new file mode 100644
--- /dev/null
+++ b/Engine/Plugins/Experimental/TextureGraph/Source/TextureGraphEngine/Job/TempHashService.h
@@ -0,0 +1,26 @@
+// Copyright Epic Games, Inc. All Rights Reserved.
+#pragma once
+
+#include "CoreMinimal.h"
+#include "IdleService.h"
+#include <vector>
+
+class TEXTUREGRAPHENGINE_API TempHashService : public IdleService
+{
+protected:
+	FCriticalSection				HashesMutex;			/// Mutex for the hashes
+	std::vector<CHashPtr>			Hashes;					/// The hashes that we're waiting to resolve
+
+public:
+									TempHashService();
+	virtual							~TempHashService() override;
+
+	virtual AsyncJobResultPtr		Tick() override;
+	virtual void					Stop() override;
+
+	void							Add(CHashPtr HashValue);
+};
+
+typedef std::shared_ptr<TempHashService>	TempHashServicePtr;
+typedef std::weak_ptr<TempHashService>		TempHashServicePtrW;
+
diff --git a/Engine/Plugins/Experimental/TextureGraph/Source/TextureGraphEngine/Job/ThumbnailsService.h b/Engine/Plugins/Experimental/TextureGraph/Source/TextureGraphEngine/Job/ThumbnailsService.h
new file mode 100644
--- /dev/null
+++ b/Engine/Plugins/Experimental/TextureGraph/Source/TextureGraphEngine/Job/ThumbnailsService.h
@@ -0,0 +1,51 @@
+// Copyright Epic Games, Inc. All Rights Reserved.
+#pragma once
+
+#include "CoreMinimal.h"
+#include "IdleService.h"
+#include "Model/Mix/MixUpdateCycle.h"
+#include <vector>
+#include <unordered_map>
+
+class UModelObject;
+typedef std::shared_ptr<class Job>		JobPtr;
+typedef std::unique_ptr<class Job>		JobUPtr;
+typedef std::weak_ptr<class Job>		JobPtrW;
+
+class UMixInterface;
+
+class TEXTUREGRAPHENGINE_API ThumbnailsService : public IdleService
+{
+protected:
+	JobBatchPtr						Batch;					/// The thumbnail rendering batch
+	JobBatchPtr						PrevBatch;				/// Last rendered batch
+	MixUpdateCyclePtr				Cycle;					/// The cycle that we'll be using for the thumbnail updates
+
+	struct AddedJobs
+	{
+		JobBatchPtr					AssociatedBatch;		/// Batch holding the following job
+		JobPtrW						ThumbnailJob;			/// Thumbnail updating job assigned 
+	};
+
+	std::unordered_map<UObject*, AddedJobs> Handled;		/// The components that we have already added jobs for
+
+	JobBatchPtr						CreateNewUpdateCycle(UMixInterface* Mix);
+	JobBatchPtr						GetNextUpdateCycle();
+
+public:
+									ThumbnailsService();
+	virtual							~ThumbnailsService() override;
+
+	virtual AsyncJobResultPtr		Tick() override;
+	virtual void					Stop() override;
+	void							AddUniqueJobToCycle(UObject* CreatorComponent, UMixInterface* Mix, JobUPtr ThumbnailJob, int32 TargetId, bool bForceExec = false);	//Use this to always generate thumbnail from last update cycle
+
+	//////////////////////////////////////////////////////////////////////////
+	/// Inline functions
+	//////////////////////////////////////////////////////////////////////////
+	JobBatchPtr						GetCurrentBatch() const { return Batch; }
+};
+
+typedef std::shared_ptr<ThumbnailsService>	ThumbnailsServicePtr;
+typedef std::weak_ptr<ThumbnailsService>	ThumbnailsServicePtrW;
+
diff --git a/Engine/Plugins/Experimental/TextureGraph/Source/TextureGraphEngine/Model/Mix/Mix.h b/Engine/Plugins/Experimental/TextureGraph/Source/TextureGraphEngine/Model/Mix/Mix.h
new file mode 100644
--- /dev/null
+++ b/Engine/Plugins/Experimental/TextureGraph/Source/TextureGraphEngine/Model/Mix/Mix.h
@@ -0,0 +1,51 @@
+// Copyright Epic Games, Inc. All Rights Reserved.
+#pragma once
+
+#include "CoreMinimal.h" 
+#include "MixInterface.h"
+#include "MixerMeshComponent.h"
+
+#include <memory>
+#include <vector> 
+
+#include "Mix.generated.h"
+
+DECLARE_LOG_CATEGORY_EXTERN(LogMix, All, All);
+
+class UMixerProject;
+class UMixParameters;
+
+class RenderMesh;
+typedef std::shared_ptr<RenderMesh>	RenderMeshPtr;
+
+UCLASS(Blueprintable, BlueprintType)
+class TEXTUREGRAPHENGINE_API UMix : public UMixInterface
+{
+	GENERATED_BODY()
+
+public:
+	
+
+private:
+	static UMix*						GNullMix;					/// This is a generic mix that is used in various places and rendering actions that
+																	/// often require a Mix object. This is initialised in the first call to the public 
+																	/// static function below. 
+	static const FString				GRootSectionName;
+
+public:
+	static UMix*						NullMix();
+
+protected:	
+	UPROPERTY()
+	TObjectPtr<UMixerMeshComponent>		MeshComponent;
+
+public:
+	virtual								~UMix() override;
+	
+	virtual void						Update(MixUpdateCyclePtr Cycle) override;
+	
+	//////////////////////////////////////////////////////////////////////////
+	/// Inline Functions
+	//////////////////////////////////////////////////////////////////////////
+	virtual UMixerMeshComponent*		GetMeshComponent() const override { return MeshComponent; }
+};
diff --git a/Engine/Plugins/Experimental/TextureGraph/Source/TextureGraphEngine/Model/Mix/MixInterface.h b/Engine/Plugins/Experimental/TextureGraph/Source/TextureGraphEngine/Model/Mix/MixInterface.h
new file mode 100644
--- /dev/null
+++ b/Engine/Plugins/Experimental/TextureGraph/Source/TextureGraphEngine/Model/Mix/MixInterface.h
@@ -0,0 +1,111 @@
+// Copyright Epic Games, Inc. All Rights Reserved.
+#pragma once
+
+#include "CoreMinimal.h" 
+#include "Data/TiledBlob.h"
+#include "Model/Mix/MixUpdateCycle.h"
+
+#include "MixInterface.generated.h"
+
+class UMix;
+class UMixerMeshComponent;
+class ULayerStack;
+class UMixSettings;
+class AActor;
+
+class RenderMesh;
+typedef std::shared_ptr<RenderMesh>		RenderMeshPtr;
+
+class MeshAsset;
+typedef std::shared_ptr<MeshAsset>		MeshAssetPtr;
+
+//////////////////////////////////////////////////////////////////////////
+/// Mix can now be saved as Mix or MixInstance.
+/// To support this functionality, we introduced the new class named MixInterface. 
+/// This will be the base class for both Mix and MixInstance.
+//////////////////////////////////////////////////////////////////////////
+UCLASS(Abstract)
+class TEXTUREGRAPHENGINE_API UMixInterface : public UModelObject
+{
+	GENERATED_BODY()
+
+private:
+	static const TMap<FString, FString> s_uriAlias;						/// Contains the list of URI aliases map to keep the data meaningful, short and precise.
+	static const TMap<FString, FString> s_uriAliasDecryptor;			/// We need alias decryptor to convert it so we can extract the object from the data. 
+
+	static TMap<FString, FString>		InitURIAlias();
+	static TMap<FString, FString>		InitURIAliasDecryptor();
+
+	FString								GetObjectAlias(FString object);
+	FString								GetAliasToObject(FString Alias);
+
+protected:
+	UPROPERTY()
+	int32								Priority;			/// The priority of this mix in terms of update cycle
+
+	UPROPERTY()
+	UMixSettings*						Settings = nullptr;/// The settings for this mix interface
+
+	UPROPERTY()
+	bool								bInvalidateTextures;/// Invalidate the scene textures or not
+
+	std::atomic_int64_t					InvalidationFrameId;/// When was the mix last invalidated
+	std::atomic_int64_t					UpdatedFrameId;		/// When was the mix actually last updated
+
+	bool								bEnableLOD = true;	/// Enable LOD-ing in the system
+
+	virtual UModelObject*				RootURI() { checkNoEntry() return nullptr; };
+	virtual void						Invalidate(FModelInvalidateInfo InvalidateInfo);
+
+public:
+	virtual								~UMixInterface() override;
+
+
+	/** The mesh used by the TS editor to preview.*/
+	UPROPERTY(EditAnywhere, Category = Previewing, meta = (AllowedClasses = "/Script/Engine.StaticMesh,/Script/Engine.SkeletalMesh", ExactClass = "true"))
+	FSoftObjectPath						PreviewMesh;  // Adding it here because this will be the base class for TS_Script and TS_ScriptInstance.
+
+	DECLARE_DELEGATE_OneParam(FOnRenderDone, UMixInterface*);
+	FOnRenderDone						OnRenderDone;
+	virtual void						PostMeshLoad();
+
+	virtual bool						CanEdit() { return false; }
+
+	virtual void						SetMesh(RenderMeshPtr MeshObj, int MeshType, FVector Scale, FVector2D Dimension);
+
+#if WITH_EDITOR
+	virtual AsyncActionResultPtr		SetEditorMesh(AActor* Actor);
+	virtual AsyncActionResultPtr		SetEditorMesh(UStaticMeshComponent* MeshComponent, UWorld* World);
+#endif
+	
+	virtual void						Update(MixUpdateCyclePtr cycle);	
+
+	virtual RenderMeshPtr				GetMesh() const; // Get the mesh assigned in Settings
+
+	virtual UMixerMeshComponent*		GetMeshComponent() const { checkNoEntry() return nullptr; }
+
+	virtual int32						Width() const; 
+	virtual int32						Height() const;
+	
+	virtual int32						GetNumXTiles() const;
+	virtual int32						GetNumYTiles()const;
+	
+	virtual UMixSettings*				GetSettings() const;
+
+	bool								IsHigherPriorityThan(const UMixInterface* RHS) const;
+
+	virtual void						InvalidateWithDetails(const FInvalidationDetails& Details);
+	virtual void						InvalidateAll();
+
+	virtual void						BroadcastOnRenderingDone();
+	
+	//////////////////////////////////////////////////////////////////////////
+	/// Inline functions
+	//////////////////////////////////////////////////////////////////////////
+	FORCEINLINE int32					GetPriority() const { return Priority; }
+
+	FORCEINLINE int64					GetInvalidationFrameId() const { return InvalidationFrameId; }
+	FORCEINLINE int64					GetUpdateFrameId() const { return UpdatedFrameId; }
+};
+
+
diff --git a/Engine/Plugins/Experimental/TextureGraph/Source/TextureGraphEngine/Model/Mix/MixerMeshComponent.cpp b/Engine/Plugins/Experimental/TextureGraph/Source/TextureGraphEngine/Model/Mix/MixerMeshComponent.cpp
new file mode 100644
--- /dev/null
+++ b/Engine/Plugins/Experimental/TextureGraph/Source/TextureGraphEngine/Model/Mix/MixerMeshComponent.cpp
@@ -0,0 +1,11 @@
+// Copyright Epic Games, Inc. All Rights Reserved.
+#include "MixerMeshComponent.h"
+#include "Mix.h"
+
+void UMixerMeshComponent::SetMesh(RenderMeshPtr mesh)
+{
+	MeshObj = mesh;
+	verify(MeshObj);
+}
+
+
diff --git a/Engine/Plugins/Experimental/TextureGraph/Source/TextureGraphEngine/Model/Mix/MixerMeshComponent.h b/Engine/Plugins/Experimental/TextureGraph/Source/TextureGraphEngine/Model/Mix/MixerMeshComponent.h
new file mode 100644
--- /dev/null
+++ b/Engine/Plugins/Experimental/TextureGraph/Source/TextureGraphEngine/Model/Mix/MixerMeshComponent.h
@@ -0,0 +1,44 @@
+// Copyright Epic Games, Inc. All Rights Reserved.
+#pragma once
+
+#include "CoreMinimal.h"
+#include "Model/ModelObject.h"
+#include "3D/RenderMesh.h"
+#include "MixerMeshComponent.generated.h"
+
+class RenderMesh;
+typedef std::shared_ptr<RenderMesh>		RenderMeshPtr;
+
+class MeshAsset;
+typedef std::shared_ptr<MeshAsset>		MeshAssetPtr;
+
+UCLASS(Blueprintable, BlueprintType)
+class TEXTUREGRAPHENGINE_API UMixerMeshComponent : public UModelObject
+{
+	GENERATED_BODY()
+
+private:
+	RenderMeshPtr					MeshObj;					/// Mesh extracted from the asset pack
+
+public:
+	
+	void							SetMesh(RenderMeshPtr mesh);
+
+protected:
+	UPROPERTY()
+	int32							MeshType;
+
+	UPROPERTY()
+	float							GroundSizeIndex;
+
+	UPROPERTY()
+	FString							CustomMeshPath;
+
+	UPROPERTY(BlueprintReadOnly, Category = "Settings")
+	FVector3f						CustomMeshScale;
+
+public:
+	FORCEINLINE int32				GetMeshType() const { return MeshType; }
+
+	FORCEINLINE RenderMeshPtr		GetMesh() const { return MeshObj; }
+};
diff --git a/Engine/Plugins/Experimental/TextureGraph/Source/TextureGraphEngine/Model/Mix/ViewportSettings.h b/Engine/Plugins/Experimental/TextureGraph/Source/TextureGraphEngine/Model/Mix/ViewportSettings.h
new file mode 100644
--- /dev/null
+++ b/Engine/Plugins/Experimental/TextureGraph/Source/TextureGraphEngine/Model/Mix/ViewportSettings.h
@@ -0,0 +1,52 @@
+// Copyright Epic Games, Inc. All Rights Reserved.
+#pragma once
+
+#include "CoreMinimal.h" 
+#include "ViewportSettings.generated.h"
+
+class UMaterial;
+
+USTRUCT()
+struct TEXTUREGRAPHENGINE_API FMaterialMappingInfo
+{
+	GENERATED_USTRUCT_BODY()
+
+	UPROPERTY(VisibleAnywhere, Category = "Viewport Settings")
+	FName						MaterialInput;
+
+	UPROPERTY(EditAnywhere, Category = "Viewport Settings")
+	FName						Target;
+
+	bool						HasTarget() const { return !Target.IsNone(); }
+};
+
+USTRUCT()
+struct TEXTUREGRAPHENGINE_API FViewportSettings
+{
+	GENERATED_BODY()
+	
+	UPROPERTY(EditAnywhere, Category = "Viewport Settings", NoClear)
+	TObjectPtr<UMaterial> Material;
+	
+	UPROPERTY(EditAnywhere, EditFixedSize, Category = "Viewport Settings")
+	TArray<FMaterialMappingInfo> MaterialMappingInfos;
+
+	void InitDefaultSettings(FName InitialTargetName);
+	void SetDefaultTarget(FName DefaultTargetName);
+
+	UMaterial* GetDefaultMaterial();
+	
+	FName GetMaterialName() const;
+	FName GetMaterialMappingInfo(const FName MaterialInput);
+	bool ContainsMaterialMappingInfo(const FName InMaterialInput);
+
+	bool RemoveMaterialMappingForTarget(FName OutputNode);
+	void OnMaterialUpdate();
+	void OnTargetRename(const FName OldName, const FName NewName);
+	
+	DECLARE_MULTICAST_DELEGATE(FViewportSettingsUpdateEvent)
+	FViewportSettingsUpdateEvent OnViewportMaterialChangeEvent;
+
+	DECLARE_MULTICAST_DELEGATE(FMaterialMappingChangedEvent)
+	FMaterialMappingChangedEvent OnMaterialMappingChangedEvent;
+};
diff --git a/Engine/Plugins/Experimental/TextureGraph/Source/TextureGraphEngine/Model/StaticImageResource.h b/Engine/Plugins/Experimental/TextureGraph/Source/TextureGraphEngine/Model/StaticImageResource.h
new file mode 100644
--- /dev/null
+++ b/Engine/Plugins/Experimental/TextureGraph/Source/TextureGraphEngine/Model/StaticImageResource.h
@@ -0,0 +1,40 @@
+// Copyright Epic Games, Inc. All Rights Reserved.
+#pragma once
+
+#include "CoreMinimal.h"
+#include "ModelObject.h"
+#include "Device/FX/Device_FX.h"
+#include "Mix/MixUpdateCycle.h"
+#include "StaticImageResource.generated.h"
+
+//////////////////////////////////////////////////////////////////////////
+/// Base static image resource class
+////////////////////////////////////////////////////////////////////////// 
+UCLASS(Blueprintable, BlueprintType)
+class TEXTUREGRAPHENGINE_API UStaticImageResource : public UModelObject
+{
+public:
+	GENERATED_BODY()
+	
+private:
+	friend class Job_LoadStaticImageResource;
+	
+	UPROPERTY()
+	FString 						AssetUUID;						/// Unique id for the asset within the entire system
+	
+	TiledBlobPtr					BlobObj;						/// The blob that represents the data for this source
+
+	virtual AsyncTiledBlobRef		Load(MixUpdateCyclePtr Cycle);
+	
+public:
+	virtual							~UStaticImageResource() override;
+	
+	virtual TiledBlobPtr			GetBlob(MixUpdateCyclePtr Cycle, BufferDescriptor DesiredDesc, int32 TargetId);
+
+	//////////////////////////////////////////////////////////////////////////
+	/// Inline functions
+	//////////////////////////////////////////////////////////////////////////
+	FORCEINLINE const FString&		GetAssetUUID() const { return AssetUUID; }
+	FORCEINLINE void				SetAssetUUID(const FString& UUID) { AssetUUID = UUID; }
+};
+
diff --git a/Engine/Plugins/Experimental/TextureGraph/Source/TextureGraphEngine/Profiling/RenderDoc/RenderDocManager.h b/Engine/Plugins/Experimental/TextureGraph/Source/TextureGraphEngine/Profiling/RenderDoc/RenderDocManager.h
new file mode 100644
--- /dev/null
+++ b/Engine/Plugins/Experimental/TextureGraph/Source/TextureGraphEngine/Profiling/RenderDoc/RenderDocManager.h
@@ -0,0 +1,28 @@
+// Copyright Epic Games, Inc. All Rights Reserved.
+#pragma once
+
+#include "CoreMinimal.h"
+#include <memory>
+/**
+ * 
+ */
+
+DECLARE_LOG_CATEGORY_EXTERN(LogRenderDocTextureGraph, Log, All);
+
+namespace TextureGraphEditor
+{
+	class TEXTUREGRAPHENGINE_API RenderDocManager
+	{
+
+	public:
+		RenderDocManager();
+		~RenderDocManager();
+		void												Initialize();
+
+		void												CaptureNextBatch();
+		void												CapturePreviousBatch();
+		void												BeginCapture();
+		void												EndCapture();
+};
+	typedef std::unique_ptr<RenderDocManager>				RenderDocManagerPtr;
+}
\ No newline at end of file
diff --git a/Engine/Plugins/Experimental/TextureGraph/Source/TextureGraphEngine/TextureGraphErrorReporter.h b/Engine/Plugins/Experimental/TextureGraph/Source/TextureGraphEngine/TextureGraphErrorReporter.h
new file mode 100644
--- /dev/null
+++ b/Engine/Plugins/Experimental/TextureGraph/Source/TextureGraphEngine/TextureGraphErrorReporter.h
@@ -0,0 +1,54 @@
+// Copyright Epic Games, Inc. All Rights Reserved.
+#pragma once
+#include "CoreMinimal.h"
+
+DECLARE_LOG_CATEGORY_EXTERN(LogTextureGraphError, Log, All);
+struct FTextureGraphErrorReport
+{
+	int32 ErrorId;
+	FString ErrorMsg;
+	UObject* ReferenceObj = nullptr;
+	
+	FString GetFormattedMessage()
+	{
+		return FString::Format(TEXT("({0}) {1}"), { ErrorId, *ErrorMsg });
+	}
+};
+using ErrorReportMap = TMap<int32, TArray<FTextureGraphErrorReport>>;
+enum class ETextureGraphErrorType: int32
+{
+	UNSUPPORTED_TYPE,
+	MISSING_REQUIRED_INPUT,
+	RECURSIVE_CALL,
+	UNSUPPORTED_MATERIAL,
+	SUBGRAPH_INTERNAL_ERROR,
+	NODE_WARNING
+};
+class TEXTUREGRAPHENGINE_API FTextureGraphErrorReporter
+{
+public:
+	FTextureGraphErrorReporter()
+	{
+	}
+	
+	virtual ~FTextureGraphErrorReporter() { }
+	
+	virtual FTextureGraphErrorReport ReportLog(int32 ErrorId, const FString& ErrorMsg, UObject* ReferenceObj = nullptr);
+
+	virtual FTextureGraphErrorReport ReportWarning(int32 ErrorId, const FString& ErrorMsg, UObject* ReferenceObj = nullptr);
+
+	virtual FTextureGraphErrorReport ReportError(int32 ErrorId, const FString& ErrorMsg, UObject* ReferenceObj = nullptr);
+	
+	virtual void Clear()
+	{
+		CompilationErrors.Empty();
+	}
+protected:
+	ErrorReportMap CompilationErrors;
+
+public:
+	ErrorReportMap GetCompilationErrors() const
+	{
+		return CompilationErrors;
+	}
+};
diff --git a/Engine/Plugins/Experimental/TextureGraph/Source/TextureGraphEngine/Transform/Expressions/T_Blend.cpp b/Engine/Plugins/Experimental/TextureGraph/Source/TextureGraphEngine/Transform/Expressions/T_Blend.cpp
new file mode 100644
--- /dev/null
+++ b/Engine/Plugins/Experimental/TextureGraph/Source/TextureGraphEngine/Transform/Expressions/T_Blend.cpp
@@ -0,0 +1,160 @@
+// Copyright Epic Games, Inc. All Rights Reserved.
+
+#include "T_Blend.h"
+#include "Job/JobArgs.h"
+#include "TextureGraphEngine.h"
+#include "FxMat/MaterialManager.h"
+#include "Job/JobBatch.h"
+#include "2D/TargetTextureSet.h"
+#include "Model/Mix/MixInterface.h"
+#include "Model/Mix/MixSettings.h"
+
+IMPLEMENT_GLOBAL_SHADER(FSH_BlendNormal		, "/Plugin/TextureGraph/Expressions/Expression_Blend.usf", "FSH_BlendNormal"	, SF_Pixel);
+IMPLEMENT_GLOBAL_SHADER(FSH_BlendAdd		, "/Plugin/TextureGraph/Expressions/Expression_Blend.usf", "FSH_BlendAdd"		, SF_Pixel);
+IMPLEMENT_GLOBAL_SHADER(FSH_BlendSubtract	, "/Plugin/TextureGraph/Expressions/Expression_Blend.usf", "FSH_BlendSubtract"	, SF_Pixel);
+IMPLEMENT_GLOBAL_SHADER(FSH_BlendMultiply	, "/Plugin/TextureGraph/Expressions/Expression_Blend.usf", "FSH_BlendMultiply"	, SF_Pixel);
+IMPLEMENT_GLOBAL_SHADER(FSH_BlendDivide		, "/Plugin/TextureGraph/Expressions/Expression_Blend.usf", "FSH_BlendDivide"	, SF_Pixel);
+IMPLEMENT_GLOBAL_SHADER(FSH_BlendDifference	, "/Plugin/TextureGraph/Expressions/Expression_Blend.usf", "FSH_BlendDifference", SF_Pixel);
+IMPLEMENT_GLOBAL_SHADER(FSH_BlendMax		, "/Plugin/TextureGraph/Expressions/Expression_Blend.usf", "FSH_BlendMax"		, SF_Pixel);
+IMPLEMENT_GLOBAL_SHADER(FSH_BlendMin		, "/Plugin/TextureGraph/Expressions/Expression_Blend.usf", "FSH_BlendMin"		, SF_Pixel);
+IMPLEMENT_GLOBAL_SHADER(FSH_BlendStep		, "/Plugin/TextureGraph/Expressions/Expression_Blend.usf", "FSH_BlendStep"		, SF_Pixel);
+IMPLEMENT_GLOBAL_SHADER(FSH_BlendOverlay	, "/Plugin/TextureGraph/Expressions/Expression_Blend.usf", "FSH_BlendOverlay"	, SF_Pixel);
+IMPLEMENT_GLOBAL_SHADER(FSH_BlendDistort	, "/Plugin/TextureGraph/Expressions/Expression_Blend.usf", "FSH_BlendDistort"	, SF_Pixel);
+
+T_Blend::T_Blend()
+{
+}
+
+T_Blend::~T_Blend()
+{
+}
+
+TiledBlobPtr T_Blend::Create(MixUpdateCyclePtr InCycle, BufferDescriptor DesiredDesc, TiledBlobPtr InBackgroundTexture,
+	TiledBlobPtr InForeGroundTexture, TiledBlobPtr InMask, float InOpacity, int InTargetId, EBlendModes InBlendMode)
+{
+	switch(InBlendMode)
+	{
+		case EBlendModes::Normal:
+			return CreateNormal(InCycle, DesiredDesc, InBackgroundTexture, InForeGroundTexture, InMask, InOpacity, InTargetId);
+		case EBlendModes::Add:
+			return CreateAdd(InCycle, DesiredDesc, InBackgroundTexture, InForeGroundTexture, InMask, InOpacity, InTargetId);
+		case EBlendModes::Subtract:
+			return CreateSubtract(InCycle, DesiredDesc, InBackgroundTexture, InForeGroundTexture, InMask, InOpacity, InTargetId);
+		case EBlendModes::Multiply:
+			return CreateMultiply(InCycle, DesiredDesc, InBackgroundTexture, InForeGroundTexture, InMask, InOpacity, InTargetId);
+		case EBlendModes::Divide:
+			return CreateDivide(InCycle, DesiredDesc, InBackgroundTexture, InForeGroundTexture, InMask, InOpacity, InTargetId);
+		case EBlendModes::Difference:
+			return CreateDifference(InCycle, DesiredDesc, InBackgroundTexture, InForeGroundTexture, InMask, InOpacity, InTargetId);
+		case EBlendModes::Max:
+			return CreateMax(InCycle, DesiredDesc, InBackgroundTexture, InForeGroundTexture, InMask, InOpacity, InTargetId);
+		case EBlendModes::Min:
+			return CreateMin(InCycle, DesiredDesc, InBackgroundTexture, InForeGroundTexture, InMask, InOpacity, InTargetId);
+		case EBlendModes::Step:
+			return CreateStep(InCycle, DesiredDesc, InBackgroundTexture, InForeGroundTexture, InMask, InOpacity, InTargetId);
+		case EBlendModes::Overlay:
+			return CreateOverlay(InCycle, DesiredDesc, InBackgroundTexture, InForeGroundTexture, InMask, InOpacity, InTargetId);
+		// case EBlendModes::Distort:
+		// 	return CreateDistort(InCycle, DesiredDesc, InBackgroundTexture, InForeGroundTexture, InMask, InOpacity, InTargetId);
+
+	default:
+		// Unhandled case
+		checkNoEntry();
+		return TextureHelper::GetBlack();
+	}
+}
+
+template <typename FSH_Type>
+TiledBlobPtr CreateGenericBlend(MixUpdateCyclePtr InCycle, int32 InTargetId, BufferDescriptor DesiredDesc, TiledBlobPtr InBackgroundTexture,
+	TiledBlobPtr InForegroundTexture, TiledBlobPtr InMask, float InOpacity, FString InTransformName)
+{
+	const RenderMaterial_FXPtr RenderMaterial = TextureGraphEngine::GetMaterialManager()->CreateMaterial_FX<VSH_Simple, FSH_Type>(InTransformName);
+	check(RenderMaterial);
+
+	if (!InBackgroundTexture)
+	{
+		InBackgroundTexture = TextureHelper::GetBlack();
+	}
+	
+	if (!InForegroundTexture)
+	{
+		InForegroundTexture = TextureHelper::GetBlack();
+	}
+	
+	if(!InMask)
+	{
+		InMask = TextureHelper::GetWhite();
+	}  
+
+	JobUPtr JobPtr = std::make_unique<Job>(InCycle->GetMix(), InTargetId, std::static_pointer_cast<BlobTransform>(RenderMaterial));
+	
+	JobPtr
+		->AddArg(ARG_BLOB(InBackgroundTexture, "BackgroundTexture"))
+		->AddArg(ARG_BLOB(InForegroundTexture, "ForegroundTexture"))
+		->AddArg(ARG_BLOB(InMask, "MaskTexture"))
+		->AddArg(ARG_FLOAT(InOpacity, "Opacity"))
+		;
+
+	const FString Name = FString::Printf(TEXT("[%llu] - Blend - %s"), InCycle->GetBatch()->GetBatchId(), *InTransformName);
+
+	TiledBlobPtr Result = JobPtr->InitResult(Name, &DesiredDesc);
+
+	InCycle->AddJob(InTargetId, std::move(JobPtr));
+
+	return Result;
+}
+
+TiledBlobPtr T_Blend::CreateNormal(MixUpdateCyclePtr InCycle, BufferDescriptor DesiredDesc, TiledBlobPtr InBackgroundTexture, TiledBlobPtr InForegroundTexture, TiledBlobPtr InMask, float InOpacity, int InTargetId)
+{
+	return CreateGenericBlend<FSH_BlendNormal>(InCycle, InTargetId, DesiredDesc, InBackgroundTexture, InForegroundTexture, InMask,InOpacity, "T_BlendNormal");
+}
+
+TiledBlobPtr T_Blend::CreateAdd(MixUpdateCyclePtr InCycle, BufferDescriptor DesiredDesc, TiledBlobPtr InBackgroundTexture, TiledBlobPtr InForegroundTexture, TiledBlobPtr InMask, float InOpacity, int InTargetId)
+{
+	return CreateGenericBlend<FSH_BlendAdd>(InCycle, InTargetId, DesiredDesc, InBackgroundTexture, InForegroundTexture, InMask,InOpacity, "T_BlendAdd");
+}
+
+TiledBlobPtr T_Blend::CreateSubtract(MixUpdateCyclePtr InCycle, BufferDescriptor DesiredDesc, TiledBlobPtr InBackgroundTexture, TiledBlobPtr InForegroundTexture, TiledBlobPtr InMask, float InOpacity, int InTargetId)
+{
+	return CreateGenericBlend<FSH_BlendSubtract>(InCycle, InTargetId, DesiredDesc, InBackgroundTexture, InForegroundTexture, InMask,InOpacity, "T_BlendSubtract");
+}
+
+TiledBlobPtr T_Blend::CreateMultiply(MixUpdateCyclePtr InCycle, BufferDescriptor DesiredDesc, TiledBlobPtr InBackgroundTexture, TiledBlobPtr InForeGroundTexture, TiledBlobPtr InMask, float InOpacity, int InTargetId)
+{
+	return CreateGenericBlend<FSH_BlendMultiply>(InCycle, InTargetId, DesiredDesc, InBackgroundTexture, InForeGroundTexture, InMask,InOpacity, "T_BlendMultiply");
+}
+
+TiledBlobPtr T_Blend::CreateDivide(MixUpdateCyclePtr InCycle, BufferDescriptor DesiredDesc, TiledBlobPtr InBackgroundTexture, TiledBlobPtr InForegroundTexture, TiledBlobPtr InMask, float InOpacity, int InTargetId)
+{
+	return CreateGenericBlend<FSH_BlendDivide>(InCycle, InTargetId, DesiredDesc, InBackgroundTexture, InForegroundTexture, InMask,InOpacity, "T_BlendDivide");
+}
+
+TiledBlobPtr T_Blend::CreateDifference(MixUpdateCyclePtr InCycle, BufferDescriptor DesiredDesc, TiledBlobPtr InBackgroundTexture, TiledBlobPtr InForegroundTexture, TiledBlobPtr InMask, float InOpacity, int InTargetId)
+{
+	return CreateGenericBlend<FSH_BlendDifference>(InCycle, InTargetId, DesiredDesc, InBackgroundTexture, InForegroundTexture, InMask,InOpacity, "T_BlendDifference");
+}
+
+TiledBlobPtr T_Blend::CreateMax(MixUpdateCyclePtr InCycle, BufferDescriptor DesiredDesc, TiledBlobPtr InBackgroundTexture, TiledBlobPtr InForegroundTexture, TiledBlobPtr InMask, float InOpacity, int InTargetId)
+{
+	return CreateGenericBlend<FSH_BlendMax>(InCycle, InTargetId, DesiredDesc, InBackgroundTexture, InForegroundTexture, InMask,InOpacity, "T_BlendMax");
+}
+
+TiledBlobPtr T_Blend::CreateMin(MixUpdateCyclePtr InCycle, BufferDescriptor DesiredDesc, TiledBlobPtr InBackgroundTexture, TiledBlobPtr InForegroundTexture, TiledBlobPtr InMask, float InOpacity, int InTargetId)
+{
+	return CreateGenericBlend<FSH_BlendMin>(InCycle, InTargetId, DesiredDesc, InBackgroundTexture, InForegroundTexture, InMask,InOpacity, "T_BlendMin");
+}
+
+TiledBlobPtr T_Blend::CreateStep(MixUpdateCyclePtr InCycle, BufferDescriptor DesiredDesc, TiledBlobPtr InBackgroundTexture, TiledBlobPtr InForegroundTexture, TiledBlobPtr InMask, float InOpacity, int InTargetId)
+{
+	return CreateGenericBlend<FSH_BlendStep>(InCycle, InTargetId, DesiredDesc, InBackgroundTexture, InForegroundTexture, InMask,InOpacity, "T_BlendStep");
+}
+
+TiledBlobPtr T_Blend::CreateOverlay(MixUpdateCyclePtr InCycle, BufferDescriptor DesiredDesc, TiledBlobPtr InBackgroundTexture, TiledBlobPtr InForegroundTexture, TiledBlobPtr InMask, float InOpacity, int InTargetId)
+{
+	return CreateGenericBlend<FSH_BlendOverlay>(InCycle, InTargetId, DesiredDesc, InBackgroundTexture, InForegroundTexture, InMask,InOpacity, "T_BlendOverlay");
+}
+
+TiledBlobPtr T_Blend::CreateDistort(MixUpdateCyclePtr InCycle, BufferDescriptor DesiredDesc, TiledBlobPtr InBackgroundTexture, TiledBlobPtr InForegroundTexture, TiledBlobPtr InMask, float InOpacity, int InTargetId)
+{
+	return CreateGenericBlend<FSH_BlendDistort>(InCycle, InTargetId, DesiredDesc, InBackgroundTexture, InForegroundTexture, InMask,InOpacity, "T_BlendDistort");
+}
\ No newline at end of file
diff --git a/Engine/Plugins/Experimental/TextureGraph/Source/TextureGraphEngine/Transform/Expressions/T_Blend.h b/Engine/Plugins/Experimental/TextureGraph/Source/TextureGraphEngine/Transform/Expressions/T_Blend.h
new file mode 100644
--- /dev/null
+++ b/Engine/Plugins/Experimental/TextureGraph/Source/TextureGraphEngine/Transform/Expressions/T_Blend.h
@@ -0,0 +1,172 @@
+// Copyright Epic Games, Inc. All Rights Reserved.
+#pragma once
+
+#include <2D/BlendModes.h>
+
+#include "CoreMinimal.h"
+#include "Job/Job.h"
+#include "FxMat/FxMaterial.h"
+#include "Model/Mix/MixUpdateCycle.h"
+#include <DataDrivenShaderPlatformInfo.h>
+
+//////////////////////////////////////////////////////////////////////////
+/// Shader
+//////////////////////////////////////////////////////////////////////////
+
+class TEXTUREGRAPHENGINE_API FSH_BlendBase : public FSH_Base
+{
+public:
+	SHADER_USE_PARAMETER_STRUCT(FSH_BlendBase, FSH_Base);
+	BEGIN_SHADER_PARAMETER_STRUCT(FParameters, )
+		SHADER_PARAMETER_TEXTURE(Texture2D, BackgroundTexture)
+		SHADER_PARAMETER_TEXTURE(Texture2D, ForegroundTexture)
+		SHADER_PARAMETER_TEXTURE(Texture2D, MaskTexture)
+		SHADER_PARAMETER_STRUCT(FStandardSamplerStates, SamplerStates)
+		SHADER_PARAMETER(float, Opacity)
+	END_SHADER_PARAMETER_STRUCT()
+
+public:
+	static bool ShouldCompilePermutation(const FGlobalShaderPermutationParameters& Parameters)
+	{
+		return IsPCPlatform(Parameters.Platform) && EnumHasAllFlags(Parameters.Flags, EShaderPermutationFlags::HasEditorOnlyData);
+	}
+
+	static void ModifyCompilationEnvironment(const FGlobalShaderPermutationParameters& InParams, FShaderCompilerEnvironment& InEnv)
+	{
+	}
+};
+
+//////////////////////////////////////////////////////////////////////////
+/// Normal
+//////////////////////////////////////////////////////////////////////////
+class TEXTUREGRAPHENGINE_API FSH_BlendNormal : public FSH_BlendBase
+{
+public:
+	DECLARE_GLOBAL_SHADER(FSH_BlendNormal);
+	SHADER_USE_PARAMETER_STRUCT(FSH_BlendNormal, FSH_BlendBase);
+};
+
+//////////////////////////////////////////////////////////////////////////
+/// Add
+//////////////////////////////////////////////////////////////////////////
+class TEXTUREGRAPHENGINE_API FSH_BlendAdd : public FSH_BlendBase
+{
+public:
+	DECLARE_GLOBAL_SHADER(FSH_BlendAdd);
+	SHADER_USE_PARAMETER_STRUCT(FSH_BlendAdd, FSH_BlendBase);
+};
+
+//////////////////////////////////////////////////////////////////////////
+/// Add
+//////////////////////////////////////////////////////////////////////////
+class TEXTUREGRAPHENGINE_API FSH_BlendSubtract : public FSH_BlendBase
+{
+public:
+	DECLARE_GLOBAL_SHADER(FSH_BlendSubtract);
+	SHADER_USE_PARAMETER_STRUCT(FSH_BlendSubtract, FSH_BlendBase);
+};
+
+//////////////////////////////////////////////////////////////////////////
+/// Multiply
+//////////////////////////////////////////////////////////////////////////
+class TEXTUREGRAPHENGINE_API FSH_BlendMultiply : public FSH_BlendBase
+{
+public:
+	DECLARE_GLOBAL_SHADER(FSH_BlendMultiply);
+	SHADER_USE_PARAMETER_STRUCT(FSH_BlendMultiply, FSH_BlendBase);
+};
+
+//////////////////////////////////////////////////////////////////////////
+/// Divide
+//////////////////////////////////////////////////////////////////////////
+class TEXTUREGRAPHENGINE_API FSH_BlendDivide : public FSH_BlendBase
+{
+public:
+	DECLARE_GLOBAL_SHADER(FSH_BlendDivide);
+	SHADER_USE_PARAMETER_STRUCT(FSH_BlendDivide, FSH_BlendBase);
+};
+//////////////////////////////////////////////////////////////////////////
+/// Difference
+//////////////////////////////////////////////////////////////////////////
+class TEXTUREGRAPHENGINE_API FSH_BlendDifference : public FSH_BlendBase
+{
+public:
+	DECLARE_GLOBAL_SHADER(FSH_BlendDifference);
+	SHADER_USE_PARAMETER_STRUCT(FSH_BlendDifference, FSH_BlendBase);
+};
+
+//////////////////////////////////////////////////////////////////////////
+/// Max
+//////////////////////////////////////////////////////////////////////////
+class TEXTUREGRAPHENGINE_API FSH_BlendMax : public FSH_BlendBase
+{
+public:
+	DECLARE_GLOBAL_SHADER(FSH_BlendMax);
+	SHADER_USE_PARAMETER_STRUCT(FSH_BlendMax, FSH_BlendBase);
+};
+
+//////////////////////////////////////////////////////////////////////////
+/// Min
+//////////////////////////////////////////////////////////////////////////
+class TEXTUREGRAPHENGINE_API FSH_BlendMin : public FSH_BlendBase
+{
+public:
+	DECLARE_GLOBAL_SHADER(FSH_BlendMin);
+	SHADER_USE_PARAMETER_STRUCT(FSH_BlendMin, FSH_BlendBase);
+};
+
+//////////////////////////////////////////////////////////////////////////
+/// Step
+//////////////////////////////////////////////////////////////////////////
+class TEXTUREGRAPHENGINE_API FSH_BlendStep : public FSH_BlendBase
+{
+public:
+	DECLARE_GLOBAL_SHADER(FSH_BlendStep);
+	SHADER_USE_PARAMETER_STRUCT(FSH_BlendStep, FSH_BlendBase);
+};
+
+//////////////////////////////////////////////////////////////////////////
+/// Overlay
+//////////////////////////////////////////////////////////////////////////
+class TEXTUREGRAPHENGINE_API FSH_BlendOverlay : public FSH_BlendBase
+{
+public:
+	DECLARE_GLOBAL_SHADER(FSH_BlendOverlay);
+	SHADER_USE_PARAMETER_STRUCT(FSH_BlendOverlay, FSH_BlendBase);
+};
+
+//////////////////////////////////////////////////////////////////////////
+/// Distort
+//////////////////////////////////////////////////////////////////////////
+class TEXTUREGRAPHENGINE_API FSH_BlendDistort : public FSH_BlendBase
+{
+public:
+	DECLARE_GLOBAL_SHADER(FSH_BlendDistort);
+	SHADER_USE_PARAMETER_STRUCT(FSH_BlendDistort, FSH_BlendBase);
+};
+/**
+ * 
+ */
+class TEXTUREGRAPHENGINE_API T_Blend
+{
+public:
+	T_Blend();
+	~T_Blend();
+
+	//////////////////////////////////////////////////////////////////////////
+	/// Static functions
+	//////////////////////////////////////////////////////////////////////////
+
+	static TiledBlobPtr				Create(MixUpdateCyclePtr InCycle, BufferDescriptor DesiredDesc, TiledBlobPtr InBackgroundTexture, TiledBlobPtr InForeGroundTexture, TiledBlobPtr InMask, float InOpacity, int InTargetId, EBlendModes InBlendMode);
+	static TiledBlobPtr				CreateNormal(MixUpdateCyclePtr InCycle, BufferDescriptor DesiredDesc, TiledBlobPtr InBackgroundTexture, TiledBlobPtr InForeGroundTexture, TiledBlobPtr InMask, float InOpacity, int InTargetId);
+	static TiledBlobPtr				CreateAdd(MixUpdateCyclePtr InCycle, BufferDescriptor DesiredDesc, TiledBlobPtr InBackgroundTexture, TiledBlobPtr InForeGroundTexture, TiledBlobPtr InMask, float InOpacity, int InTargetId);
+	static TiledBlobPtr				CreateSubtract(MixUpdateCyclePtr InCycle, BufferDescriptor DesiredDesc, TiledBlobPtr InBackgroundTexture, TiledBlobPtr InForeGroundTexture, TiledBlobPtr InMask, float InOpacity, int InTargetId);
+	static TiledBlobPtr				CreateMultiply(MixUpdateCyclePtr InCycle, BufferDescriptor DesiredDesc, TiledBlobPtr InBackgroundTexture, TiledBlobPtr InForeGroundTexture, TiledBlobPtr InMask, float InOpacity, int InTargetId);
+	static TiledBlobPtr				CreateDivide(MixUpdateCyclePtr InCycle, BufferDescriptor DesiredDesc, TiledBlobPtr InBackgroundTexture, TiledBlobPtr InForeGroundTexture, TiledBlobPtr InMask, float InOpacity, int InTargetId);
+	static TiledBlobPtr				CreateDifference(MixUpdateCyclePtr InCycle, BufferDescriptor DesiredDesc, TiledBlobPtr InBackgroundTexture, TiledBlobPtr InForeGroundTexture, TiledBlobPtr InMask, float InOpacity, int InTargetId);
+	static TiledBlobPtr				CreateMax(MixUpdateCyclePtr InCycle, BufferDescriptor DesiredDesc, TiledBlobPtr InBackgroundTexture, TiledBlobPtr InForeGroundTexture, TiledBlobPtr InMask, float InOpacity, int InTargetId);
+	static TiledBlobPtr				CreateMin(MixUpdateCyclePtr InCycle, BufferDescriptor DesiredDesc, TiledBlobPtr InBackgroundTexture, TiledBlobPtr InForeGroundTexture, TiledBlobPtr InMask, float InOpacity, int InTargetId);
+	static TiledBlobPtr				CreateStep(MixUpdateCyclePtr InCycle, BufferDescriptor DesiredDesc, TiledBlobPtr InBackgroundTexture, TiledBlobPtr InForeGroundTexture, TiledBlobPtr InMask, float InOpacity, int InTargetId);
+	static TiledBlobPtr				CreateOverlay(MixUpdateCyclePtr InCycle, BufferDescriptor DesiredDesc, TiledBlobPtr InBackgroundTexture, TiledBlobPtr InForeGroundTexture, TiledBlobPtr InMask, float InOpacity, int InTargetId);
+	static TiledBlobPtr				CreateDistort(MixUpdateCyclePtr InCycle, BufferDescriptor DesiredDesc, TiledBlobPtr InBackgroundTexture, TiledBlobPtr InForeGroundTexture, TiledBlobPtr InMask, float InOpacity, int InTargetId);
+};
diff --git a/Engine/Plugins/Experimental/TextureGraph/Source/TextureGraphEngine/Transform/Expressions/T_Channel.h b/Engine/Plugins/Experimental/TextureGraph/Source/TextureGraphEngine/Transform/Expressions/T_Channel.h
new file mode 100644
--- /dev/null
+++ b/Engine/Plugins/Experimental/TextureGraph/Source/TextureGraphEngine/Transform/Expressions/T_Channel.h
@@ -0,0 +1,100 @@
+// Copyright Epic Games, Inc. All Rights Reserved.
+#pragma once
+
+#include "CoreMinimal.h"
+#include "Job/Job.h"
+#include "FxMat/FxMaterial.h"
+#include "Model/Mix/MixUpdateCycle.h"
+#include <DataDrivenShaderPlatformInfo.h>
+
+//////////////////////////////////////////////////////////////////////////
+/// Basic Channel Op
+//////////////////////////////////////////////////////////////////////////
+class TEXTUREGRAPHENGINE_API FSH_ChannelSplitter : public FSH_Base
+{
+public:
+	SHADER_USE_PARAMETER_STRUCT(FSH_ChannelSplitter, FSH_Base);
+
+	BEGIN_SHADER_PARAMETER_STRUCT(FParameters, )
+		SHADER_PARAMETER_STRUCT(FStandardSamplerStates, SamplerStates)
+		SHADER_PARAMETER_TEXTURE(Texture2D, SourceTexture)
+	END_SHADER_PARAMETER_STRUCT()
+
+	static bool						ShouldCompilePermutation(const FGlobalShaderPermutationParameters& Parameters) 
+	{
+		return IsPCPlatform(Parameters.Platform) && EnumHasAllFlags(Parameters.Flags, EShaderPermutationFlags::HasEditorOnlyData);
+	}
+	static void						ModifyCompilationEnvironment(const FGlobalShaderPermutationParameters& params, FShaderCompilerEnvironment& env) {}
+};
+
+//////////////////////////////////////////////////////////////////////////
+/// Red
+//////////////////////////////////////////////////////////////////////////
+class TEXTUREGRAPHENGINE_API FSH_ChannelSplitter_Red : public FSH_ChannelSplitter
+{
+public:
+	DECLARE_GLOBAL_SHADER(FSH_ChannelSplitter_Red);
+	SHADER_USE_PARAMETER_STRUCT(FSH_ChannelSplitter_Red, FSH_ChannelSplitter);
+};
+
+//////////////////////////////////////////////////////////////////////////
+/// Green
+//////////////////////////////////////////////////////////////////////////
+class TEXTUREGRAPHENGINE_API FSH_ChannelSplitter_Green : public FSH_ChannelSplitter
+{
+public:
+	DECLARE_GLOBAL_SHADER(FSH_ChannelSplitter_Green);
+	SHADER_USE_PARAMETER_STRUCT(FSH_ChannelSplitter_Green, FSH_ChannelSplitter);
+};
+
+//////////////////////////////////////////////////////////////////////////
+/// Red
+//////////////////////////////////////////////////////////////////////////
+class TEXTUREGRAPHENGINE_API FSH_ChannelSplitter_Blue : public FSH_ChannelSplitter
+{
+public:
+	DECLARE_GLOBAL_SHADER(FSH_ChannelSplitter_Blue);
+	SHADER_USE_PARAMETER_STRUCT(FSH_ChannelSplitter_Blue, FSH_ChannelSplitter);
+};
+
+//////////////////////////////////////////////////////////////////////////
+/// Red
+//////////////////////////////////////////////////////////////////////////
+class TEXTUREGRAPHENGINE_API FSH_ChannelSplitter_Alpha : public FSH_ChannelSplitter
+{
+public:
+	DECLARE_GLOBAL_SHADER(FSH_ChannelSplitter_Alpha);
+	SHADER_USE_PARAMETER_STRUCT(FSH_ChannelSplitter_Alpha, FSH_ChannelSplitter);
+};
+
+//////////////////////////////////////////////////////////////////////////
+/// Basic Channel Op
+//////////////////////////////////////////////////////////////////////////
+class TEXTUREGRAPHENGINE_API FSH_ChannelCombiner : public FSH_Base
+{
+public:
+	DECLARE_GLOBAL_SHADER(FSH_ChannelCombiner);
+	SHADER_USE_PARAMETER_STRUCT(FSH_ChannelCombiner, FSH_Base);
+
+	BEGIN_SHADER_PARAMETER_STRUCT(FParameters, )
+		SHADER_PARAMETER_STRUCT(FStandardSamplerStates, SamplerStates)
+		SHADER_PARAMETER_TEXTURE(Texture2D, SourceRed)
+		SHADER_PARAMETER_TEXTURE(Texture2D, SourceGreen)
+		SHADER_PARAMETER_TEXTURE(Texture2D, SourceBlue)
+		SHADER_PARAMETER_TEXTURE(Texture2D, SourceAlpha)
+	END_SHADER_PARAMETER_STRUCT()
+
+	static bool						ShouldCompilePermutation(const FGlobalShaderPermutationParameters& Parameters) { return true; }
+	static void						ModifyCompilationEnvironment(const FGlobalShaderPermutationParameters& params, FShaderCompilerEnvironment& env) {}
+};
+
+//////////////////////////////////////////////////////////////////////////
+/// Helper C++ class
+//////////////////////////////////////////////////////////////////////////
+class TEXTUREGRAPHENGINE_API T_Channel
+{
+public:
+	//////////////////////////////////////////////////////////////////////////
+	/// Static functions
+	//////////////////////////////////////////////////////////////////////////
+};
diff --git a/Engine/Plugins/Experimental/TextureGraph/Source/TextureGraphEngine/Transform/Expressions/T_Color.h b/Engine/Plugins/Experimental/TextureGraph/Source/TextureGraphEngine/Transform/Expressions/T_Color.h
new file mode 100644
--- /dev/null
+++ b/Engine/Plugins/Experimental/TextureGraph/Source/TextureGraphEngine/Transform/Expressions/T_Color.h
@@ -0,0 +1,153 @@
+// Copyright Epic Games, Inc. All Rights Reserved.
+#pragma once
+
+#include "CoreMinimal.h"
+#include "Job/Job.h"
+#include "FxMat/FxMaterial.h"
+#include "Model/Mix/MixUpdateCycle.h"
+#include <DataDrivenShaderPlatformInfo.h>
+
+//////////////////////////////////////////////////////////////////////////
+/// Simple grayscale shader
+//////////////////////////////////////////////////////////////////////////
+class TEXTUREGRAPHENGINE_API FSH_Grayscale : public FSH_Base
+{
+public:
+	SHADER_USE_PARAMETER_STRUCT(FSH_Grayscale, FSH_Base);
+	DECLARE_GLOBAL_SHADER(FSH_Grayscale);
+
+	BEGIN_SHADER_PARAMETER_STRUCT(FParameters, )
+		SHADER_PARAMETER_STRUCT(FStandardSamplerStates, SamplerStates)
+		SHADER_PARAMETER_TEXTURE(Texture2D, SourceTexture)
+	END_SHADER_PARAMETER_STRUCT()
+
+	static bool						ShouldCompilePermutation(const FGlobalShaderPermutationParameters& Parameters) 
+	{
+		return IsPCPlatform(Parameters.Platform) && EnumHasAllFlags(Parameters.Flags, EShaderPermutationFlags::HasEditorOnlyData);
+	}
+	static void						ModifyCompilationEnvironment(const FGlobalShaderPermutationParameters& params, FShaderCompilerEnvironment& env) {}
+};
+
+//////////////////////////////////////////////////////////////////////////
+/// Simple levels shader
+//////////////////////////////////////////////////////////////////////////
+class TEXTUREGRAPHENGINE_API FSH_Levels : public FSH_Base
+{
+public:
+	SHADER_USE_PARAMETER_STRUCT(FSH_Levels, FSH_Base);
+	DECLARE_GLOBAL_SHADER(FSH_Levels);
+
+	BEGIN_SHADER_PARAMETER_STRUCT(FParameters, )
+		SHADER_PARAMETER_STRUCT(FStandardSamplerStates, SamplerStates)
+		SHADER_PARAMETER_TEXTURE(Texture2D, Input)
+		SHADER_PARAMETER(float, MinValue)
+		SHADER_PARAMETER(float, MaxValue)
+		SHADER_PARAMETER(float, Gamma)
+	END_SHADER_PARAMETER_STRUCT()
+
+	static bool						ShouldCompilePermutation(const FGlobalShaderPermutationParameters& Parameters) 
+	{
+		return IsPCPlatform(Parameters.Platform) && EnumHasAllFlags(Parameters.Flags, EShaderPermutationFlags::HasEditorOnlyData);
+	}
+	static void						ModifyCompilationEnvironment(const FGlobalShaderPermutationParameters& params, FShaderCompilerEnvironment& env) {}
+};
+
+
+//////////////////////////////////////////////////////////////////////////
+/// Simple levels shader
+//////////////////////////////////////////////////////////////////////////
+class TEXTUREGRAPHENGINE_API FSH_Threshold : public FSH_Base
+{
+public:
+	SHADER_USE_PARAMETER_STRUCT(FSH_Threshold, FSH_Base);
+	DECLARE_GLOBAL_SHADER(FSH_Threshold);
+
+	BEGIN_SHADER_PARAMETER_STRUCT(FParameters, )
+		SHADER_PARAMETER_STRUCT(FStandardSamplerStates, SamplerStates)
+		SHADER_PARAMETER_TEXTURE(Texture2D, Input)
+		SHADER_PARAMETER(float, Threshold)
+		END_SHADER_PARAMETER_STRUCT()
+
+		static bool						ShouldCompilePermutation(const FGlobalShaderPermutationParameters& Parameters)
+	{
+		return IsPCPlatform(Parameters.Platform) && EnumHasAllFlags(Parameters.Flags, EShaderPermutationFlags::HasEditorOnlyData);
+	}
+	static void						ModifyCompilationEnvironment(const FGlobalShaderPermutationParameters& params, FShaderCompilerEnvironment& env) {}
+};
+
+//////////////////////////////////////////////////////////////////////////
+/// Simple HSV shader
+//////////////////////////////////////////////////////////////////////////
+class TEXTUREGRAPHENGINE_API FSH_HSV : public FSH_Base
+{
+public:
+	SHADER_USE_PARAMETER_STRUCT(FSH_HSV, FSH_Base);
+	DECLARE_GLOBAL_SHADER(FSH_HSV);
+
+	BEGIN_SHADER_PARAMETER_STRUCT(FParameters, )
+		SHADER_PARAMETER_STRUCT(FStandardSamplerStates, SamplerStates)
+		SHADER_PARAMETER_TEXTURE(Texture2D, Input)
+		SHADER_PARAMETER(float, Hue)
+		SHADER_PARAMETER(float, Saturation)
+		SHADER_PARAMETER(float, Value)
+	END_SHADER_PARAMETER_STRUCT()
+
+	static bool						ShouldCompilePermutation(const FGlobalShaderPermutationParameters& Parameters) 
+	{
+		return IsPCPlatform(Parameters.Platform) && EnumHasAllFlags(Parameters.Flags, EShaderPermutationFlags::HasEditorOnlyData);
+	}
+	static void						ModifyCompilationEnvironment(const FGlobalShaderPermutationParameters& params, FShaderCompilerEnvironment& env) {}
+};
+
+//////////////////////////////////////////////////////////////////////////
+/// Simple RGB2HSV shader
+//////////////////////////////////////////////////////////////////////////
+class TEXTUREGRAPHENGINE_API FSH_RGB2HSV : public FSH_Base
+{
+public:
+	SHADER_USE_PARAMETER_STRUCT(FSH_RGB2HSV, FSH_Base);
+	DECLARE_GLOBAL_SHADER(FSH_RGB2HSV);
+
+	BEGIN_SHADER_PARAMETER_STRUCT(FParameters, )
+		SHADER_PARAMETER_STRUCT(FStandardSamplerStates, SamplerStates)
+		SHADER_PARAMETER_TEXTURE(Texture2D, Input)
+	END_SHADER_PARAMETER_STRUCT()
+
+	static bool						ShouldCompilePermutation(const FGlobalShaderPermutationParameters& Parameters) 
+	{
+		return IsPCPlatform(Parameters.Platform) && EnumHasAllFlags(Parameters.Flags, EShaderPermutationFlags::HasEditorOnlyData);
+	}
+	static void						ModifyCompilationEnvironment(const FGlobalShaderPermutationParameters& params, FShaderCompilerEnvironment& env) {}
+};
+
+//////////////////////////////////////////////////////////////////////////
+/// Simple HSV2RGB shader
+//////////////////////////////////////////////////////////////////////////
+class TEXTUREGRAPHENGINE_API FSH_HSV2RGB : public FSH_Base
+{
+public:
+	SHADER_USE_PARAMETER_STRUCT(FSH_HSV2RGB, FSH_Base);
+	DECLARE_GLOBAL_SHADER(FSH_HSV2RGB);
+
+	BEGIN_SHADER_PARAMETER_STRUCT(FParameters, )
+		SHADER_PARAMETER_STRUCT(FStandardSamplerStates, SamplerStates)
+		SHADER_PARAMETER_TEXTURE(Texture2D, Input)
+	END_SHADER_PARAMETER_STRUCT()
+
+	static bool						ShouldCompilePermutation(const FGlobalShaderPermutationParameters& Parameters) 
+	{
+		return IsPCPlatform(Parameters.Platform) && EnumHasAllFlags(Parameters.Flags, EShaderPermutationFlags::HasEditorOnlyData);
+	}
+	static void						ModifyCompilationEnvironment(const FGlobalShaderPermutationParameters& params, FShaderCompilerEnvironment& env) {}
+};
+
+//////////////////////////////////////////////////////////////////////////
+/// Helper C++ class
+//////////////////////////////////////////////////////////////////////////
+class TEXTUREGRAPHENGINE_API T_Color
+{
+public:
+	//////////////////////////////////////////////////////////////////////////
+	/// Static functions
+	//////////////////////////////////////////////////////////////////////////
+};
diff --git a/Engine/Plugins/Experimental/TextureGraph/Source/TextureGraphEngine/Transform/Expressions/T_Invert.cpp b/Engine/Plugins/Experimental/TextureGraph/Source/TextureGraphEngine/Transform/Expressions/T_Invert.cpp
new file mode 100644
--- /dev/null
+++ b/Engine/Plugins/Experimental/TextureGraph/Source/TextureGraphEngine/Transform/Expressions/T_Invert.cpp
@@ -0,0 +1,31 @@
+// Copyright Epic Games, Inc. All Rights Reserved.
+
+#include "T_Invert.h"
+#include "Job/JobArgs.h"
+#include "TextureGraphEngine.h"
+#include "Math/Vector.h"
+#include "FxMat/MaterialManager.h"
+#include "Job/JobBatch.h"
+#include "Helper/GraphicsUtil.h"
+
+IMPLEMENT_GLOBAL_SHADER(FSH_Invert, "/Plugin/TextureGraph/Expressions/Expression_Invert.usf", "FSH_Invert", SF_Pixel);
+
+TiledBlobPtr T_Invert::Create(MixUpdateCyclePtr Cycle, BufferDescriptor DesiredOutputDesc, TiledBlobPtr Source, float MaxValue, int32 TargetId)
+{
+	RenderMaterial_FXPtr RenderMaterial = TextureGraphEngine::GetMaterialManager()->CreateMaterial_FX<VSH_Simple, FSH_Invert>(TEXT("T_Invert"));
+
+	check(RenderMaterial);
+
+	JobUPtr RenderJob = std::make_unique<Job>(Cycle->GetMix(), TargetId, std::static_pointer_cast<BlobTransform>(RenderMaterial));
+	RenderJob
+		->AddArg(ARG_BLOB(Source, "SourceTexture"))
+		->AddArg(ARG_FLOAT(MaxValue, "MaxValue"))
+		;
+
+	FString Name = TEXT("Invert"); // FString::Printf(TEXT("GaussianBlur.[%s].[%d].[%llu]"), *gmask->ID(), InContext->TargetId, InContext->Cycle->Batch()->BatchId());
+
+	TiledBlobPtr Result = RenderJob->InitResult(Name, &DesiredOutputDesc);
+	Cycle->AddJob(TargetId, std::move(RenderJob));
+
+	return Result;
+}
diff --git a/Engine/Plugins/Experimental/TextureGraph/Source/TextureGraphEngine/Transform/Expressions/T_Maths_OneInput.h b/Engine/Plugins/Experimental/TextureGraph/Source/TextureGraphEngine/Transform/Expressions/T_Maths_OneInput.h
new file mode 100644
--- /dev/null
+++ b/Engine/Plugins/Experimental/TextureGraph/Source/TextureGraphEngine/Transform/Expressions/T_Maths_OneInput.h
@@ -0,0 +1,91 @@
+// Copyright Epic Games, Inc. All Rights Reserved.
+#pragma once
+
+#include "CoreMinimal.h"
+#include "Job/Job.h"
+#include "FxMat/FxMaterial.h"
+#include "Model/Mix/MixUpdateCycle.h"
+#include <DataDrivenShaderPlatformInfo.h>
+
+#include "T_Maths_OneInput.generated.h"
+
+//////////////////////////////////////////////////////////////////////////
+/// Basic one input maths op
+//////////////////////////////////////////////////////////////////////////
+class TEXTUREGRAPHENGINE_API FSH_MathsOp_OneInput : public FSH_Base
+{
+public:
+	SHADER_USE_PARAMETER_STRUCT(FSH_MathsOp_OneInput, FSH_Base);
+
+	BEGIN_SHADER_PARAMETER_STRUCT(FParameters, )
+		SHADER_PARAMETER_STRUCT(FStandardSamplerStates, SamplerStates)
+		SHADER_PARAMETER_TEXTURE(Texture2D, Input)
+	END_SHADER_PARAMETER_STRUCT()
+
+	TEXTURE_ENGINE_DEFAULT_PERMUTATION;
+	TEXTUREGRAPH_ENGINE_DEFAULT_COMPILATION_ENV;
+};
+
+DECLARE_EMPTY_GLOBAL_SHADER_DERIVED_FROM(FSH_Sin, FSH_MathsOp_OneInput);
+DECLARE_EMPTY_GLOBAL_SHADER_DERIVED_FROM(FSH_Cos, FSH_MathsOp_OneInput);
+DECLARE_EMPTY_GLOBAL_SHADER_DERIVED_FROM(FSH_Tan, FSH_MathsOp_OneInput);
+DECLARE_EMPTY_GLOBAL_SHADER_DERIVED_FROM(FSH_ASin, FSH_MathsOp_OneInput);
+DECLARE_EMPTY_GLOBAL_SHADER_DERIVED_FROM(FSH_ACos, FSH_MathsOp_OneInput);
+DECLARE_EMPTY_GLOBAL_SHADER_DERIVED_FROM(FSH_ATan, FSH_MathsOp_OneInput);
+
+DECLARE_EMPTY_GLOBAL_SHADER_DERIVED_FROM(FSH_ToRadians, FSH_MathsOp_OneInput);
+DECLARE_EMPTY_GLOBAL_SHADER_DERIVED_FROM(FSH_ToDegrees, FSH_MathsOp_OneInput);
+
+DECLARE_EMPTY_GLOBAL_SHADER_DERIVED_FROM(FSH_Abs, FSH_MathsOp_OneInput);
+DECLARE_EMPTY_GLOBAL_SHADER_DERIVED_FROM(FSH_Sqrt, FSH_MathsOp_OneInput);
+DECLARE_EMPTY_GLOBAL_SHADER_DERIVED_FROM(FSH_Square, FSH_MathsOp_OneInput);
+DECLARE_EMPTY_GLOBAL_SHADER_DERIVED_FROM(FSH_Cube, FSH_MathsOp_OneInput);
+DECLARE_EMPTY_GLOBAL_SHADER_DERIVED_FROM(FSH_Cbrt, FSH_MathsOp_OneInput);
+DECLARE_EMPTY_GLOBAL_SHADER_DERIVED_FROM(FSH_Exp, FSH_MathsOp_OneInput);
+DECLARE_EMPTY_GLOBAL_SHADER_DERIVED_FROM(FSH_Log2, FSH_MathsOp_OneInput);
+DECLARE_EMPTY_GLOBAL_SHADER_DERIVED_FROM(FSH_Log10, FSH_MathsOp_OneInput);
+DECLARE_EMPTY_GLOBAL_SHADER_DERIVED_FROM(FSH_Log, FSH_MathsOp_OneInput);
+DECLARE_EMPTY_GLOBAL_SHADER_DERIVED_FROM(FSH_Floor, FSH_MathsOp_OneInput);
+DECLARE_EMPTY_GLOBAL_SHADER_DERIVED_FROM(FSH_Ceil, FSH_MathsOp_OneInput);
+DECLARE_EMPTY_GLOBAL_SHADER_DERIVED_FROM(FSH_Round, FSH_MathsOp_OneInput);
+
+
+//////////////////////////////////////////////////////////////////////////
+/// Trigonometric functions
+//////////////////////////////////////////////////////////////////////////
+UENUM(BlueprintType)
+enum ETrigFunction : int
+{
+	Sin							UMETA(DisplayName = "Sine"),
+	Cos							UMETA(DisplayName = "Cosine"),
+	Tan							UMETA(DisplayName = "Tangent"),
+	ASin						UMETA(DisplayName = "ArcSine"),
+	ACos						UMETA(DisplayName = "ArcCosine"),
+	ATan						UMETA(DisplayName = "ArcTangent"),
+};
+
+//////////////////////////////////////////////////////////////////////////
+/// Helper C++ class
+//////////////////////////////////////////////////////////////////////////
+class TEXTUREGRAPHENGINE_API T_Maths_OneInput
+{
+public:
+	//////////////////////////////////////////////////////////////////////////
+	/// Static functions
+	//////////////////////////////////////////////////////////////////////////
+	static TiledBlobPtr			CreateTrigonometry(MixUpdateCyclePtr Cycle, BufferDescriptor DesiredOutputDesc, int32 TargetId, TiledBlobPtr Input, ETrigFunction Function);
+	static TiledBlobPtr			CreateToRadians(MixUpdateCyclePtr Cycle, BufferDescriptor DesiredOutputDesc, int32 TargetId, TiledBlobPtr Input);
+	static TiledBlobPtr			CreateToDegrees(MixUpdateCyclePtr Cycle, BufferDescriptor DesiredOutputDesc, int32 TargetId, TiledBlobPtr Input);
+	static TiledBlobPtr			CreateAbs(MixUpdateCyclePtr Cycle, BufferDescriptor DesiredOutputDesc, int32 TargetId, TiledBlobPtr Input);
+	static TiledBlobPtr			CreateSqrt(MixUpdateCyclePtr Cycle, BufferDescriptor DesiredOutputDesc, int32 TargetId, TiledBlobPtr Input);
+	static TiledBlobPtr			CreateSquare(MixUpdateCyclePtr Cycle, BufferDescriptor DesiredOutputDesc, int32 TargetId, TiledBlobPtr Input);
+	static TiledBlobPtr			CreateCbrt(MixUpdateCyclePtr Cycle, BufferDescriptor DesiredOutputDesc, int32 TargetId, TiledBlobPtr Input);
+	static TiledBlobPtr			CreateCube(MixUpdateCyclePtr Cycle, BufferDescriptor DesiredOutputDesc, int32 TargetId, TiledBlobPtr Input);
+	static TiledBlobPtr			CreateExp(MixUpdateCyclePtr Cycle, BufferDescriptor DesiredOutputDesc, int32 TargetId, TiledBlobPtr Input);
+	static TiledBlobPtr			CreateLog2(MixUpdateCyclePtr Cycle, BufferDescriptor DesiredOutputDesc, int32 TargetId, TiledBlobPtr Input);
+	static TiledBlobPtr			CreateLog10(MixUpdateCyclePtr Cycle, BufferDescriptor DesiredOutputDesc, int32 TargetId, TiledBlobPtr Input);
+	static TiledBlobPtr			CreateLog(MixUpdateCyclePtr Cycle, BufferDescriptor DesiredOutputDesc, int32 TargetId, TiledBlobPtr Input);
+	static TiledBlobPtr			CreateFloor(MixUpdateCyclePtr Cycle, BufferDescriptor DesiredOutputDesc, int32 TargetId, TiledBlobPtr Input);
+	static TiledBlobPtr			CreateCeil(MixUpdateCyclePtr Cycle, BufferDescriptor DesiredOutputDesc, int32 TargetId, TiledBlobPtr Input);
+	static TiledBlobPtr			CreateRound(MixUpdateCyclePtr Cycle, BufferDescriptor DesiredOutputDesc, int32 TargetId, TiledBlobPtr Input);
+};
diff --git a/Engine/Plugins/Experimental/TextureGraph/Source/TextureGraphEngine/Transform/Mix/T_InvalidateTiles.h b/Engine/Plugins/Experimental/TextureGraph/Source/TextureGraphEngine/Transform/Mix/T_InvalidateTiles.h
new file mode 100644
--- /dev/null
+++ b/Engine/Plugins/Experimental/TextureGraph/Source/TextureGraphEngine/Transform/Mix/T_InvalidateTiles.h
@@ -0,0 +1,25 @@
+// Copyright Epic Games, Inc. All Rights Reserved.
+#pragma once
+
+#include "CoreMinimal.h"
+#include "Transform/BlobTransform.h"
+#include "Job/Job.h"
+#include "Model/Mix/MixUpdateCycle.h"
+
+class TEXTUREGRAPHENGINE_API T_InvalidateTiles : public BlobTransform
+{
+public:
+									T_InvalidateTiles();
+	virtual							~T_InvalidateTiles() override;
+
+	virtual Device*					TargetDevice(size_t index) const override;
+	virtual AsyncTransformResultPtr	Exec(const TransformArgs& args) override;
+
+	virtual bool					GeneratesData() const override { return false; }
+
+	////////////////////////////////////////////////////////////////////////// 
+	/// Static functions
+	//////////////////////////////////////////////////////////////////////////
+	static JobUPtr					CreateJob(MixUpdateCyclePtr cycle, int32 targetId);
+	static void						Create(MixUpdateCyclePtr cycle, int32 targetId);
+};
diff --git a/Engine/Plugins/Experimental/TextureGraph/Source/TextureGraphEngine/Transform/Mix/T_UpdateTargets.h b/Engine/Plugins/Experimental/TextureGraph/Source/TextureGraphEngine/Transform/Mix/T_UpdateTargets.h
new file mode 100644
--- /dev/null
+++ b/Engine/Plugins/Experimental/TextureGraph/Source/TextureGraphEngine/Transform/Mix/T_UpdateTargets.h
@@ -0,0 +1,28 @@
+// Copyright Epic Games, Inc. All Rights Reserved.
+#pragma once
+
+#include "CoreMinimal.h"
+#include "Transform/BlobTransform.h"
+#include "Job/Job.h"
+#include "Model/Mix/MixUpdateCycle.h"
+
+class TEXTUREGRAPHENGINE_API T_UpdateTargets : public BlobTransform
+{
+private:
+	bool							_shouldUpdate = false;				/// This will ensure default texture is used when there is no child in layer set
+	static int						_gpuTesselationFactor;
+public:
+									T_UpdateTargets();
+	virtual							~T_UpdateTargets() override;
+	virtual Device*					TargetDevice(size_t index) const override;
+	virtual AsyncTransformResultPtr	Exec(const TransformArgs& args) override;
+
+	virtual bool					GeneratesData() const override { return false; }
+	virtual void					SetShouldUpdate(bool shouldUpdate) { _shouldUpdate = shouldUpdate; }
+	////////////////////////////////////////////////////////////////////////// 
+	/// Static functions
+	//////////////////////////////////////////////////////////////////////////
+	static void						SetGPUTesselationFactor(int tesselationFactor) { _gpuTesselationFactor = tesselationFactor; }
+	static JobUPtr					CreateJob(MixUpdateCyclePtr cycle, int32 targetId, bool shouldUpdate);
+	static void						Create(MixUpdateCyclePtr cycle, int32 targetId, bool shouldUpdate);
+};
diff --git a/Engine/Plugins/Experimental/TextureGraph/Source/TextureGraphEngine/Transform/Utility/T_CombineTiledBlob.cpp b/Engine/Plugins/Experimental/TextureGraph/Source/TextureGraphEngine/Transform/Utility/T_CombineTiledBlob.cpp
new file mode 100644
--- /dev/null
+++ b/Engine/Plugins/Experimental/TextureGraph/Source/TextureGraphEngine/Transform/Utility/T_CombineTiledBlob.cpp
@@ -0,0 +1,106 @@
+// Copyright Epic Games, Inc. All Rights Reserved.
+#include "T_CombineTiledBlob.h"
+#include "3D/RenderMesh.h"
+#include "FxMat/MaterialManager.h"
+#include "Job/JobBatch.h"
+#include "Device/FX/DeviceBuffer_FX.h"
+#include "Device/FX/Device_FX.h"
+#include "Model/Mix/MixInterface.h"
+
+TiledBlobPtr T_CombineTiledBlob::Create(MixUpdateCyclePtr Cycle, BufferDescriptor DesiredOutputDesc, int32 TargetId, TiledBlobPtr SourceTex)
+{
+	if (SourceTex && SourceTex->IsTiled())
+	{
+		FString Name = FString::Printf(TEXT("[%s].[%d].[%llu] CombineTiledBlob"), *SourceTex->Name(), TargetId, Cycle->GetBatch()->GetFrameId());
+
+		std::shared_ptr<CombineTiledBlob_Transform> CombinedTiledBlobMat = std::make_shared< CombineTiledBlob_Transform>(Name, SourceTex);
+
+		JobUPtr JobObj = std::make_unique<Job>(Cycle->GetMix(), TargetId, std::static_pointer_cast<BlobTransform>(CombinedTiledBlobMat));
+		JobObj->AddArg(WithUnbounded(ARG_BLOB(SourceTex, "Source")));
+
+		// Express the dependency of the new job on the job delivering the SourceTex
+		auto PrevJob = SourceTex->Job();
+		if (!PrevJob.expired())
+			JobObj->AddPrev(std::static_pointer_cast<DeviceNativeTask>(PrevJob.lock()));
+
+		TiledBlobPtr Result = JobObj->InitResult(Name, &DesiredOutputDesc, 1, 1);
+
+		Cycle->AddJob(TargetId, std::move(JobObj));
+		Result->MakeSingleBlob();
+
+		return Result;
+	}
+
+	return SourceTex;
+}
+
+
+//////////////////////////////////////////////////////////////////////////
+CombineTiledBlob_Transform::CombineTiledBlob_Transform(FString InName, TiledBlobPtr InSource)
+	: BlobTransform(InName)
+	, Source(InSource)
+{
+}
+
+Device* CombineTiledBlob_Transform::TargetDevice(size_t DevIndex) const
+{
+	return Device_FX::Get();
+}
+
+
+bool CombineTiledBlob_Transform::GeneratesData() const 
+{
+	return true;
+}
+bool CombineTiledBlob_Transform::CanHandleTiles() const
+{ 
+	return false;
+}
+
+AsyncTransformResultPtr	CombineTiledBlob_Transform::Exec(const TransformArgs& Args)
+{
+	BlobPtr Target = Args.Target.lock();
+	check(Target);
+
+	DeviceBuffer_FX* DevBuffer = dynamic_cast<DeviceBuffer_FX*>(Target->GetBufferRef().get());
+	check(DevBuffer);
+
+	T_Tiles<DeviceBufferRef> Tiles(Source->Rows(), Source->Cols());
+
+	for (int32 TileX = 0; TileX < Source->Rows(); TileX++)
+	{
+		for (int32 TileY = 0; TileY < Source->Cols(); TileY++)
+		{
+			check(Source->GetTiles()[TileX][TileY]);
+
+			BlobPtr Tile = Source->GetTiles()[TileX][TileY];
+			check(Tile->GetBufferRef());
+
+			Tiles[TileX][TileY] = Tile->GetBufferRef();
+		}
+	}
+
+	Device_FX* FXDevice = Device_FX::Get();
+
+	return FXDevice->DrawTilesToBuffer_Deferred(Target->GetBufferRef(), Tiles).then([this, Target](DeviceBufferRef)
+	{
+			TransformResultPtr Result = std::make_shared<TransformResult>();
+			Result->Target = Target;
+			return Result;
+	});
+}
+
+std::shared_ptr<BlobTransform> CombineTiledBlob_Transform::DuplicateInstance(FString NewName)
+{
+	return std::make_shared<CombineTiledBlob_Transform>(NewName, Source);
+}
+
+AsyncBufferResultPtr CombineTiledBlob_Transform::Bind(BlobPtr BlobObj, const ResourceBindInfo& BindInfo)
+{
+	return cti::make_ready_continuable<BufferResultPtr>(std::make_shared<BufferResult>());
+}
+
+AsyncBufferResultPtr CombineTiledBlob_Transform::Unbind(BlobPtr BlobObj, const ResourceBindInfo& BindInfo)
+{
+	return cti::make_ready_continuable<BufferResultPtr>(std::make_shared<BufferResult>());
+}
diff --git a/Engine/Plugins/Experimental/TextureGraph/Source/TextureGraphEngine/Transform/Utility/T_CombineTiledBlob.h b/Engine/Plugins/Experimental/TextureGraph/Source/TextureGraphEngine/Transform/Utility/T_CombineTiledBlob.h
new file mode 100644
--- /dev/null
+++ b/Engine/Plugins/Experimental/TextureGraph/Source/TextureGraphEngine/Transform/Utility/T_CombineTiledBlob.h
@@ -0,0 +1,44 @@
+// Copyright Epic Games, Inc. All Rights Reserved.
+#pragma once
+
+#include "CoreMinimal.h"
+#include "ShaderCompilerCore.h"
+#include "Job/Job.h"
+#include "FxMat/FxMaterial.h"
+#include "Model/Mix/MixUpdateCycle.h"
+#include <DataDrivenShaderPlatformInfo.h>
+
+
+/**
+ * A dedicated BlobTransform creating a Combined version of a TiledBlob
+ */
+class TEXTUREGRAPHENGINE_API CombineTiledBlob_Transform : public BlobTransform
+{
+private:
+	TiledBlobPtr					Source;
+
+public:
+	CombineTiledBlob_Transform(FString InName, TiledBlobPtr InSource);
+	virtual AsyncBufferResultPtr	Bind(BlobPtr Value, const ResourceBindInfo& BindInfo) override;
+	virtual AsyncBufferResultPtr	Unbind(BlobPtr BlobObj, const ResourceBindInfo& BindInfo) override;
+	virtual Device*					TargetDevice(size_t DevIndex) const override;
+	virtual std::shared_ptr<BlobTransform> DuplicateInstance(FString NewName) override;
+	virtual bool					GeneratesData() const override;
+	virtual bool					CanHandleTiles() const override;
+
+	virtual AsyncTransformResultPtr	Exec(const TransformArgs& Args) override;
+};
+
+/**
+ * Tiles To Combined Transform
+ */
+class TEXTUREGRAPHENGINE_API T_CombineTiledBlob
+{
+public:
+
+	//////////////////////////////////////////////////////////////////////////
+	/// Static functions
+	//////////////////////////////////////////////////////////////////////////
+	
+	static TiledBlobPtr				Create(MixUpdateCyclePtr Cycle, BufferDescriptor DesiredOutputDesc, int32 TargetId, TiledBlobPtr SourceTex);
+};
diff --git a/Engine/Plugins/Experimental/TextureGraph/Source/TextureGraphEngine/Transform/Utility/T_LoadStaticResource.h b/Engine/Plugins/Experimental/TextureGraph/Source/TextureGraphEngine/Transform/Utility/T_LoadStaticResource.h
new file mode 100644
--- /dev/null
+++ b/Engine/Plugins/Experimental/TextureGraph/Source/TextureGraphEngine/Transform/Utility/T_LoadStaticResource.h
@@ -0,0 +1,34 @@
+// Copyright Epic Games, Inc. All Rights Reserved.
+#pragma once
+
+#include "CoreMinimal.h"
+#include "Job/Job.h"
+#include "Model/Mix/MixUpdateCycle.h"
+
+class UStaticImageResource;
+
+class TEXTUREGRAPHENGINE_API Job_LoadStaticImageResource : public Job
+{
+	UStaticImageResource*			Source;			/// The source for the static image resource
+
+public:
+									Job_LoadStaticImageResource(UMixInterface* InMix, UStaticImageResource* InSource, int32 TargetId, UObject* InErrorOwner = nullptr, uint16 Priority = (uint16)E_Priority::kHigh);
+	virtual							~Job_LoadStaticImageResource() override = default;
+
+protected:
+	virtual void					GetDependencies(JobPtrVec& prior, JobPtrVec& after, JobRunInfo runInfo) override { RunInfo = runInfo; Prev.clear(); }
+	virtual AsyncPrepareResult		PrepareTargets(JobBatch* batch) override { return cti::make_ready_continuable(0); }
+	virtual AsyncPrepareResult		PrepareResources(JobBatch* batch) override { return cti::make_ready_continuable(0); }
+
+	virtual int32					Exec() override { return 0; }
+	virtual cti::continuable<int32>	BindArgs_All(JobRunInfo runInfo) override { return cti::make_ready_continuable(0); }
+	virtual cti::continuable<int32>	UnbindArgs_All(JobRunInfo runInfo) override { return cti::make_ready_continuable(0); }
+	virtual bool					CanHandleTiles() const override { return false; }
+
+	virtual cti::continuable<int32>	PreExecAsync(ENamedThreads::Type execThread, ENamedThreads::Type returnThread) override;
+	virtual cti::continuable<int32>	ExecAsync(ENamedThreads::Type execThread, ENamedThreads::Type returnThread) override;
+};
+
+//////////////////////////////////////////////////////////////////////////
+/// File based
+////////////////////////////////////////////////////////////////////////// 
diff --git a/Engine/Plugins/Experimental/TextureGraph/Source/TextureGraphEngine/Transform/Utility/T_SplitToTiles.cpp b/Engine/Plugins/Experimental/TextureGraph/Source/TextureGraphEngine/Transform/Utility/T_SplitToTiles.cpp
new file mode 100644
--- /dev/null
+++ b/Engine/Plugins/Experimental/TextureGraph/Source/TextureGraphEngine/Transform/Utility/T_SplitToTiles.cpp
@@ -0,0 +1,67 @@
+// Copyright Epic Games, Inc. All Rights Reserved.
+#include "T_SplitToTiles.h"
+#include "TextureGraphEngine.h"
+#include "Job/JobArgs.h"
+#include "3D/RenderMesh.h"
+#include "FxMat/MaterialManager.h"
+#include "Job/JobBatch.h"
+#include "Model/Mix/MixInterface.h"
+
+IMPLEMENT_GLOBAL_SHADER(CSH_SplitToTiles, "/Plugin/TextureGraph/Utils/SplitToTiles_comp.usf", "CSH_SplitToTiles", SF_Compute);
+
+T_SplitToTiles::T_SplitToTiles()
+{
+}
+
+T_SplitToTiles::~T_SplitToTiles()
+{
+}
+
+TiledBlobPtr CreateSplitToTilesCompute(MixUpdateCyclePtr Cycle, TiledBlobPtr SourceTex, int32 TargetId)
+{
+	CSH_SplitToTiles::FPermutationDomain PermutationVector;
+
+	FString Name = FString::Printf(TEXT("[%s].[%d].[%llu] SplitToTiles"), *SourceTex->Name(), TargetId, Cycle->GetBatch()->GetFrameId());
+
+	int32 SrcWidth = SourceTex->GetWidth();
+	int32 SrcHeight = SourceTex->GetHeight();
+
+	int32 DstWidth = SrcWidth;
+	int32 DstHeight = SrcHeight;
+
+	int32 srcNumCols = (int32)SourceTex->GetTiles().Cols();
+	int32 srcNumRows = (int32)SourceTex->GetTiles().Rows();	
+
+	int32 DstNumCols = Cycle->GetMix()->GetNumXTiles();
+	int32 DstNumRows = Cycle->GetMix()->GetNumYTiles();
+
+	FIntVector4 SrcDimensions(DstWidth / DstNumCols, DstHeight / DstNumRows, srcNumCols, srcNumRows);
+	FTileInfo TileInfo;
+
+	RenderMaterial_FXPtr Transform = TextureGraphEngine::GetMaterialManager()->CreateMaterial_CmpFX<CSH_SplitToTiles>(
+		Name, TEXT("Result"), PermutationVector, SrcDimensions.X, SrcDimensions.Y, 1);
+
+	JobUPtr JobObj = std::make_unique<Job>(Cycle->GetMix(), TargetId, std::static_pointer_cast<BlobTransform>(Transform));
+	JobObj
+		->AddArg(ARG_BLOB(SourceTex, "SourceTexture"))
+		->AddArg(ARG_TILEINFO(TileInfo, "TileInfo"));
+
+	BufferDescriptor Desc = SourceTex->GetDescriptor();
+	Desc.Name = Name;
+	Desc.AllowUAV();
+
+	TiledBlobPtr Result = JobObj->InitResult(Name, &Desc);
+
+	Cycle->AddJob(TargetId, std::move(JobObj));
+
+	check(Desc.Width == SourceTex->GetWidth() && Desc.Height == SourceTex->GetHeight())
+
+	return Result;
+}
+
+TiledBlobPtr T_SplitToTiles::Create(MixUpdateCyclePtr Cycle, int32 TargetId, TiledBlobPtr SourceTex)
+{
+	/// Cannot generate. Just return the source texture
+	check(TextureHelper::IsPowerOf2(SourceTex->GetWidth()) && TextureHelper::IsPowerOf2(SourceTex->GetHeight()));
+	return CreateSplitToTilesCompute(Cycle, SourceTex, TargetId);
+}
diff --git a/Engine/Plugins/Experimental/TextureGraph/Source/TextureGraphEngine/Transform/Utility/T_TextureHistogram.cpp b/Engine/Plugins/Experimental/TextureGraph/Source/TextureGraphEngine/Transform/Utility/T_TextureHistogram.cpp
new file mode 100644
--- /dev/null
+++ b/Engine/Plugins/Experimental/TextureGraph/Source/TextureGraphEngine/Transform/Utility/T_TextureHistogram.cpp
@@ -0,0 +1,107 @@
+// Copyright Epic Games, Inc. All Rights Reserved.
+#include "T_TextureHistogram.h"
+#include "TextureGraphEngine.h"
+#include "Job/JobArgs.h"
+#include "TextureGraphEngine.h"
+#include "Job/HistogramService.h"
+#include "Job/JobBatch.h"
+#include "Job/Scheduler.h"
+#include "Helper/MathUtils.h"
+#include "Helper/GraphicsUtil.h"
+#include "Device/FX/DeviceBuffer_FX.h"
+#include "Engine/TextureRenderTarget2D.h"
+#include "Engine/Texture2D.h"
+#include "Device/DeviceManager.h"
+#include "Device/Mem/Device_Mem.h"
+#include "Data/Blobber.h"
+#include "2D/Tex.h"
+#include "FxMat/MaterialManager.h"
+#include "T_MinMax.h"
+
+IMPLEMENT_GLOBAL_SHADER(CSH_Histogram, "/Plugin/TextureGraph/Utils/Histogram_comp.usf", "CSH_Histogram", SF_Compute);
+
+static constexpr int NumBins = 256;
+
+T_TextureHistogram::T_TextureHistogram()
+{
+}
+
+T_TextureHistogram::~T_TextureHistogram()
+{
+}
+
+RenderMaterial_FXPtr T_TextureHistogram::CreateMaterial_Histogram(FString Name, FString OutputId, const CSH_Histogram::FPermutationDomain& cmpshPermutationDomain,
+	int NumThreadsX, int NumThreadsY, int NumThreadsZ)
+{
+	std::shared_ptr<FxMaterial_Histogram<CSH_Histogram>> Mat = std::make_shared<FxMaterial_Histogram<CSH_Histogram>>(OutputId, &cmpshPermutationDomain, NumThreadsX, NumThreadsY, NumThreadsZ);
+	return std::make_shared<RenderMaterial_FX>(Name, std::static_pointer_cast<FxMaterial>(Mat));
+}
+
+TiledBlobPtr T_TextureHistogram::Create(UMixInterface* InMix, TiledBlobPtr SourceTex, int32 TargetId)
+{
+	if (SourceTex->HasHistogram())
+	{
+		TiledBlobPtr SourceHistogram = std::static_pointer_cast<TiledBlob>(SourceTex->GetHistogram());
+		return SourceHistogram;
+	}
+
+	check(SourceTex);
+	HistogramServicePtr Service = TextureGraphEngine::GetScheduler()->GetHistogramService().lock();
+	check(Service);
+
+	JobBatchPtr Batch = Service->GetOrCreateNewBatch(InMix);
+
+	CSH_Histogram::FPermutationDomain PermutationVector;
+
+	FIntVector4 SrcDimensions(SourceTex->GetWidth(), SourceTex->GetHeight(), 1, 1);
+
+	auto SrcTiles = std::make_shared<JobArg_Blob>(SourceTex, "SourceTiles");
+
+	FString Name = FString::Printf(TEXT("[%s].[%d] Histogram"), *SourceTex->DisplayName(), TargetId);
+	RenderMaterial_FXPtr Transform = T_TextureHistogram::CreateMaterial_Histogram(
+		TEXT("T_Histogram"), TEXT("Result"), PermutationVector, SrcDimensions.X, SrcDimensions.Y, 1);
+
+	JobUPtr JobObj = std::make_unique<Job>(InMix, TargetId, std::static_pointer_cast<BlobTransform>(Transform));
+	JobObj
+		->AddArg(SrcTiles);
+
+	BufferDescriptor Desc;
+	Desc.Width = NumBins;
+	Desc.Height = 2;
+	Desc.Format = BufferFormat::Float;// BufferFormat::Int;
+	Desc.ItemsPerPoint = 4;
+	Desc.Name = FString::Printf(TEXT("Histogram - %s"), *SourceTex->Name());
+	Desc.AllowUAV();
+
+	JobObj->SetTiled(false);
+
+	TiledBlobPtr Result = JobObj->InitResult(Name, &Desc,1,1);
+
+	Result->MakeSingleBlob();
+
+	//Add the job using histogram idle service
+	AddHistogramJobToCycle(Batch->GetCycle(),std::move(JobObj), TargetId, InMix);
+
+
+	if (!SourceTex->HasHistogram())
+	{
+		//setting it as the histogram of source so it is retained untill the life cycle of source blob.
+		SourceTex->SetHistogram(Result);
+	}
+
+	return Result;
+}
+
+void T_TextureHistogram::AddHistogramJobToCycle(MixUpdateCyclePtr Cycle, JobUPtr Job, int32 TargetId,UMixInterface* Mix)
+{
+	HistogramServicePtr Service = TextureGraphEngine::GetScheduler()->GetHistogramService().lock();
+
+	if (!Service)
+		return;
+
+	//TODO get rid of mix here and use Null Mix instead
+	UMixInterface* mix = Job->GetMix();
+	check(mix);
+
+	Service->AddHistogramJob(Cycle,std::move(Job), TargetId,Mix);
+}
diff --git a/Engine/Plugins/Experimental/TextureGraph/Source/TextureGraphEngine/Transform/Utility/T_TextureHistogram.h b/Engine/Plugins/Experimental/TextureGraph/Source/TextureGraphEngine/Transform/Utility/T_TextureHistogram.h
new file mode 100644
--- /dev/null
+++ b/Engine/Plugins/Experimental/TextureGraph/Source/TextureGraphEngine/Transform/Utility/T_TextureHistogram.h
@@ -0,0 +1,76 @@
+// Copyright Epic Games, Inc. All Rights Reserved.
+#pragma once
+
+#include "CoreMinimal.h"
+#include "ShaderCompilerCore.h"
+#include "Job/Job.h"
+#include "FxMat/FxMaterial.h"
+#include "Model/Mix/MixUpdateCycle.h"
+#include <DataDrivenShaderPlatformInfo.h>
+
+class CSH_Histogram : public CmpSH_Base<1, 1, 1>
+{
+public:
+	DECLARE_GLOBAL_SHADER(CSH_Histogram);
+	SHADER_USE_PARAMETER_STRUCT(CSH_Histogram, CmpSH_Base);
+	BEGIN_SHADER_PARAMETER_STRUCT(FParameters, )
+		SHADER_PARAMETER_TEXTURE(Texture2D, SourceTiles)
+		SHADER_PARAMETER_UAV(RWTexture2D<float4>, Result)
+		SHADER_PARAMETER_UAV(RWStructuredBuffer<uint4>, Histogram)
+	END_SHADER_PARAMETER_STRUCT()
+
+public:
+	static bool ShouldCompilePermutation(const FGlobalShaderPermutationParameters& Parameters)
+	{
+		return IsPCPlatform(Parameters.Platform) && EnumHasAllFlags(Parameters.Flags, EShaderPermutationFlags::HasEditorOnlyData)
+			&& IsFeatureLevelSupported(Parameters.Platform, ERHIFeatureLevel::SM5);
+	}
+};
+
+template <typename CmpSH_Type>
+class FxMaterial_Histogram : public FxMaterial_Compute<CmpSH_Type>
+{
+	FRWBufferStructured Buffer;
+public:
+	// type Name for the permutation domain
+	using CmpSHPermutationDomain = typename CmpSH_Type::FPermutationDomain;
+
+	 FxMaterial_Histogram(FString outputId, const CmpSHPermutationDomain* permDomain = nullptr,
+		int numThreadsX = FxMaterial_Compute<CmpSH_Type>::GDefaultNumThreadsXY, int numThreadsY = FxMaterial_Compute<CmpSH_Type>::GDefaultNumThreadsXY, int numThreadsZ = 1,
+		FUnorderedAccessViewRHIRef unorderedAccessView = nullptr)
+		: FxMaterial_Compute<CmpSH_Type>(outputId, permDomain, numThreadsX, numThreadsY, numThreadsZ, unorderedAccessView)
+	{
+
+	}
+
+	virtual void Blit(FRHICommandListImmediate& RHI, FRHITexture2D* Target, const RenderMesh* MeshObj, int32 TargetId, FGraphicsPipelineStateInitializer* PSO = nullptr) override
+	{
+		SetBuffer(RHI);
+
+		FxMaterial_Compute<CmpSH_Type>::Blit(RHI, Target, MeshObj, TargetId, PSO);
+	}
+
+	void SetBuffer(FRHICommandListImmediate& RHI)
+	{
+		Buffer.Initialize(RHI, TEXT("HistogramRBuffer"), sizeof(uint32) * 4, 256, BUF_UnorderedAccess | BUF_ShaderResource);
+		RHI.ClearUAVUint(Buffer.UAV, FUintVector4(0, 0, 0, 0));
+		FxMaterial_Compute<CmpSH_Type>::Params.Histogram = Buffer.UAV;
+	}
+};
+
+/**
+ * T_TextureHistogram Transform
+ */
+class TEXTUREGRAPHENGINE_API T_TextureHistogram
+{
+public:
+	T_TextureHistogram();
+	~T_TextureHistogram();
+
+	static TiledBlobPtr	Create(UMixInterface* InMix, TiledBlobPtr SourceTex, int32 TargetId);
+	static void AddHistogramJobToCycle(MixUpdateCyclePtr Cycle, JobUPtr Job, int32 TargetId, UMixInterface* Mix);
+	
+private:
+	static TiledBlobPtr GetMaxValueForHistogram(MixUpdateCyclePtr Cycle, TiledBlobPtr Histogram, int32 TargetId);
+	static RenderMaterial_FXPtr CreateMaterial_Histogram(FString Name, FString OutputId, const CSH_Histogram::FPermutationDomain& cmpshPermutationDomain,int NumThreadsX, int NumThreadsY, int NumThreadsZ = 1);
+};
diff --git a/Engine/Plugins/Experimental/TextureGraph/TextureGraph.uplugin b/Engine/Plugins/Experimental/TextureGraph/TextureGraph.uplugin
--- a/Engine/Plugins/Experimental/TextureGraph/TextureGraph.uplugin
+++ b/Engine/Plugins/Experimental/TextureGraph/TextureGraph.uplugin
@@ -15,7 +15,7 @@
 	"EnabledByDefault": false,
 	"Modules": [
 		{
-			"Name": "MixerEngine",
+			"Name": "TextureGraphEngine",
 			"Type": "Runtime",
 			"LoadingPhase": "PostConfigInit",
 			"PlatformAllowList": [
