Fix crash when loading existing TG after previous change
added a test to check this case
#tests 65ab73ca0957af3900bd2de9

diff --git a/Engine/Plugins/Experimental/TextureGraph/Source/TextureGraph/Private/TextureGraph.cpp b/Engine/Plugins/Experimental/TextureGraph/Source/TextureGraph/Private/TextureGraph.cpp
--- a/Engine/Plugins/Experimental/TextureGraph/Source/TextureGraph/Private/TextureGraph.cpp
+++ b/Engine/Plugins/Experimental/TextureGraph/Source/TextureGraph/Private/TextureGraph.cpp
@@ -89,6 +89,10 @@ void UTextureGraph::PostLoad()
 		{
 			Settings->GetViewportSettings().InitDefaultSettings(TextureGraph->GetNode(OutputIds[0])->GetNodeName());
 		}
+		else // the TG doesn't have any output? that's weird but could happen, let's simply create the material regardless targeting a name
+		{
+			Settings->GetViewportSettings().InitDefaultSettings("Output");
+		}
 	}
 }
 
