Texture Graph Palette UI updates
Node UI updates

#virtualized

diff --git a/Engine/Plugins/Experimental/TextureGraph/Source/TextureGraphEditor/Private/EdGraph/STG_EditorGraphNode.cpp b/Engine/Plugins/Experimental/TextureGraph/Source/TextureGraphEditor/Private/EdGraph/STG_EditorGraphNode.cpp
--- a/Engine/Plugins/Experimental/TextureGraph/Source/TextureGraphEditor/Private/EdGraph/STG_EditorGraphNode.cpp
+++ b/Engine/Plugins/Experimental/TextureGraph/Source/TextureGraphEditor/Private/EdGraph/STG_EditorGraphNode.cpp
@@ -58,13 +58,20 @@ void STG_EditorGraphNode::Construct(const FArguments& InArgs, UTG_EdGraphNode* I
 	}
 	
 	BodyBrush = *GetNodeBodyBrush();
-	BodyBrush.OutlineSettings.Color = GetNodeTitleColor();
+	BodyBrush.OutlineSettings.Color = UTG_EdGraphSchema::NodeBodyColorOutline;
+	
+	HeaderBrush = *FTG_Style::Get().GetBrush("TG.Graph.Node.Header");
+	HeaderBrush.OutlineSettings.Color = GetNodeTitleColor().GetSpecifiedColor() * UTG_EdGraphSchema::NodeOutlineColorMultiplier;
+
 	UpdateGraphNode();
 }
 
 BEGIN_SLATE_FUNCTION_BUILD_OPTIMIZATION
 void STG_EditorGraphNode::UpdateGraphNode()
 {
+	const FMargin TitleRightWidgetPadding = FMargin(5, 0, 0, 0);
+	const FMargin TitleOuterMargin = FMargin(4);
+
 	InputPins.Empty();
 	OutputPins.Empty();
 
@@ -100,66 +107,66 @@ void STG_EditorGraphNode::UpdateGraphNode()
 		IconBrush = GraphNode->GetIconAndTint(IconColor).GetOptionalIcon();
 	}
 
-	TitleBorderMargin = FMargin(4.f, 5.f, 2.f, 4.f);
+	TitleBorderMargin = FMargin(6);
 
 	TSharedRef<SOverlay> DefaultTitleAreaWidget =
-		SNew(SOverlay)
-		+ SOverlay::Slot()
-		.HAlign(HAlign_Fill)
-		.VAlign(VAlign_Top)
-		[
-			SNew(SBorder)
-			.BorderImage(FTG_Style::Get().GetBrush("TG.Graph.Node.NoColorSpill"))//FAppStyle::GetBrush("Graph.Node.ColorSpill")
+	SNew(SOverlay)
+	+ SOverlay::Slot()
+	.HAlign(HAlign_Fill)
+	.VAlign(VAlign_Top)
+	[
+		SNew(SBorder)
+		.BorderImage(&HeaderBrush)//FAppStyle::GetBrush("Graph.Node.ColorSpill")
 		.BorderBackgroundColor(this, &STG_EditorGraphNode::GetNodeTitleColor)
-		[
-			SNew(SHorizontalBox)
-			+ SHorizontalBox::Slot()
 		.Padding(TitleBorderMargin)
-		.VAlign(VAlign_Center)
-		.HAlign(HAlign_Fill)
 		[
 			SNew(SHorizontalBox)
 			+ SHorizontalBox::Slot()
-		.VAlign(VAlign_Top)
-		.Padding(FMargin(0.f, 0.f, 4.f, 0.f))
-		.AutoWidth()
-		[
-			SNew(SImage)
-			.Image(IconBrush)
-		.ColorAndOpacity(this, &STG_EditorGraphNode::GetNodeTitleIconColor)
-		]
-	+ SHorizontalBox::Slot()
-		[
-			SNew(SVerticalBox)
-			+ SVerticalBox::Slot()
-		.VAlign(VAlign_Top)
-		.AutoHeight()
-		[
-			CreateTitleWidget(NodeTitle)
-		]
-	+ SVerticalBox::Slot()
-		.AutoHeight()
-		.VAlign(VAlign_Top)
-		[
-			NodeTitle.ToSharedRef()
-		]
-	+ SVerticalBox::Slot()
-		.AutoHeight()
-		[
-			CreateTitleDetailsWidget()
-		]
-		]
-		]
-	+ SHorizontalBox::Slot()
-		.HAlign(HAlign_Right)
-		.VAlign(VAlign_Center)
-		.Padding(5, 0, 0, 0)
-		.AutoWidth()
-		[
-			CreateTitleRightWidget()
-		]
+			.Padding(TitleBorderMargin)
+			.VAlign(VAlign_Center)
+			.HAlign(HAlign_Fill)
+			[
+				SNew(SHorizontalBox)
+				+ SHorizontalBox::Slot()
+				.VAlign(VAlign_Top)
+				.AutoWidth()
+				[
+					SNew(SImage)
+					.Image(IconBrush)
+					.ColorAndOpacity(this, &STG_EditorGraphNode::GetNodeTitleIconColor)
+				]
+				+ SHorizontalBox::Slot()
+				[
+					SNew(SVerticalBox)
+					+ SVerticalBox::Slot()
+					.VAlign(VAlign_Top)
+					.AutoHeight()
+					[
+						CreateTitleWidget(NodeTitle)
+					]
+					+ SVerticalBox::Slot()
+					.AutoHeight()
+					.VAlign(VAlign_Top)
+					[
+						NodeTitle.ToSharedRef()
+					]
+					+ SVerticalBox::Slot()
+					.AutoHeight()
+					[
+						CreateTitleDetailsWidget()
+					]
+				]
+			]
+			+ SHorizontalBox::Slot()
+			.HAlign(HAlign_Right)
+			.VAlign(VAlign_Center)
+			.Padding(TitleRightWidgetPadding)
+			.AutoWidth()
+			[
+				CreateTitleRightWidget()
+			]
 		]
-		];
+	];
 
 	SetDefaultTitleAreaWidget(DefaultTitleAreaWidget);
 
@@ -169,7 +176,7 @@ void STG_EditorGraphNode::UpdateGraphNode()
 		.LowDetail()
 		[
 			SNew(SBorder)
-			.BorderImage(FTG_Style::Get().GetBrush("TG.Graph.Node.NoColorSpill"))
+			.BorderImage(&HeaderBrush)
 			.Padding(FMargin(75.0f, 22.0f)) // Saving enough space for a 'typical' title so the transition isn't quite so abrupt
 			.BorderBackgroundColor(this, &STG_EditorGraphNode::GetNodeTitleColor)
 		]
@@ -197,7 +204,7 @@ void STG_EditorGraphNode::UpdateGraphNode()
 		.AutoHeight()
 		.HAlign(HAlign_Fill)
 		.VAlign(VAlign_Top)
-		.Padding(Settings->GetNonPinNodeBodyPadding())
+		.Padding(TitleOuterMargin)
 		[
 			TitleAreaWidget
 		]
@@ -358,7 +365,7 @@ int32 STG_EditorGraphNode::OnPaint(const FPaintArgs& Args, const FGeometry& Allo
 
 const FSlateBrush* STG_EditorGraphNode::GetShadowBrush(bool bSelected) const
 {
-	return bSelected ? FTG_Style::Get().GetBrush(TEXT("TG.Graph.Node.ShadowSelected")) : FAppStyle::GetBrush(TEXT("Graph.Node.Shadow"));
+	return bSelected ? FAppStyle::GetBrush(TEXT("Graph.Node.ShadowSelected")) : FAppStyle::GetBrush(TEXT("Graph.Node.Shadow"));
 }
 
 int STG_EditorGraphNode::GetContentAreaBorderThickness()
diff --git a/Engine/Plugins/Experimental/TextureGraph/Source/TextureGraphEditor/Private/EdGraph/STG_EditorGraphNode.h b/Engine/Plugins/Experimental/TextureGraph/Source/TextureGraphEditor/Private/EdGraph/STG_EditorGraphNode.h
--- a/Engine/Plugins/Experimental/TextureGraph/Source/TextureGraphEditor/Private/EdGraph/STG_EditorGraphNode.h
+++ b/Engine/Plugins/Experimental/TextureGraph/Source/TextureGraphEditor/Private/EdGraph/STG_EditorGraphNode.h
@@ -88,6 +88,7 @@ private:
 	FDelegateHandle OnNodeChangedHandle;
 	FDelegateHandle OnPinSelectionChangedHandle;
 	FSlateBrush BodyBrush;
+	FSlateBrush HeaderBrush;
 	TSharedPtr<SBox> PreviewWidgetBox;
 	TMap<FTG_Id, TSharedPtr<STG_NodeThumbnail>> PinThumbWidgetMap;
 	FLinearColor GetPinThumbSelectionColor();
diff --git a/Engine/Plugins/Experimental/TextureGraph/Source/TextureGraphEditor/Private/EdGraph/TG_EdGraphNode.cpp b/Engine/Plugins/Experimental/TextureGraph/Source/TextureGraphEditor/Private/EdGraph/TG_EdGraphNode.cpp
--- a/Engine/Plugins/Experimental/TextureGraph/Source/TextureGraphEditor/Private/EdGraph/TG_EdGraphNode.cpp
+++ b/Engine/Plugins/Experimental/TextureGraph/Source/TextureGraphEditor/Private/EdGraph/TG_EdGraphNode.cpp
@@ -561,44 +561,7 @@ float UTG_EdGraphNode::GetNodeAlpha() const
 FLinearColor UTG_EdGraphNode::GetTitleColor() const
 {
 	FName ExpressionCategory = GetNode()->GetExpression()->GetCategory();
-	if (ExpressionCategory == TG_Category::Output)
-	{
-		return UTG_EdGraphSchema::OutputNodesColor;
-	}
-	if (ExpressionCategory == TG_Category::Input)
-	{
-		return UTG_EdGraphSchema::InputNodesColor;
-	}
-	if (ExpressionCategory == TG_Category::Utilities)
-	{
-		return UTG_EdGraphSchema::FunctionNodesColor;
-	}
-	if (ExpressionCategory == TG_Category::Maths)
-	{
-		return UTG_EdGraphSchema::MathsNodesColor;
-	}
-	if (ExpressionCategory == TG_Category::Procedural)
-	{
-		return UTG_EdGraphSchema::GeneratorNodesColor;
-	}
-	if (ExpressionCategory == TG_Category::Adjustment ||
-		ExpressionCategory == TG_Category::Filter ||
-		ExpressionCategory == TG_Category::Channel)
-	{
-		return UTG_EdGraphSchema::OperatorNodesColor;
-	}
-
-	if (ExpressionCategory == TG_Category::DevOnly)
-	{
-		return UTG_EdGraphSchema::DevOnlyNodesColor;
-	}
-	if (ExpressionCategory == TG_Category::Custom ||
-		ExpressionCategory == TG_Category::Default)
-	{
-		return UTG_EdGraphSchema::CustomNodesColor;
-	}
-
-	return UTG_EdGraphSchema::TGNodeColor;
+	return UTG_EdGraphSchema::GetCategoryColor(ExpressionCategory);
 }
 
 FLinearColor UTG_EdGraphNode::GetNodeTitleColor() const
diff --git a/Engine/Plugins/Experimental/TextureGraph/Source/TextureGraphEditor/Private/EdGraph/TG_EdGraphSchema.cpp b/Engine/Plugins/Experimental/TextureGraph/Source/TextureGraphEditor/Private/EdGraph/TG_EdGraphSchema.cpp
--- a/Engine/Plugins/Experimental/TextureGraph/Source/TextureGraphEditor/Private/EdGraph/TG_EdGraphSchema.cpp
+++ b/Engine/Plugins/Experimental/TextureGraph/Source/TextureGraphEditor/Private/EdGraph/TG_EdGraphSchema.cpp
@@ -30,17 +30,23 @@ const FLinearColor UTG_EdGraphSchema::VectorPinColor = FLinearColor::Yellow;
 const FLinearColor UTG_EdGraphSchema::ImagePinColor = FLinearColor(0.8f, 0.2f, 0.4f, 1.0f);					// salmon (light pink)
 
 // Node title colors
-const FLinearColor UTG_EdGraphSchema::InputNodesColor =		FLinearColor::FromSRGBColor(FColor::FromHex(TEXT("0F6389")));//48748F
-const FLinearColor UTG_EdGraphSchema::FunctionNodesColor =	FLinearColor::FromSRGBColor(FColor::FromHex(TEXT("3E3E3E")));
-const FLinearColor UTG_EdGraphSchema::MathsNodesColor =		FLinearColor::FromSRGBColor(FColor::FromHex(TEXT("60825B")));
-const FLinearColor UTG_EdGraphSchema::GeneratorNodesColor = FLinearColor::FromSRGBColor(FColor::FromHex(TEXT("5B8830")));
-const FLinearColor UTG_EdGraphSchema::OperatorNodesColor =	FLinearColor::FromSRGBColor(FColor::FromHex(TEXT("8E7424")));
-const FLinearColor UTG_EdGraphSchema::OutputNodesColor =	FLinearColor::FromSRGBColor(FColor::FromHex(TEXT("8D7964")));//890F0F//FLinearColor::FromSRGBColor(FColor::FromHex(TEXT("8D7964")));
-const FLinearColor UTG_EdGraphSchema::CustomNodesColor =	FLinearColor::Black;
-const FLinearColor UTG_EdGraphSchema::DevOnlyNodesColor =	FLinearColor::FromSRGBColor(FColor::Cyan);
-
-const FLinearColor UTG_EdGraphSchema::TGNodeColor =			FLinearColor::FromSRGBColor(FColor::FromHex(TEXT("0F6389")));
-const FLinearColor UTG_EdGraphSchema::NodeBodyColor =		FLinearColor::FromSRGBColor(FColor::FromHex(TEXT("242424")));
+const FLinearColor UTG_EdGraphSchema::InputNodesColor = FLinearColor::FromSRGBColor(FColor::FromHex(TEXT("3E1950")));
+const FLinearColor UTG_EdGraphSchema::FunctionNodesColor = FLinearColor::FromSRGBColor(FColor::FromHex(TEXT("2A97704D")));
+const FLinearColor UTG_EdGraphSchema::MathsNodesColor = FLinearColor::FromSRGBColor(FColor::FromHex(TEXT("2C3B29"))); 
+const FLinearColor UTG_EdGraphSchema::ProceduralNodesColor = FLinearColor::FromSRGBColor(FColor::FromHex(TEXT("293E16")));
+const FLinearColor UTG_EdGraphSchema::OperatorNodesColor = FLinearColor::FromSRGBColor(FColor::FromHex(TEXT("383316")));
+const FLinearColor UTG_EdGraphSchema::FilterNodesColor = FLinearColor::FromSRGBColor(FColor::FromHex(TEXT("253016")));
+const FLinearColor UTG_EdGraphSchema::ChannelNodesColor = FLinearColor::FromSRGBColor(FColor::FromHex(TEXT("203328")));
+const FLinearColor UTG_EdGraphSchema::OutputNodesColor = FLinearColor::FromSRGBColor(FColor::FromHex(TEXT("4B453B")));
+const FLinearColor UTG_EdGraphSchema::CustomNodesColor = FLinearColor::Black;
+const FLinearColor UTG_EdGraphSchema::DevOnlyNodesColor = FLinearColor::FromSRGBColor(FColor::Cyan);
+const FLinearColor UTG_EdGraphSchema::UtilitiesNodesColor = FLinearColor::FromSRGBColor(FColor::FromHex(TEXT("133C33")));
+
+const FLinearColor UTG_EdGraphSchema::TGNodeColor = FLinearColor::FromSRGBColor(FColor::FromHex(TEXT("0F6389")));
+const FLinearColor UTG_EdGraphSchema::NodeBodyColor = FLinearColor::FromSRGBColor(FColor::FromHex(TEXT("0F0F0F")));
+const FLinearColor UTG_EdGraphSchema::NodeBodyColorOutline = FLinearColor::FromSRGBColor(FColor::FromHex(TEXT("353535")));
+
+const float UTG_EdGraphSchema::NodeOutlineColorMultiplier = 2.2f;
 
 const FName UTG_EdGraphSchema::PC_Boolean(TEXT("bool"));
 const FName UTG_EdGraphSchema::PC_Byte(TEXT("byte"));
@@ -483,4 +489,52 @@ void UTG_EdGraphSchema::GetPaletteActions(FGraphActionMenuBuilder& ActionMenuBui
 	GetTG_ExpressionsActions(ActionMenuBuilder);
 }
 
+FLinearColor UTG_EdGraphSchema::GetCategoryColor(FName Category)
+{
+	if (Category == TG_Category::Output)
+	{
+		return UTG_EdGraphSchema::OutputNodesColor;
+	}
+	if (Category == TG_Category::Input)
+	{
+		return UTG_EdGraphSchema::InputNodesColor;
+	}
+	if (Category == TG_Category::Maths)
+	{
+		return UTG_EdGraphSchema::MathsNodesColor;
+	}
+	if (Category == TG_Category::Procedural)
+	{
+		return UTG_EdGraphSchema::ProceduralNodesColor;
+	}
+	if (Category == TG_Category::Adjustment)
+	{
+		return UTG_EdGraphSchema::OperatorNodesColor;
+	}
+	if (Category == TG_Category::Filter)
+	{
+		return UTG_EdGraphSchema::FilterNodesColor;
+	}
+	if (Category == TG_Category::Channel)
+	{
+		return UTG_EdGraphSchema::ChannelNodesColor;
+	}
+	if (Category == TG_Category::Utilities)
+	{
+		return UTG_EdGraphSchema::UtilitiesNodesColor;
+	}
+
+	if (Category == TG_Category::DevOnly)
+	{
+		return UTG_EdGraphSchema::DevOnlyNodesColor;
+	}
+	if (Category == TG_Category::Custom ||
+		Category == TG_Category::Default)
+	{
+		return UTG_EdGraphSchema::CustomNodesColor;
+	}
+
+	return UTG_EdGraphSchema::TGNodeColor;
+}
+
 #undef LOCTEXT_NAMESPACE
diff --git a/Engine/Plugins/Experimental/TextureGraph/Source/TextureGraphEditor/Private/EdGraph/TG_EdGraphSchema.h b/Engine/Plugins/Experimental/TextureGraph/Source/TextureGraphEditor/Private/EdGraph/TG_EdGraphSchema.h
--- a/Engine/Plugins/Experimental/TextureGraph/Source/TextureGraphEditor/Private/EdGraph/TG_EdGraphSchema.h
+++ b/Engine/Plugins/Experimental/TextureGraph/Source/TextureGraphEditor/Private/EdGraph/TG_EdGraphSchema.h
@@ -27,13 +27,17 @@ public:
 	static const FLinearColor InputNodesColor;
 	static const FLinearColor FunctionNodesColor;
 	static const FLinearColor MathsNodesColor;
-	static const FLinearColor GeneratorNodesColor;
+	static const FLinearColor ProceduralNodesColor;
 	static const FLinearColor OperatorNodesColor;
 	static const FLinearColor OutputNodesColor;
+	static const FLinearColor FilterNodesColor;
+	static const FLinearColor ChannelNodesColor;
 	static const FLinearColor CustomNodesColor;
 	static const FLinearColor DevOnlyNodesColor;
+	static const FLinearColor UtilitiesNodesColor;
 	static const FLinearColor TGNodeColor;
 	static const FLinearColor NodeBodyColor;
+	static const FLinearColor NodeBodyColorOutline;
 
 	static const FName PC_Boolean;
 	static const FName PC_Byte;
@@ -49,6 +53,8 @@ public:
 	static const FName PC_Struct;    // SubCategoryObject is the ScriptStruct of the struct passed thru this pin, 'self' is not a valid SubCategory. DefaultObject should always be empty, the DefaultValue string may be used for supported structs.
 	static const FName PC_Array;
 	static const FName PC_Enum;    // SubCategoryObject is the UEnum object passed thru this pin
+
+	static const float NodeOutlineColorMultiplier;
 	
 	//~ Begin EdGraphSchema Interface
 	virtual void GetGraphContextActions(FGraphContextMenuBuilder& ContextMenuBuilder) const override;
@@ -73,6 +79,7 @@ public:
 	UTG_Pin* GetTGPinFromEdPin(const UEdGraphPin* InPin) const;
 
 	void GetPaletteActions(FGraphActionMenuBuilder& ActionMenuBuilder, const FString& CategoryName) const;
+	static FLinearColor GetCategoryColor(FName Category);
 
 private:
 	void GetTG_ExpressionsActions(FGraphActionMenuBuilder& ActionMenuBuilder) const;
diff --git a/Engine/Plugins/Experimental/TextureGraph/Source/TextureGraphEditor/Private/STG_GraphActionMenu.cpp b/Engine/Plugins/Experimental/TextureGraph/Source/TextureGraphEditor/Private/STG_GraphActionMenu.cpp
deleted file mode 100644
diff --git a/Engine/Plugins/Experimental/TextureGraph/Source/TextureGraphEditor/Private/STG_GraphActionMenu.h b/Engine/Plugins/Experimental/TextureGraph/Source/TextureGraphEditor/Private/STG_GraphActionMenu.h
deleted file mode 100644
diff --git a/Engine/Plugins/Experimental/TextureGraph/Source/TextureGraphEditor/Private/STG_Palette.cpp b/Engine/Plugins/Experimental/TextureGraph/Source/TextureGraphEditor/Private/STG_Palette.cpp
--- a/Engine/Plugins/Experimental/TextureGraph/Source/TextureGraphEditor/Private/STG_Palette.cpp
+++ b/Engine/Plugins/Experimental/TextureGraph/Source/TextureGraphEditor/Private/STG_Palette.cpp
@@ -1,18 +1,21 @@
 // Copyright Epic Games, Inc. All Rights Reserved.
 
 #include "STG_Palette.h"
-#include "STG_GraphActionMenu.h"
+#include "Widgets/ActionMenuWidgets/STG_ActionMenuTileView.h"
 #include "Modules/ModuleManager.h"
 #include "Widgets/SOverlay.h"
 #include "Styling/CoreStyle.h"
 #include "Styling/AppStyle.h"
 #include "EdGraph/TG_EdGraphSchema.h"
+#include "EdGraph/TG_EdGraph.h"
 
 #include "EditorWidgetsModule.h"
 #include "AssetRegistry/AssetRegistryModule.h"
 #include "TG_Editor.h"
+#include "Brushes/SlateImageBrush.h"
 #include "Widgets/Input/STextComboBox.h"
 #include "Widgets/Text/SInlineEditableTextBlock.h"
+#include "Widgets/Input/SComboButton.h"
 #include "TG_Style.h"
 
 #define LOCTEXT_NAMESPACE "TGPalette"
@@ -26,117 +29,77 @@ void STG_PaletteItem::Construct(const FArguments& InArgs, FCreateWidgetForAction
 
 	const FSlateBrush* IconBrush = GetIconBrush();
 	FSlateColor IconColor = FSlateColor::UseForeground();
-	FText IconToolTip = GraphAction->GetTooltipDescription();
+	FText IconToolTip = GraphAction->GetCategory();
 	bool bIsReadOnly = false;
 
-	TSharedRef<SWidget> IconWidget = CreateIconWidget( IconToolTip, IconBrush, IconColor );
-	TSharedRef<SWidget> NameSlotWidget = CreateTextSlotWidget(InCreateData, bIsReadOnly );
+	TSharedRef<SWidget> IconWidget = CreateIconWidget(IconToolTip, IconBrush, IconColor);
+	TSharedRef<SWidget> NameSlotWidget = CreateTextSlotWidget(InCreateData, bIsReadOnly);
 	InlineRenameWidget->SetOverflowPolicy(ETextOverflowPolicy::Ellipsis);
 
 	auto NameAreaBackground = FAppStyle::Get().GetBrush(FName("ContentBrowser.AssetTileItem.NameAreaBackground"));
-	const int IconWidth = 70;
-	const int IconHeight = 80;
+	const int IconWidth = 22;
+	const int IconHeight = 22;
 	// Create the actual widget
 	this->ChildSlot
-	.Padding(FMargin(0.0f, 0.0f, 4.0f, 4.0f))
-	[
-		// Drop shadow border
-		SNew(SBorder)
-		.Padding(FMargin(0.0f, 0.0f, 5.0f, 5.0f))
-		.BorderImage( FAppStyle::Get().GetBrush(FName("ContentBrowser.AssetTileItem.DropShadow")))
 		[
-			SNew(SOverlay)
-			+ SOverlay::Slot()
-			[
-				SNew(SBorder)
-				.Padding(1)
-				.BorderImage(this,&STG_PaletteItem::GetBorderImage)
-				[
-					SNew(SVerticalBox)
-					// Thumbnail
-					+ SVerticalBox::Slot()
-					.AutoHeight()
-					[
-					// The remainder of the space is reserved for the name.
-						SNew(SBox)
-						.Padding(0)
-						.WidthOverride(IconWidth)
-						.HeightOverride(IconHeight)
-						[
-							//ItemContents
-							IconWidget
-						]
-					]
+			SNew(SHorizontalBox)
+			.ToolTipText(InCreateData->Action->GetTooltipDescription())
 
-					+ SVerticalBox::Slot()
-					.AutoHeight()
-					[
-						SNew(SBorder)
-						.BorderImage(FAppStyle::GetBrush("WhiteBrush"))
-						.BorderBackgroundColor(FSlateColor::UseForeground())
-						.Padding(FMargin(0,2,0,0))
-					]
+		//Icon
+		+ SHorizontalBox::Slot()
+		.AutoWidth()
+		.HAlign(HAlign_Left)
+		[
+			SNew(SBox)
+			.Padding(3)
+		.WidthOverride(IconWidth)
+		.HeightOverride(IconHeight)
+		[
+			IconWidget
+		]
+		]
 
-					+ SVerticalBox::Slot()
-					[
-						SNew(SBorder)
-						.Padding(FMargin(2.0f, 3.0f))
-						.BorderImage(NameAreaBackground)
-						[
-							SNew(SVerticalBox)
-							+ SVerticalBox::Slot()
-							.Padding(2.0f, 2.0f, 0.0f, 0.0f)
-							.VAlign(VAlign_Top)
-							.HAlign( HAlign_Left)
-							[
-								SNew(SBox)
-								.MaxDesiredHeight(14)
-								[
-									NameSlotWidget
-								]
-							]
-						]
-					]
-				]
-			]
+	+ SHorizontalBox::Slot()
+		.Padding(8, 0, 0, 0)
+		.HAlign(HAlign_Left)
+		.AutoWidth()
+		[
+			SNullWidget::NullWidget
 		]
-	];
+
+	+ SHorizontalBox::Slot()
+		.HAlign(HAlign_Left)
+		.VAlign(VAlign_Center)
+		.AutoWidth()
+		[
+			SNew(STextBlock)
+			.Font(FCoreStyle::GetDefaultFontStyle("Regular", 9))
+		.Text(InCreateData->Action->GetMenuDescription())
+		//.HighlightText(InArgs._HighlightText)
+		]
+		];
 }
 
-const FSlateBrush* STG_PaletteItem::GetIconBrush()
+FString STG_PaletteItem::GetIconBrushName()
 {
-	FString ActionName = ActionPtr.Pin()->GetMenuDescription().ToString();
-	FString BrushName = "TG_Editor.Palette." + ActionName;
-	if (FTG_Style::Get().HasKey(FName(BrushName)))
-	{
-		const FSlateBrush* IconBrush = FTG_Style::Get().GetBrush(FName(BrushName));
-		return IconBrush;
+	FString ActionName = ActionPtr.Pin()->GetCategory().ToString();
+	return "TG_Editor.Palette." + ActionName + "_Colored";
+}
 
-	}
-	return FTG_Style::Get().GetBrush(FName("TG_Editor.Palette.Default"));
+FString STG_PaletteItem::GetDefaultIconBrushName()
+{
+	return "TG_Editor.Palette.Default_Colored";
 }
 
-const FSlateBrush* STG_PaletteItem::GetBorderImage() const
+const FSlateBrush* STG_PaletteItem::GetIconBrush()
 {
-	const bool bIsSelected = false;
-	const bool bIsHoveredOrDraggedOver = IsHovered();
-	if (bIsSelected && bIsHoveredOrDraggedOver)
-	{
-		static const FName SelectedHover("ContentBrowser.AssetTileItem.SelectedHoverBorder");
-		return FAppStyle::Get().GetBrush(SelectedHover);
-	}
-	else if (bIsSelected)
-	{
-		static const FName Selected("ContentBrowser.AssetTileItem.SelectedBorder");
-		return FAppStyle::Get().GetBrush(Selected);
-	}
-	else if (bIsHoveredOrDraggedOver)
+	if (FTG_Style::Get().HasKey(FName(GetIconBrushName())))
 	{
-		static const FName Hovered("ContentBrowser.AssetTileItem.HoverBorder");
-		return FAppStyle::Get().GetBrush(Hovered);
-	}
+		const FSlateBrush* IconBrush = FTG_Style::Get().GetBrush(FName(GetIconBrushName()));
+		return IconBrush;
 
-	return FAppStyle::Get().GetBrush("ContentBrowser.AssetTileItem.ThumbnailAreaBackground");//FStyleDefaults::GetNoBrush();
+	}
+	return FTG_Style::Get().GetBrush(FName(GetDefaultIconBrushName()));
 }
 
 TSharedRef<SWidget> STG_PaletteItem::CreateHotkeyDisplayWidget(const TSharedPtr<const FInputChord> HotkeyChord)
@@ -155,12 +118,122 @@ FText STG_PaletteItem::GetItemTooltip() const
 	return ActionPtr.Pin()->GetTooltipDescription();
 }
 
+////////////////////////////////////////////////////////////////////
+
+void STG_PaletteTileItem::Construct(const FArguments& InArgs, FCreateWidgetForActionData* const InCreateData)
+{
+	check(InCreateData->Action.IsValid());
+
+	TSharedPtr<FEdGraphSchemaAction> GraphAction = InCreateData->Action;
+	ActionPtr = InCreateData->Action;
+
+	const FSlateBrush* IconBrush = GetIconBrush();
+	FText IconToolTip = GraphAction->GetTooltipDescription();
+	bool bIsReadOnly = false;
+
+	IconWidget = SNew(SImage)
+		.ToolTipText(IconToolTip)
+		.Image(IconBrush)
+		.ColorAndOpacity(FStyleColors::Foreground);
+
+	TSharedRef<SWidget> NameSlotWidget = CreateTextSlotWidget(InCreateData, bIsReadOnly );
+	InlineRenameWidget->SetOverflowPolicy(ETextOverflowPolicy::Ellipsis);
+
+	FLinearColor Color = UTG_EdGraphSchema::GetCategoryColor(FName(GraphAction->GetCategory().ToString()));
+
+	BackgroundBrush = *(FTG_Style::Get().GetBrush("TG.Palette.Backgorund"));
+	BackgroundBrush.OutlineSettings.Color = FSlateColor(Color);
+	
+	const int IconWidth = 16;
+	const int IconHeight = 16;
+	const int ItemHeight = 36;
+	const int IconPadding = (ItemHeight - IconHeight) / 2;
+	// Create the actual widget
+	this->ChildSlot
+	[
+		SNew(SOverlay)
+
+		+ SOverlay::Slot()
+		[
+			SAssignNew(BackgroundArea, SBox)
+			.Padding(FMargin(2))
+			.HAlign(HAlign_Fill)
+			.HeightOverride(ItemHeight)
+			[
+				SNew(SBorder)
+				.BorderImage(FAppStyle::GetBrush("UniformShadow"))
+				.Padding(2)
+				[
+					SNew(SImage)
+					.Image(&BackgroundBrush)
+					.ColorAndOpacity(Color)
+				]
+			]
+		]
+
+		+SOverlay::Slot()
+		[
+			SNew(SHorizontalBox)
+
+			+ SHorizontalBox::Slot()
+			.AutoWidth()
+			.VAlign(VAlign_Center)
+			[
+				SNew(SBox)
+				.Padding(IconPadding)
+				.WidthOverride(ItemHeight)
+				.HeightOverride(ItemHeight)
+				[
+					IconWidget.ToSharedRef()
+				]
+			]
+			
+			+ SHorizontalBox::Slot()
+			.Padding(0,0,2,0)
+			[
+				NameSlotWidget
+			]
+		]
+	];
+}
+
+void STG_PaletteTileItem::OnMouseEnter(const FGeometry& MyGeometry, const FPointerEvent& MouseEvent)
+{
+	SGraphPaletteItem::OnMouseEnter(MyGeometry, MouseEvent);
+	BackgroundArea->SetPadding(FMargin(0));
+	InlineRenameWidget->SetForegroundColor(FStyleColors::White);
+	IconWidget->SetColorAndOpacity(FStyleColors::White);
+}
+
+void STG_PaletteTileItem::OnMouseLeave(const FPointerEvent& MouseEvent)
+{
+	SGraphPaletteItem::OnMouseLeave(MouseEvent);
+	BackgroundArea->SetPadding(FMargin(2));
+	InlineRenameWidget->SetForegroundColor(FStyleColors::Foreground);
+	IconWidget->SetColorAndOpacity(FStyleColors::Foreground);
+}
+
+FString STG_PaletteTileItem::GetIconBrushName()
+{
+	FString ActionName = ActionPtr.Pin()->GetCategory().ToString();
+	return "TG_Editor.Palette." + ActionName;
+}
+
+FString STG_PaletteTileItem::GetDefaultIconBrushName()
+{
+	return "TG_Editor.Palette.Default";
+}
+
 //////////////////////////////////////////////////////////////////////////
 
 void STG_Palette::Construct(const FArguments& InArgs, TWeakPtr<FTG_Editor> InTGEditorPtr)
 {
 	TGEditorPtr = InTGEditorPtr;
 
+	BackgroundBrush = FSlateRoundedBoxBrush(FStyleColors::Input, 4.0f);
+	CheckedBrush = FSlateRoundedBoxBrush(FStyleColors::InputOutline, 4.0f); 
+	BackgroundHoverBrush = FSlateRoundedBoxBrush(FStyleColors::Panel, 4.0f);
+
 	// Create the asset discovery indicator
 	FEditorWidgetsModule& EditorWidgetsModule = FModuleManager::LoadModuleChecked<FEditorWidgetsModule>("EditorWidgets");
 	TSharedRef<SWidget> AssetDiscoveryIndicator = EditorWidgetsModule.CreateAssetDiscoveryIndicator(EAssetDiscoveryIndicatorScaleMode::Scale_Vertical);
@@ -169,6 +242,39 @@ void STG_Palette::Construct(const FArguments& InArgs, TWeakPtr<FTG_Editor> InTGE
 	CategoryNames.Add(MakeShareable(new FString(TEXT("Expressions"))));
 	CategoryNames.Add(MakeShareable(new FString(TEXT("Functions"))));
 
+	CreatePaletteMenu();
+
+	// Register with the Asset Registry to be informed when it is done loading up files.
+	FAssetRegistryModule& AssetRegistryModule = FModuleManager::GetModuleChecked<FAssetRegistryModule>(TEXT("AssetRegistry"));
+	AssetRegistryModule.Get().OnAssetAdded().AddSP(this, &STG_Palette::AddAssetFromAssetRegistry);
+	AssetRegistryModule.Get().OnAssetRemoved().AddSP(this, &STG_Palette::RemoveAssetFromRegistry);
+	AssetRegistryModule.Get().OnAssetRenamed().AddSP(this, &STG_Palette::RenameAssetFromRegistry);
+}
+
+void STG_Palette::CreatePaletteMenu()
+{
+	const int ToggleMenuHeight = 22;
+	const int ToggleMenuWidth = 47;
+	const int Padding = 5;
+
+	TSharedPtr<SWidget> Menu = IsTileViewChecked ? CreatePaletteTileMenu() : CreatePaletteListMenu();
+	
+	TSharedPtr<SWidget> FoundSearchBox = FindWidgetInChildren(Menu, "SSearchBox");
+	if (FoundSearchBox.IsValid())
+	{
+		TSharedPtr<SWidget> ParentWidget = FoundSearchBox->GetParentWidget();
+		TSharedPtr<SVerticalBox> Box = StaticCastSharedPtr<SVerticalBox>(ParentWidget);
+		if (Box)
+		{
+			Box->RemoveSlot(FoundSearchBox.ToSharedRef());
+		}
+	}
+
+	if (!ToggleButtons)
+	{
+		ToggleButtons = CreateViewToggle();
+	}
+
 	this->ChildSlot
 	[
 		SNew(SBorder)
@@ -177,72 +283,187 @@ void STG_Palette::Construct(const FArguments& InArgs, TWeakPtr<FTG_Editor> InTGE
 		[
 			SNew(SVerticalBox)
 
-			//// Filter UI
-			//+SVerticalBox::Slot()
-			//.AutoHeight()
-			//[
-			//	SNew(SHorizontalBox)
-
-			//	// Comment
-			//	+SHorizontalBox::Slot()
-			//	.VAlign(VAlign_Center)
-			//	.AutoWidth()
-			//	[
-			//		SNew(STextBlock)
-			//		.Text(LOCTEXT("Category", "Category: "))
-			//	]
-
-			//	// Combo button to select a class
-			//	+SHorizontalBox::Slot()
-			//		.VAlign(VAlign_Center)
-			//		[
-			//			SAssignNew(CategoryComboBox, STextComboBox)
-			//			.OptionsSource(&CategoryNames)
-			//			.OnSelectionChanged(this, &STG_Palette::CategorySelectionChanged)
-			//			.InitiallySelectedItem(CategoryNames[0])
-			//		]
-			//]
-
-			// Content list
 			+SVerticalBox::Slot()
+			.Padding(Padding)
+			.HAlign(HAlign_Fill)
+			.AutoHeight()
+			[
+				SNew(SHorizontalBox)
+
+				+ SHorizontalBox::Slot()
+				.Padding(0,0, Padding,0)
 				[
-					SNew(SOverlay)
+					FoundSearchBox.IsValid() ? FoundSearchBox.ToSharedRef() : SNullWidget::NullWidget
+				]
 
-					+SOverlay::Slot()
-					.HAlign(HAlign_Fill)
-					.VAlign(VAlign_Fill)
+				+ SHorizontalBox::Slot()
+				.HAlign(HAlign_Right)
+				.AutoWidth()
+				[
+					SNew(SBox)
+					.WidthOverride(ToggleMenuWidth)
+					.HeightOverride(ToggleMenuHeight)
 					[
-						// Old Expression and Function lists were auto expanded so do the same here for now
-						SAssignNew(TGGraphActionMenu, STG_GraphActionMenu)
-						.OnActionDragged(this, &STG_Palette::OnActionDragged)
-						.OnCreateWidgetForAction(this, &STG_Palette::OnCreateWidgetForAction)
-						.OnCollectAllActions(this, &STG_Palette::CollectAllActions)
-						.AutoExpandActionMenu(true)
+						ToggleButtons.ToSharedRef()
 					]
+				]
+			]
+
+			// Content list
+			+SVerticalBox::Slot()
+			[
+				SNew(SOverlay)
 
-					//+SOverlay::Slot()
-					//	.HAlign(HAlign_Fill)
-					//	.VAlign(VAlign_Bottom)
-					//	.Padding(FMargin(24, 0, 24, 0))
-					//	[
-					//		// Asset discovery indicator
-					//		AssetDiscoveryIndicator
-					//	]
+				+SOverlay::Slot()
+				.HAlign(HAlign_Fill)
+				.VAlign(VAlign_Fill)
+				[
+					Menu.ToSharedRef()
 				]
+			]
 
 		]
 	];
+}
 
-	// Register with the Asset Registry to be informed when it is done loading up files.
-	FAssetRegistryModule& AssetRegistryModule = FModuleManager::GetModuleChecked<FAssetRegistryModule>(TEXT("AssetRegistry"));
-	AssetRegistryModule.Get().OnAssetAdded().AddSP(this, &STG_Palette::AddAssetFromAssetRegistry);
-	AssetRegistryModule.Get().OnAssetRemoved().AddSP(this, &STG_Palette::RemoveAssetFromRegistry);
-	AssetRegistryModule.Get().OnAssetRenamed().AddSP(this, &STG_Palette::RenameAssetFromRegistry);
+TSharedPtr<SWidget> STG_Palette::FindWidgetInChildren(TSharedPtr<SWidget> Parent, FName ChildType) const
+{
+	TSharedPtr<SWidget> FoundWidget = TSharedPtr<SWidget>();
+	for (int32 ChildIndex = 0; ChildIndex < Parent->GetAllChildren()->Num(); ++ChildIndex)
+	{
+		TSharedRef<SWidget> ChildWidget = Parent->GetAllChildren()->GetChildAt(ChildIndex);
+		if (ChildWidget->GetType() == ChildType)
+		{
+			return ChildWidget.ToSharedPtr();
+		}
+		else if (ChildWidget->GetAllChildren()->Num() > 0)
+		{
+			FoundWidget = FindWidgetInChildren(ChildWidget, ChildType);
+			if (FoundWidget.IsValid() && FoundWidget->GetType() == ChildType)
+			{
+				return FoundWidget;
+			}
+		}
+	}
+	return FoundWidget;
+}
+
+TSharedRef<SWidget> STG_Palette::CreatePaletteTileMenu()
+{
+	return SAssignNew(TGActionMenuTileView,STG_ActionMenuTileView)
+		.OnActionDragged(this, &STG_Palette::OnActionDragged)
+		.OnCreateWidgetForAction(this, &STG_Palette::OnCreateWidgetForAction)
+		.OnCollectAllActions(this, &STG_Palette::CollectAllActions)
+		.AutoExpandActionMenu(true);
+}
+
+TSharedRef<SWidget> STG_Palette::CreatePaletteListMenu()
+{
+	return SNew(STG_GraphActionMenu)
+		.OnActionDragged(this, &STG_Palette::OnActionDragged)
+		.AutoExpandActionMenu(true)
+		.bSpwanOnSelect(false)
+		.GraphObj(Cast<UTG_EdGraph>(TGEditorPtr.Pin()->TG_EdGraph));
+}
+
+TSharedRef<SWidget> STG_Palette::CreateViewToggle()
+{
+	auto OnToggleChanged = [this](ECheckBoxState NewState, bool IsTileButton)
+	{
+		if (IsTileButton)
+		{
+			IsTileViewChecked = NewState == ECheckBoxState::Checked;
+			IsListViewChecked = !IsTileViewChecked;
+		}
+		else
+		{
+			IsListViewChecked = NewState == ECheckBoxState::Checked;
+			IsTileViewChecked = !IsListViewChecked;
+		}
+
+		CreatePaletteMenu();
+	};
+
+	return SNew(SBorder)
+	.BorderImage(&BackgroundBrush)
+	[
+
+		SNew(SHorizontalBox)
+
+		+ SHorizontalBox::Slot()
+		[
+			SNew(SOverlay)
+
+			+SOverlay::Slot()
+			[
+				SNew(SCheckBox)
+				.Style(&FAppStyle::Get().GetWidgetStyle<FCheckBoxStyle>("DetailsView.ChannelToggleButton"))
+				.Type(ESlateCheckBoxType::ToggleButton)
+				.CheckedImage(&CheckedBrush)
+				.CheckedHoveredImage(&CheckedBrush)
+				.UncheckedImage(&BackgroundBrush)
+				.UncheckedHoveredImage(&BackgroundHoverBrush)
+				.OnCheckStateChanged_Lambda(OnToggleChanged, true)
+				.IsChecked(this, &STG_Palette::OnGetToggleCheckState, true)
+				.Visibility(this,&STG_Palette::GetVisibility,true)
+			]
+			+ SOverlay::Slot()
+			[
+				SNew(SImage)
+				.Image(FTG_Style::Get().GetBrush("TG_Editor.TileIcon"))
+				.Visibility(EVisibility::HitTestInvisible)
+			]
+		]
+
+		+ SHorizontalBox::Slot()
+		[
+			SNew(SOverlay)
+
+			+ SOverlay::Slot()
+			[
+				SNew(SCheckBox)
+				.Style(&FAppStyle::Get().GetWidgetStyle<FCheckBoxStyle>("DetailsView.ChannelToggleButton"))
+				.Type(ESlateCheckBoxType::ToggleButton)
+				.CheckedImage(&CheckedBrush)
+				.CheckedHoveredImage(&CheckedBrush)
+				.UncheckedImage(&BackgroundBrush)
+				.UncheckedHoveredImage(&BackgroundHoverBrush)
+				.OnCheckStateChanged_Lambda(OnToggleChanged, false)
+				.IsChecked(this, &STG_Palette::OnGetToggleCheckState, false)
+				.Visibility(this, &STG_Palette::GetVisibility, false)
+			]
+			+ SOverlay::Slot()
+			[
+				SNew(SImage)
+				.Image(FTG_Style::Get().GetBrush("TG_Editor.ListIcon"))
+				.Visibility(EVisibility::HitTestInvisible)
+			]
+		]
+	];
+//.ColorAndOpacity(Color)
+}
+
+EVisibility STG_Palette::GetVisibility(bool IsTileView) const
+{
+	if (IsTileView)
+	{
+		return IsTileViewChecked ? EVisibility::HitTestInvisible : EVisibility::Visible;
+	}
+	else return IsListViewChecked ? EVisibility::HitTestInvisible : EVisibility::Visible;
+}
+
+ECheckBoxState STG_Palette::OnGetToggleCheckState(bool IsTileView) const
+{
+	if (IsTileView)
+	{
+		return IsTileViewChecked ? ECheckBoxState::Checked : ECheckBoxState::Unchecked;
+	}
+	else return IsListViewChecked ? ECheckBoxState::Checked : ECheckBoxState::Unchecked;
 }
 
 TSharedRef<SWidget> STG_Palette::OnCreateWidgetForAction(FCreateWidgetForActionData* const InCreateData)
 {
-	return	SNew(STG_PaletteItem, InCreateData);
+	return	SNew(STG_PaletteTileItem, InCreateData);
 }
 
 void STG_Palette::CollectAllActions(FGraphActionListBuilderBase& OutAllActions)
@@ -297,7 +518,7 @@ void STG_Palette::RefreshAssetInRegistry(const FAssetData& InAddedAssetData)
 
 void STG_Palette::RefreshActionsList(bool bPreserveExpansion)
 {
-	TGGraphActionMenu->RefreshAllActions(bPreserveExpansion);
+	TGActionMenuTileView->RefreshAllActions(bPreserveExpansion);
 }
 
 #undef LOCTEXT_NAMESPACE
diff --git a/Engine/Plugins/Experimental/TextureGraph/Source/TextureGraphEditor/Private/STG_Palette.h b/Engine/Plugins/Experimental/TextureGraph/Source/TextureGraphEditor/Private/STG_Palette.h
--- a/Engine/Plugins/Experimental/TextureGraph/Source/TextureGraphEditor/Private/STG_Palette.h
+++ b/Engine/Plugins/Experimental/TextureGraph/Source/TextureGraphEditor/Private/STG_Palette.h
@@ -7,29 +7,62 @@
 #include "Widgets/SWidget.h"
 #include "Framework/Commands/InputChord.h"
 #include "Widgets/DeclarativeSyntaxSupport.h"
+#include "Widgets/ActionMenuWidgets/STG_GraphActionMenu.h"
+#include <Widgets/Input/SCheckBox.h>
 #include "SGraphPalette.h"
 
+#include "Brushes/SlateRoundedBoxBrush.h"
+
 class FTG_Editor;
-class STG_GraphActionMenu;
-/** Widget for displaying a single item  */
+class STG_ActionMenuTileView;
+class SBox;
+class STGGraphActionMenu;
+
+/* This class is the custom Graph action menu item for Texture Graph*/
 class STG_PaletteItem : public SGraphPaletteItem
 {
 public:
-	SLATE_BEGIN_ARGS( STG_PaletteItem ) {};
+	SLATE_BEGIN_ARGS(STG_PaletteItem) {};
 	SLATE_END_ARGS()
 
-	void Construct(const FArguments& InArgs, FCreateWidgetForActionData* const InCreateData);
+		void Construct(const FArguments& InArgs, FCreateWidgetForActionData* const InCreateData);
 
-private:
+protected:
+	virtual FString GetIconBrushName();
+
+	virtual FString GetDefaultIconBrushName();
 	/* Create the hotkey display widget */
 	TSharedRef<SWidget> CreateHotkeyDisplayWidget(const TSharedPtr<const FInputChord> HotkeyChord);
 
 	const FSlateBrush* GetIconBrush();
-	const FSlateBrush* GetBorderImage() const;
 
 	virtual FText GetItemTooltip() const override;
 };
 
+/** Widget for displaying a single item  */
+class STG_PaletteTileItem : public STG_PaletteItem
+{
+public:
+	SLATE_BEGIN_ARGS( STG_PaletteTileItem ) {};
+	SLATE_END_ARGS()
+
+	void Construct(const FArguments& InArgs, FCreateWidgetForActionData* const InCreateData);
+
+	virtual void OnMouseEnter(const FGeometry& MyGeometry, const FPointerEvent& MouseEvent) override;
+
+	virtual void OnMouseLeave(const FPointerEvent& MouseEvent) override;
+
+protected:
+	virtual FString GetIconBrushName() override;
+
+	virtual FString GetDefaultIconBrushName() override;
+
+private:
+	TSharedPtr<SBox> BackgroundArea;
+	TSharedPtr<SImage> IconWidget;
+	FSlateBrush BackgroundBrush;
+};
+
 //////////////////////////////////////////////////////////////////////////
 
 class STG_Palette : public SGraphPalette
@@ -47,6 +80,16 @@ protected:
 	virtual void RefreshActionsList(bool bPreserveExpansion) override;
 	// End of SGraphPalette Interface
 
+	void CreatePaletteMenu();
+
+	TSharedRef<SWidget> CreatePaletteTileMenu();
+
+	TSharedPtr<SWidget> FindWidgetInChildren(TSharedPtr<SWidget> Parent, FName ChildType) const;
+
+	TSharedRef<SWidget> CreatePaletteListMenu();
+
+	TSharedRef<SWidget> CreateViewToggle();
+
 	/** Get the currently selected category name */
 	FString GetFilterCategoryName() const;
 
@@ -64,7 +107,16 @@ protected:
 
 	void RefreshAssetInRegistry(const FAssetData& InAddedAssetData);
 
+	EVisibility GetVisibility(bool IsTileView) const;
+
+	ECheckBoxState OnGetToggleCheckState(bool IsTileView) const;
+
 protected:
+	bool IsTileViewChecked = true;
+	bool IsListViewChecked = false;
+	FSlateBrush BackgroundBrush;
+	FSlateBrush BackgroundHoverBrush;
+	FSlateBrush CheckedBrush;
 	/** Pointer back to the material editor that owns us */
 	TWeakPtr<FTG_Editor> TGEditorPtr;
 
@@ -74,5 +126,7 @@ protected:
 	/** Combo box used to select category */
 	TSharedPtr<STextComboBox> CategoryComboBox;
 
-	TSharedPtr<STG_GraphActionMenu> TGGraphActionMenu;
+	TSharedPtr<SWidget> ToggleButtons;
+
+	TSharedPtr<STG_ActionMenuTileView> TGActionMenuTileView;
 };
diff --git a/Engine/Plugins/Experimental/TextureGraph/Source/TextureGraphEditor/Private/TG_Editor.cpp b/Engine/Plugins/Experimental/TextureGraph/Source/TextureGraphEditor/Private/TG_Editor.cpp
--- a/Engine/Plugins/Experimental/TextureGraph/Source/TextureGraphEditor/Private/TG_Editor.cpp
+++ b/Engine/Plugins/Experimental/TextureGraph/Source/TextureGraphEditor/Private/TG_Editor.cpp
@@ -40,7 +40,7 @@
 #include "EditorViewportTabContent.h"
 #include "Framework/Application/SlateApplication.h"
 #include "Framework/Commands/GenericCommands.h"
-#include "GraphEditor.h"
+//#include "GraphEditor.h"
 #include "SGraphPanel.h"
 #include "Kismet2/BlueprintEditorUtils.h"
 #include "PropertyEditorModule.h"
@@ -71,6 +71,7 @@
 
 #include "Expressions/Input/TG_Expression_InputParam.h"
 #include "STG_OutputSelectionDlg.h"
+#include "Widgets/ActionMenuWidgets/STG_GraphActionMenu.h"
 
 DEFINE_LOG_CATEGORY(LogTextureGraphEditor);
 #define LOCTEXT_NAMESPACE "TG_Editor"
@@ -753,6 +754,7 @@ TSharedRef<class SGraphEditor> FTG_Editor::CreateGraphEditorWidget()
 	InEvents.OnNodeDoubleClicked = FSingleNodeEvent::CreateSP(this, &FTG_Editor::OnNodeDoubleClicked);
 	InEvents.OnTextCommitted = FOnNodeTextCommitted::CreateSP(this, &FTG_Editor::OnNodeTitleCommitted);
 	InEvents.OnVerifyTextCommit = FOnNodeVerifyTextCommit::CreateSP(this, &FTG_Editor::OnVerifyNodeTextCommit);
+	InEvents.OnCreateActionMenu = SGraphEditor::FOnCreateActionMenu::CreateSP(this, &FTG_Editor::OnCreateGraphActionMenu);
 	//InEvents.OnSpawnNodeByShortcut = SGraphEditor::FOnSpawnNodeByShortcut::CreateSP(this, &FTG_Editor::OnSpawnGraphNodeByShortcut, static_cast<UEdGraph*>(InGraph));
 
 	// Support the delete case
@@ -791,6 +793,20 @@ TSharedRef<class SGraphEditor> FTG_Editor::CreateGraphEditorWidget()
 		.AssetEditorToolkit(this->AsShared());
 }
 
+FActionMenuContent FTG_Editor::OnCreateGraphActionMenu(UEdGraph* InGraph, const FVector2D& InNodePosition, const TArray<UEdGraphPin*>& InDraggedPins, bool bAutoExpand, SGraphEditor::FActionMenuClosed InOnMenuClosed)
+{
+	TSharedRef<STG_GraphActionMenu> ActionMenu =
+		SNew(STG_GraphActionMenu)
+		.GraphObj(InGraph)
+		.NewNodePosition(InNodePosition)
+		.DraggedFromPins(InDraggedPins)
+		.AutoExpandActionMenu(bAutoExpand)
+		.OnClosedCallback(InOnMenuClosed);
+		//.OnCloseReason(this, &FBlueprintEditor::OnGraphActionMenuClosed);
+
+	return FActionMenuContent(ActionMenu, ActionMenu->GetFilterTextBox());
+}
+
 TSharedRef<class STG_SelectionPreview> FTG_Editor::CreateSelectionViewWidget()
 {
 	return SNew(STG_SelectionPreview)
diff --git a/Engine/Plugins/Experimental/TextureGraph/Source/TextureGraphEditor/Private/TG_Editor.h b/Engine/Plugins/Experimental/TextureGraph/Source/TextureGraphEditor/Private/TG_Editor.h
--- a/Engine/Plugins/Experimental/TextureGraph/Source/TextureGraphEditor/Private/TG_Editor.h
+++ b/Engine/Plugins/Experimental/TextureGraph/Source/TextureGraphEditor/Private/TG_Editor.h
@@ -15,6 +15,7 @@
 #include "TG_Parameter.h"
 #include "EdGraph/TG_EdGraphNode.h"
 #include "STG_Palette.h"
+#include "GraphEditor.h"
 
 class UTextureGraph;
 class UTG_Graph;
@@ -181,6 +182,8 @@ private:
 	/** Create new graph editor widget */
 	TSharedRef<class SGraphEditor>					CreateGraphEditorWidget();
 
+	FActionMenuContent								OnCreateGraphActionMenu(UEdGraph* InGraph, const FVector2D& InNodePosition, const TArray<UEdGraphPin*>& InDraggedPins, bool bAutoExpand, SGraphEditor::FActionMenuClosed InOnMenuClosed);
+
 	/** Create Selection view widget */
 	TSharedRef<class STG_SelectionPreview>			CreateSelectionViewWidget();
 
diff --git a/Engine/Plugins/Experimental/TextureGraph/Source/TextureGraphEditor/Private/TG_Style.cpp b/Engine/Plugins/Experimental/TextureGraph/Source/TextureGraphEditor/Private/TG_Style.cpp
--- a/Engine/Plugins/Experimental/TextureGraph/Source/TextureGraphEditor/Private/TG_Style.cpp
+++ b/Engine/Plugins/Experimental/TextureGraph/Source/TextureGraphEditor/Private/TG_Style.cpp
@@ -29,7 +29,7 @@ void FTG_Style::Unregister()
 TArray<FString> FTG_Style::GetPaletteIconNames()
 {
 	FString FolderPath = FPaths::Combine(IPluginManager::Get().FindPlugin("TextureGraph")->GetBaseDir() , TEXT("Content/Style/Palette/") );
-	FString FileExtension = ".png";
+	FString FileExtension = ".svg";
 
 	IPlatformFile& PlatformFile = FPlatformFileManager::Get().GetPlatformFile();
 	TArray<FString> ImageFileNames;
@@ -96,7 +96,7 @@ FTG_Style::FTG_Style() : FSlateStyleSet("TG_Style")
 	
 	const FVector2D PinSize(20.45f, 15.0f);
 	Set(TSEditorStyleConstants::Pin_Generic_Image_C, new IMAGE_BRUSH_SVG("Style/TG_Pin_Generic_C", PinSize));
-	Set(TSEditorStyleConstants::Pin_Generic_Image_DC, new IMAGE_BRUSH_SVG("Style/TG_Pin_Generic_DC", PinSize));
+	Set(TSEditorStyleConstants::Pin_Generic_Image_DC, new IMAGE_BRUSH_SVG("Style/TG_Pin_Generic_DC", FVector2D(15,15)));
 	Set(TSEditorStyleConstants::Pin_IN_Image_C, new IMAGE_BRUSH_SVG("Style/TG_Pin_IN_Image_C", PinSize));
 	Set(TSEditorStyleConstants::Pin_IN_Image_DC, new IMAGE_BRUSH_SVG("Style/TG_Pin_IN_Image_Unplugged", PinSize));
 	Set(TSEditorStyleConstants::Pin_IN_Vector_C, new IMAGE_BRUSH_SVG("Style/TG_Pin_IN_Vector_C", PinSize));
@@ -112,24 +112,31 @@ FTG_Style::FTG_Style() : FSlateStyleSet("TG_Style")
 	Set(TSEditorStyleConstants::Pin_OUT_Scalar_C, new IMAGE_BRUSH_SVG("Style/TG_Pin_IN_Scalar_C", PinSize));
 	Set(TSEditorStyleConstants::Pin_OUT_Scalar_DC, new IMAGE_BRUSH_SVG("Style/TG_Pin_IN_Scalar_Unplugged", PinSize));
 
+	Set("TG_Editor.TileIcon", new IMAGE_BRUSH_SVG("Style/TileIcon", Icon20x20));
+	Set("TG_Editor.ListIcon", new IMAGE_BRUSH_SVG("Style/ListIcon", Icon20x20));
+
 	//Creating brush for every icon in the palatte folder
 	for (auto PlatteIconName : GetPaletteIconNames())
 	{
 		FString Path = "Style/Palette/" + PlatteIconName;
 		FString PropertyName = "TG_Editor.Palette." + PlatteIconName;
 
-		Set(FName(PropertyName), new IMAGE_BRUSH(Path, PinSize));
+		Set(FName(PropertyName), new IMAGE_BRUSH_SVG(Path, Icon16x16));
 	}
 	FLinearColor NoSpillColor(1, 1, 1, 1.0);
-	int BodyRadius = 4.0;
+	int BodyRadius = 10;
+	int NodeHeaderRadius = 7;
+	int PalleteRadius = 4;
 
 	Set("TG.Graph.Node.BodyBackground", new FSlateRoundedBoxBrush(NoSpillColor, BodyRadius));
 	Set("TG.Graph.Node.BodyBorder", new FSlateRoundedBoxBrush(NoSpillColor, BodyRadius));
 	Set("TG.Graph.Node.AssetBackground", new FSlateRoundedBoxBrush(NoSpillColor, BodyRadius));
 
 	Set("TG.Graph.Node.Body", new FSlateRoundedBoxBrush(NoSpillColor, BodyRadius, NoSpillColor, 2.0));
-	Set("TG.Graph.Node.NoColorSpill", new FSlateRoundedBoxBrush(NoSpillColor, FVector4(BodyRadius, BodyRadius, 0.0, 0.0)));
+	Set("TG.Graph.Node.Header", new FSlateRoundedBoxBrush(NoSpillColor, NodeHeaderRadius, NoSpillColor, 2.0));
 	Set("TG.Graph.Node.ShadowSelected", new BOX_BRUSH("Style/TG_shadow_selected", FMargin(18.0/64.0)));
+
+	Set("TG.Palette.Backgorund", new FSlateRoundedBoxBrush(NoSpillColor, PalleteRadius, NoSpillColor, 2.0));
 	
 	SetParentStyleName("EditorStyle");
 	FTextBlockStyle NormalText = GetParentStyle()->GetWidgetStyle<FTextBlockStyle>("NormalText");
diff --git a/Engine/Plugins/Experimental/TextureGraph/Source/TextureGraphEditor/Private/Widgets/ActionMenuWidgets/STG_ActionMenuTileView.cpp b/Engine/Plugins/Experimental/TextureGraph/Source/TextureGraphEditor/Private/Widgets/ActionMenuWidgets/STG_ActionMenuTileView.cpp
new file mode 100644
--- /dev/null
+++ b/Engine/Plugins/Experimental/TextureGraph/Source/TextureGraphEditor/Private/Widgets/ActionMenuWidgets/STG_ActionMenuTileView.cpp
@@ -0,0 +1,1426 @@
+// Copyright Epic Games, Inc. All Rights Reserved.
+
+#include "STG_ActionMenuTileView.h"
+#include "EdGraph/EdGraph.h"
+#include "EdGraph/EdGraphNode.h"
+#include "EdGraphSchema_K2.h"
+#include "EdGraphSchema_K2_Actions.h"
+#include "EditorCategoryUtils.h"
+#include "Fonts/SlateFontInfo.h"
+#include "Framework/Application/SlateApplication.h"
+#include "Framework/Views/ITypedTableView.h"
+#include "GraphActionNode.h"
+#include "GraphEditorDragDropAction.h"
+#include "IDocumentation.h"
+#include "Input/DragAndDrop.h"
+#include "Input/Events.h"
+#include "InputCoreTypes.h"
+#include "Internationalization/Internationalization.h"
+#include "K2Node.h"
+#include "Layout/Children.h"
+#include "Layout/Margin.h"
+#include "Layout/Visibility.h"
+#include "Math/NumericLimits.h"
+#include "Math/UnrealMathSSE.h"
+#include "Misc/AssertionMacros.h"
+#include "Misc/CString.h"
+#include "ProfilingDebugging/CpuProfilerTrace.h"
+#include "SlotBase.h"
+#include "Styling/AppStyle.h"
+#include "Styling/CoreStyle.h"
+#include "Styling/ISlateStyle.h"
+#include "Styling/SlateTypes.h"
+#include "Styling/StyleColors.h"
+#include "UObject/ObjectPtr.h"
+#include "UObject/UObjectGlobals.h"
+#include "UObject/UnrealNames.h"
+#include "Widgets/Input/SSearchBox.h"
+#include "Widgets/Layout/SBorder.h"
+#include "Widgets/Layout/SScrollBorder.h"
+#include "Widgets/Layout/SScrollBox.h"
+#include "Widgets/Layout/SSeparator.h"
+#include "Widgets/SBoxPanel.h"
+#include "Widgets/SNullWidget.h"
+#include "Widgets/SToolTip.h"
+#include "Widgets/Text/SInlineEditableTextBlock.h"
+#include "Widgets/Text/SRichTextBlock.h"
+#include "Widgets/Text/STextBlock.h"
+#include "Widgets/Views/STableRow.h"
+#include "Widgets/Input/SButton.h"
+
+#define LOCTEXT_NAMESPACE "STG_ActionMenuTileView"
+
+const int STG_ActionMenuTileView::TileWidth = 169;
+const int STG_ActionMenuTileView::TileHeight = 36;
+
+class STG_CategoryHeaderExpanderArrow : public SExpanderArrow
+{
+	SLATE_BEGIN_ARGS(STG_CategoryHeaderExpanderArrow) {}
+
+	SLATE_ARGUMENT(bool, LeftAligned)
+		SLATE_ARGUMENT(bool, AlwaysVisible)
+
+		SLATE_END_ARGS()
+
+	void Construct(const FArguments& InArgs, const TSharedPtr<class ITableRow>& TableRow, const TSharedPtr< STreeView< TSharedPtr<FGraphActionNode> > >& TreeView, TSharedPtr<FGraphActionNode> Item)
+	{
+		bLeftAligned = InArgs._LeftAligned;
+		bAlwaysVisible = InArgs._AlwaysVisible;
+		OwnerTreeView = TreeView;
+		OwnerItem = Item;
+
+		SExpanderArrow::Construct(
+			SExpanderArrow::FArguments()
+			.BaseIndentLevel(1),
+			TableRow);
+
+		ExpanderArrow->SetOnClicked(FOnClicked::CreateSP(this, &STG_CategoryHeaderExpanderArrow::OnClick));
+
+		// Override visibility so that groups can have an expander arrow despite having no children.
+		ExpanderArrow->SetVisibility(
+			TAttribute<EVisibility>(this, &STG_CategoryHeaderExpanderArrow::GetExpanderVisibility_Extended));
+	}
+
+	EVisibility GetExpanderVisibility_Extended() const
+	{
+		if (bAlwaysVisible)
+		{
+			return EVisibility::Visible;
+		}
+		return SExpanderArrow::GetExpanderVisibility();
+	}
+
+	FReply OnClick()
+	{
+		OwnerTreeView.Pin()->SetItemExpansion(OwnerItem.Pin(), !OwnerTreeView.Pin()->IsItemExpanded(OwnerItem.Pin()));
+
+		return FReply::Handled();
+	}
+
+	virtual FCursorReply OnCursorQuery(const FGeometry& MyGeometry, const FPointerEvent& CursorEvent) const
+	{
+		return FCursorReply::Cursor(EMouseCursor::Default);
+	}
+
+	bool bLeftAligned;
+	bool bAlwaysVisible;
+	TWeakPtr< STreeView< TSharedPtr<FGraphActionNode> > > OwnerTreeView;
+	TWeakPtr < FGraphActionNode > OwnerItem;
+};
+
+template<typename ItemType>
+class STG_CategoryHeaderTableRow : public STableRow<ItemType>
+{
+public:
+	SLATE_BEGIN_ARGS(STG_CategoryHeaderTableRow)
+	{}
+	SLATE_DEFAULT_SLOT(typename STG_CategoryHeaderTableRow::FArguments, Content)
+	SLATE_END_ARGS()
+
+	void Construct(const FArguments& InArgs, const TSharedRef<STableViewBase>& InOwnerTableView)
+	{
+		STableRow<ItemType>::ChildSlot
+		.Padding(0.0f, 2.0f, .0f, 0.0f)
+		[
+			SAssignNew(ContentBorder, SBorder)
+			.BorderImage(this, &STG_CategoryHeaderTableRow::GetBackgroundImage)
+			.Padding(FMargin(3.0f, 5.0f))
+			[
+				SNew(SHorizontalBox)
+				+ SHorizontalBox::Slot()
+				.VAlign(VAlign_Center)
+				.Padding(5.0f)
+				.AutoWidth()
+				[
+					SNew(SExpanderArrow, STableRow< ItemType >::SharedThis(this))
+				]
+				+ SHorizontalBox::Slot()
+				.VAlign(VAlign_Center)
+				.AutoWidth()
+				[
+					InArgs._Content.Widget
+				]
+			]
+		];
+
+		STableRow < ItemType >::ConstructInternal(
+			typename STableRow< ItemType >::FArguments()
+			.Style(FAppStyle::Get(), "DetailsView.TreeView.TableRow")
+			.ShowSelection(false),
+			InOwnerTableView
+		);
+	}
+
+	const FSlateBrush* GetBackgroundImage() const
+	{
+		if (STableRow<ItemType>::IsHovered())
+		{
+			return FAppStyle::Get().GetBrush("Brushes.Secondary");
+		}
+		else
+		{
+			return FAppStyle::Get().GetBrush("Brushes.Header");
+		}
+	}
+
+	virtual void SetContent(TSharedRef< SWidget > InContent) override
+	{
+		ContentBorder->SetContent(InContent);
+	}
+
+	virtual void SetRowContent(TSharedRef< SWidget > InContent) override
+	{
+		ContentBorder->SetContent(InContent);
+	}
+
+	virtual const FSlateBrush* GetBorder() const
+	{
+		return nullptr;
+	}
+
+	FReply OnMouseButtonDown(const FGeometry& MyGeometry, const FPointerEvent& MouseEvent)
+	{
+		if (MouseEvent.GetEffectingButton() == EKeys::LeftMouseButton)
+		{
+			STableRow<ItemType>::ToggleExpansion();
+			return FReply::Handled();
+		}
+		else
+		{
+			return FReply::Unhandled();
+		}
+	}
+private:
+	TSharedPtr<SBorder> ContentBorder;
+};
+
+class STG_GraphActionCategoryWidget : public SCompoundWidget
+{
+	SLATE_BEGIN_ARGS(STG_GraphActionCategoryWidget)
+	{}
+	SLATE_ATTRIBUTE(FText, HighlightText)
+	SLATE_EVENT(FOnTextCommitted, OnTextCommitted)
+	SLATE_EVENT(FIsSelected, IsSelected)
+	SLATE_ATTRIBUTE(bool, IsReadOnly)
+	SLATE_END_ARGS()
+
+	TWeakPtr<FGraphActionNode> ActionNode;
+	TAttribute<bool> IsReadOnly;
+public:
+	TWeakPtr<SInlineEditableTextBlock> InlineWidget;
+
+	void Construct(const FArguments& InArgs, TSharedPtr<FGraphActionNode> InActionNode)
+	{
+		ActionNode = InActionNode;
+
+		FText CategoryTooltip;
+		FString CategoryLink, CategoryExcerpt;
+		FEditorCategoryUtils::GetCategoryTooltipInfo(*InActionNode->GetDisplayName().ToString(), CategoryTooltip, CategoryLink, CategoryExcerpt);
+
+		TSharedRef<SToolTip> ToolTipWidget = IDocumentation::Get()->CreateToolTip(CategoryTooltip, NULL, CategoryLink, CategoryExcerpt);
+		IsReadOnly = InArgs._IsReadOnly;
+
+		this->ChildSlot
+		[
+			SNew(SHorizontalBox)
+			+ SHorizontalBox::Slot()
+			.VAlign(VAlign_Center)
+			[
+				SAssignNew(InlineWidget, SInlineEditableTextBlock)
+				.Font(FCoreStyle::GetDefaultFontStyle("Bold", 9))
+				.Text(FEditorCategoryUtils::GetCategoryDisplayString(InActionNode->GetDisplayName()))
+				.ToolTip(ToolTipWidget)
+				.HighlightText(InArgs._HighlightText)
+				.OnTextCommitted(InArgs._OnTextCommitted)
+				.IsSelected(InArgs._IsSelected)
+				.IsReadOnly(InArgs._IsReadOnly)
+			]
+		];
+	}
+
+	// SWidget interface
+	virtual FReply OnDrop(const FGeometry& MyGeometry, const FDragDropEvent& DragDropEvent) override
+	{
+		TSharedPtr<FGraphEditorDragDropAction> GraphDropOp = DragDropEvent.GetOperationAs<FGraphEditorDragDropAction>();
+		if (GraphDropOp.IsValid())
+		{
+			GraphDropOp->DroppedOnCategory(ActionNode.Pin()->GetCategoryPath());
+			return FReply::Handled();
+		}
+		return FReply::Unhandled();
+	}
+
+	virtual void OnDragEnter(const FGeometry& MyGeometry, const FDragDropEvent& DragDropEvent) override
+	{
+		TSharedPtr<FGraphEditorDragDropAction> GraphDropOp = DragDropEvent.GetOperationAs<FGraphEditorDragDropAction>();
+		if (GraphDropOp.IsValid())
+		{
+			GraphDropOp->SetHoveredCategoryName(ActionNode.Pin()->GetDisplayName());
+		}
+	}
+
+	virtual void OnDragLeave(const FDragDropEvent& DragDropEvent) override
+	{
+		TSharedPtr<FGraphEditorDragDropAction> GraphDropOp = DragDropEvent.GetOperationAs<FGraphEditorDragDropAction>();
+		if (GraphDropOp.IsValid())
+		{
+			GraphDropOp->SetHoveredCategoryName(FText::GetEmpty());
+		}
+	}
+
+	// End of SWidget interface
+};
+
+
+FString STG_ActionMenuTileView::LastUsedFilterText;
+
+void STG_ActionMenuTileView::Construct(const FArguments& InArgs, bool bIsReadOnly/* = true*/)
+{
+	this->SelectedSuggestionScore = TNumericLimits<float>::Lowest();
+	this->SelectedSuggestionSourceIndex = INDEX_NONE;
+	this->SelectedSuggestion = INDEX_NONE;
+	this->bIgnoreUIUpdate = false;
+	this->bUseSectionStyling = InArgs._UseSectionStyling;
+	this->bAllowPreselectedItemActivation = InArgs._bAllowPreselectedItemActivation;
+
+	this->bAutoExpandActionMenu = InArgs._AutoExpandActionMenu;
+	this->bShowFilterTextBox = InArgs._ShowFilterTextBox;
+	this->bAlphaSortItems = InArgs._AlphaSortItems;
+	this->bSortItemsRecursively = InArgs._SortItemsRecursively;
+	this->OnActionSelected = InArgs._OnActionSelected;
+	this->OnActionDoubleClicked = InArgs._OnActionDoubleClicked;
+	this->OnActionDragged = InArgs._OnActionDragged;
+	this->OnCategoryDragged = InArgs._OnCategoryDragged;
+	this->OnCreateWidgetForAction = InArgs._OnCreateWidgetForAction;
+	this->OnCreateCustomRowExpander = InArgs._OnCreateCustomRowExpander;
+	this->OnGetActionList = InArgs._OnGetActionList;
+	this->OnCollectAllActions = InArgs._OnCollectAllActions;
+	this->OnCollectStaticSections = InArgs._OnCollectStaticSections;
+	this->OnCategoryTextCommitted = InArgs._OnCategoryTextCommitted;
+	this->OnCanRenameSelectedAction = InArgs._OnCanRenameSelectedAction;
+	this->OnGetSectionTitle = InArgs._OnGetSectionTitle;
+	this->OnGetSectionToolTip = InArgs._OnGetSectionToolTip;
+	this->OnGetSectionWidget = InArgs._OnGetSectionWidget;
+	this->FilteredRootAction = FGraphActionNode::NewRootNode();
+	this->OnActionMatchesName = InArgs._OnActionMatchesName;
+	this->DraggedFromPins = InArgs._DraggedFromPins;
+	this->GraphObj = InArgs._GraphObj;
+
+	// Default graph action list (also provides an empty source list to start with)
+	AllActions = MakeShared<FGraphActionListBuilderBase>();
+
+	if (OnGetActionList.IsBound())
+	{
+		// If we are obtaining a new action list at refresh time, ensure that the indirect collection delegate is unbound
+		if (!ensureMsgf(!OnCollectAllActions.IsBound(), TEXT("The OnCollectAllActions delegate is bound, but will not be invoked, because OnGetActionList has also been bound and will be used to obtain the action list for this menu. To resolve this, one of these events should be removed from its construction.")))
+		{
+			OnCollectAllActions.Unbind();
+		}
+	}
+
+	// If a delegate for filtering text is passed in, assign it so that it will be used instead of the built-in filter box
+	if (InArgs._OnGetFilterText.IsBound())
+	{
+		this->OnGetFilterText = InArgs._OnGetFilterText;
+	}
+
+	TreeView = SNew(STreeView< TSharedPtr<FGraphActionNode> >)
+		.TreeItemsSource(&(this->FilteredRootAction->Children))
+		.OnGenerateRow(this, &STG_ActionMenuTileView::GenerateChildForTree, bIsReadOnly)
+		.OnSelectionChanged(this, &STG_ActionMenuTileView::OnItemSelected)
+		.OnMouseButtonDoubleClick(this, &STG_ActionMenuTileView::OnItemDoubleClicked)
+		.OnContextMenuOpening(InArgs._OnContextMenuOpening)
+		.OnGetChildren(this, &STG_ActionMenuTileView::OnGetChildrenForCategory)
+		.SelectionMode(ESelectionMode::Single)
+		.OnSetExpansionRecursive(this, &STG_ActionMenuTileView::OnSetExpansionRecursive)
+		.HighlightParentNodesForSelection(true);
+
+	this->ChildSlot
+		[
+			SNew(SVerticalBox)
+
+			// FILTER BOX
+			+ SVerticalBox::Slot()
+			.AutoHeight()
+			[
+				SAssignNew(FilterTextBox, SSearchBox)
+				// If there is an external filter delegate, do not display this filter box
+				.Visibility(InArgs._OnGetFilterText.IsBound() ? EVisibility::Collapsed : EVisibility::Visible)
+				.OnTextChanged(this, &STG_ActionMenuTileView::OnFilterTextChanged)
+				.OnTextCommitted(this, &STG_ActionMenuTileView::OnFilterTextCommitted)
+				.DelayChangeNotificationsWhileTyping(false)
+			]
+
+			// ACTION LIST
+			+ SVerticalBox::Slot()
+			.Padding(FMargin(0.0f, 2.0f, 0.0f, 0.0f))
+			.VAlign(VAlign_Fill)
+			.FillHeight(1.0)
+			[
+				SNew(SScrollBox/*, TreeView.ToSharedRef()*/)
+				+SScrollBox::Slot()
+				.VAlign(VAlign_Fill)
+				.FillSize(1.0)
+				[
+					TreeView.ToSharedRef()
+				]
+			]
+		];
+
+	// When the search box has focus, we want first chance handling of any key down events so we can handle the up/down and escape keys the way we want
+	FilterTextBox->SetOnKeyDownHandler(FOnKeyDown::CreateSP(this, &STG_ActionMenuTileView::OnKeyDown));
+
+	if (!InArgs._ShowFilterTextBox)
+	{
+		FilterTextBox->SetVisibility(EVisibility::Collapsed);
+	}
+
+	// Get all actions.
+	RefreshAllActions(false);
+}
+
+void STG_ActionMenuTileView::RefreshAllActions(bool bPreserveExpansion, bool bHandleOnSelectionEvent/*=true*/)
+{
+	TRACE_CPUPROFILER_EVENT_SCOPE(STG_ActionMenuTileView::RefreshAllActions);
+
+	// Save Selection (of only the first selected thing)
+	TArray< TSharedPtr<FGraphActionNode> > SelectedNodes = GetSelectedItems();
+	TSharedPtr<FGraphActionNode> SelectedAction = SelectedNodes.Num() > 0 ? SelectedNodes[0] : nullptr;
+
+	if (OnGetActionList.IsBound())
+	{
+		// Obtain the source action list directly.
+		AllActions = OnGetActionList.Execute();
+	}
+	else
+	{
+		TRACE_CPUPROFILER_EVENT_SCOPE(STG_ActionMenuTileView::CollectAllActions);
+
+		// Collect actions into our local list context.
+		AllActions->Empty();
+		OnCollectAllActions.ExecuteIfBound(*AllActions);
+	}
+
+	GenerateFilteredItems(bPreserveExpansion);
+
+	// Re-apply selection #0 if possible
+	if (SelectedAction.IsValid())
+	{
+		// Clear the selection, we will be re-selecting the previous action
+		TreeView->ClearSelection();
+
+		if (bHandleOnSelectionEvent)
+		{
+			SelectItemByName(*SelectedAction->GetDisplayName().ToString(), ESelectInfo::OnMouseClick, SelectedAction->SectionID, SelectedNodes[0]->IsCategoryNode());
+		}
+		else
+		{
+			// If we do not want to handle the selection, set it directly so it will reselect the item but not handle the event.
+			SelectItemByName(*SelectedAction->GetDisplayName().ToString(), ESelectInfo::Direct, SelectedAction->SectionID, SelectedNodes[0]->IsCategoryNode());
+		}
+	}
+}
+
+void STG_ActionMenuTileView::GetSectionExpansion(TMap<int32, bool>& SectionExpansion) const
+{
+
+}
+
+void STG_ActionMenuTileView::SetSectionExpansion(const TMap<int32, bool>& InSectionExpansion)
+{
+	for (auto& PossibleSection : FilteredRootAction->Children)
+	{
+		if (PossibleSection->IsSectionHeadingNode())
+		{
+			const bool* IsExpanded = InSectionExpansion.Find(PossibleSection->SectionID);
+			if (IsExpanded != nullptr)
+			{
+				TreeView->SetItemExpansion(PossibleSection, *IsExpanded);
+			}
+		}
+	}
+}
+
+TSharedRef<SEditableTextBox> STG_ActionMenuTileView::GetFilterTextBox()
+{
+	return FilterTextBox.ToSharedRef();
+}
+
+void STG_ActionMenuTileView::GetSelectedActions(TArray< TSharedPtr<FEdGraphSchemaAction> >& OutSelectedActions) const
+{
+	OutSelectedActions.Empty();
+
+	TArray< TSharedPtr<FGraphActionNode> > SelectedNodes = GetSelectedItems();
+	if (SelectedNodes.Num() > 0)
+	{
+		for (int32 NodeIndex = 0; NodeIndex < SelectedNodes.Num(); NodeIndex++)
+		{
+			OutSelectedActions.Append(SelectedNodes[NodeIndex]->Actions);
+		}
+	}
+}
+
+void STG_ActionMenuTileView::OnRequestRenameOnActionNode()
+{
+	TArray< TSharedPtr<FGraphActionNode> > SelectedNodes = GetSelectedItems();
+	if (SelectedNodes.Num() > 0)
+	{
+		if (!SelectedNodes[0]->BroadcastRenameRequest())
+		{
+			TreeView->RequestScrollIntoView(SelectedNodes[0]);
+		}
+	}
+}
+
+bool STG_ActionMenuTileView::CanRequestRenameOnActionNode() const
+{
+	return false;
+}
+
+FString STG_ActionMenuTileView::GetSelectedCategoryName() const
+{
+	TArray< TSharedPtr<FGraphActionNode> > SelectedNodes = GetSelectedItems();
+	return (SelectedNodes.Num() > 0) ? SelectedNodes[0]->GetDisplayName().ToString() : FString();
+}
+
+void STG_ActionMenuTileView::GetSelectedCategorySubActions(TArray<TSharedPtr<FEdGraphSchemaAction>>& OutActions) const
+{
+	TArray< TSharedPtr<FGraphActionNode> > SelectedNodes = GetSelectedItems();
+	for (int32 SelectionIndex = 0; SelectionIndex < SelectedNodes.Num(); SelectionIndex++)
+	{
+		if (SelectedNodes[SelectionIndex].IsValid())
+		{
+			GetCategorySubActions(SelectedNodes[SelectionIndex], OutActions);
+		}
+	}
+}
+
+void STG_ActionMenuTileView::GetCategorySubActions(TWeakPtr<FGraphActionNode> InAction, TArray<TSharedPtr<FEdGraphSchemaAction>>& OutActions) const
+{
+	if (InAction.IsValid())
+	{
+		TSharedPtr<FGraphActionNode> CategoryNode = InAction.Pin();
+		TArray<TSharedPtr<FGraphActionNode>> Children;
+		CategoryNode->GetLeafNodes(Children);
+
+		for (int32 i = 0; i < Children.Num(); ++i)
+		{
+			TSharedPtr<FGraphActionNode> CurrentChild = Children[i];
+
+			if (CurrentChild.IsValid() && CurrentChild->IsActionNode())
+			{
+				for (int32 ActionIndex = 0; ActionIndex != CurrentChild->Actions.Num(); ActionIndex++)
+				{
+					OutActions.Add(CurrentChild->Actions[ActionIndex]);
+				}
+			}
+		}
+	}
+}
+
+bool STG_ActionMenuTileView::SelectItemByName(const FName& ItemName, ESelectInfo::Type SelectInfo, int32 SectionId/* = INDEX_NONE */, bool bIsCategory/* = false*/)
+{
+	if (ItemName != NAME_None)
+	{
+		TSharedPtr<FGraphActionNode> SelectionNode;
+
+		TArray<TSharedPtr<FGraphActionNode>> GraphNodes;
+		FilteredRootAction->GetAllNodes(GraphNodes);
+		for (int32 i = 0; i < GraphNodes.Num() && !SelectionNode.IsValid(); ++i)
+		{
+			TSharedPtr<FGraphActionNode> CurrentGraphNode = GraphNodes[i];
+			FEdGraphSchemaAction* GraphAction = CurrentGraphNode->GetPrimaryAction().Get();
+
+			// If the user is attempting to select a category, make sure it's a category
+			if (CurrentGraphNode->IsCategoryNode() == bIsCategory)
+			{
+				if (SectionId == INDEX_NONE || CurrentGraphNode->SectionID == SectionId)
+				{
+					if (GraphAction)
+					{
+						if ((OnActionMatchesName.IsBound() && OnActionMatchesName.Execute(GraphAction, ItemName)) /*|| GraphActionMenuHelpers::ActionMatchesName(GraphAction, ItemName)*/)
+						{
+							SelectionNode = GraphNodes[i];
+
+							break;
+						}
+					}
+
+					if (CurrentGraphNode->GetDisplayName().ToString() == FName::NameToDisplayString(ItemName.ToString(), false))
+					{
+						SelectionNode = CurrentGraphNode;
+
+						break;
+					}
+				}
+			}
+
+			// One of the children may match
+			for (int32 ChildIdx = 0; ChildIdx < CurrentGraphNode->Children.Num() && !SelectionNode.IsValid(); ++ChildIdx)
+			{
+				TSharedPtr<FGraphActionNode> CurrentChildNode = CurrentGraphNode->Children[ChildIdx];
+
+				for (int32 ActionIndex = 0; ActionIndex < CurrentChildNode->Actions.Num(); ActionIndex++)
+				{
+					FEdGraphSchemaAction* ChildGraphAction = CurrentChildNode->Actions[ActionIndex].Get();
+
+					// If the user is attempting to select a category, make sure it's a category
+					if (CurrentChildNode->IsCategoryNode() == bIsCategory)
+					{
+						if (SectionId == INDEX_NONE || CurrentChildNode->SectionID == SectionId)
+						{
+							if (ChildGraphAction)
+							{
+								if ((OnActionMatchesName.IsBound() && OnActionMatchesName.Execute(ChildGraphAction, ItemName)) /*|| GraphActionMenuHelpers::ActionMatchesName(ChildGraphAction, ItemName)*/)
+								{
+									SelectionNode = GraphNodes[i]->Children[ChildIdx];
+
+									break;
+								}
+							}
+							else if (CurrentChildNode->GetDisplayName().ToString() == FName::NameToDisplayString(ItemName.ToString(), false))
+							{
+								SelectionNode = CurrentChildNode;
+
+								break;
+							}
+						}
+					}
+				}
+			}
+		}
+
+		if (SelectionNode.IsValid())
+		{
+			// Expand the parent nodes
+			for (TSharedPtr<FGraphActionNode> ParentAction = SelectionNode->GetParentNode().Pin(); ParentAction.IsValid(); ParentAction = ParentAction->GetParentNode().Pin())
+			{
+				TreeView->SetItemExpansion(ParentAction, true);
+			}
+
+			// Select the node
+			TreeView->SetSelection(SelectionNode, SelectInfo);
+			TreeView->RequestScrollIntoView(SelectionNode);
+			return true;
+		}
+	}
+	else
+	{
+		TreeView->ClearSelection();
+		return true;
+	}
+	return false;
+}
+
+void STG_ActionMenuTileView::ExpandCategory(const FText& CategoryName)
+{
+	if (!CategoryName.IsEmpty())
+	{
+		TArray<TSharedPtr<FGraphActionNode>> GraphNodes;
+		FilteredRootAction->GetAllNodes(GraphNodes);
+		for (int32 i = 0; i < GraphNodes.Num(); ++i)
+		{
+			if (GraphNodes[i]->GetDisplayName().EqualTo(CategoryName))
+			{
+				GraphNodes[i]->ExpandAllChildren(TreeView);
+			}
+		}
+	}
+}
+
+static bool CompareGraphActionNode(TSharedPtr<FGraphActionNode> A, TSharedPtr<FGraphActionNode> B)
+{
+	check(A.IsValid());
+	check(B.IsValid());
+
+	// First check grouping is the same
+	if (A->GetDisplayName().ToString() != B->GetDisplayName().ToString())
+	{
+		return false;
+	}
+
+	if (A->SectionID != B->SectionID)
+	{
+		return false;
+	}
+
+	if (A->HasValidAction() && B->HasValidAction())
+	{
+		return A->GetPrimaryAction()->GetMenuDescription().CompareTo(B->GetPrimaryAction()->GetMenuDescription()) == 0;
+	}
+	else if (!A->HasValidAction() && !B->HasValidAction())
+	{
+		return true;
+	}
+	else
+	{
+		return false;
+	}
+}
+
+template<typename ItemType, typename ComparisonType>
+void RestoreExpansionState(TSharedPtr< STreeView<ItemType> > InTree, const TArray<ItemType>& ItemSource, const TSet<ItemType>& OldExpansionState, ComparisonType ComparisonFunction)
+{
+	check(InTree.IsValid());
+
+	// Iterate over new tree items
+	for (int32 ItemIdx = 0; ItemIdx < ItemSource.Num(); ItemIdx++)
+	{
+		ItemType NewItem = ItemSource[ItemIdx];
+
+		// Look through old expansion state
+		for (typename TSet<ItemType>::TConstIterator OldExpansionIter(OldExpansionState); OldExpansionIter; ++OldExpansionIter)
+		{
+			const ItemType OldItem = *OldExpansionIter;
+			// See if this matches this new item
+			if (ComparisonFunction(OldItem, NewItem))
+			{
+				// It does, so expand it
+				InTree->SetItemExpansion(NewItem, true);
+			}
+		}
+	}
+}
+
+void STG_ActionMenuTileView::UpdateForNewActions(int32 IdxStart)
+{
+	check(bAlphaSortItems && bSortItemsRecursively);
+
+	FScoreResults Results = ScoreAndAddActions(IdxStart);
+	UpdateActiveSelection(Results);
+
+	if (ShouldExpandNodes())
+	{
+		// Expand all
+		FilteredRootAction->ExpandAllChildren(TreeView);
+	}
+
+	TreeView->RequestTreeRefresh();
+}
+
+void STG_ActionMenuTileView::GenerateFilteredItems(bool bPreserveExpansion)
+{
+	TRACE_CPUPROFILER_EVENT_SCOPE(STG_ActionMenuTileView::GenerateFilteredItems);
+
+	SelectedSuggestionScore = TNumericLimits<float>::Lowest();
+	SelectedSuggestionSourceIndex = INDEX_NONE;
+
+	// First, save off current expansion state
+	TSet< TSharedPtr<FGraphActionNode> > OldExpansionState;
+	if (bPreserveExpansion)
+	{
+		TreeView->GetExpandedItems(OldExpansionState);
+	}
+
+	// Clear the filtered root action
+	FilteredRootAction->ClearChildren();
+
+	// Collect the list of always visible sections if any, and force the creation of those sections.
+	if (OnCollectStaticSections.IsBound())
+	{
+		TArray<int32> StaticSectionIDs;
+		OnCollectStaticSections.Execute(StaticSectionIDs);
+
+		for (int32 i = 0; i < StaticSectionIDs.Num(); i++)
+		{
+			FilteredRootAction->AddSection(0, StaticSectionIDs[i]);
+		}
+	}
+
+	FScoreResults Results = ScoreAndAddActions();
+
+	FilteredRootAction->SortChildren(bAlphaSortItems, bSortItemsRecursively);
+
+	TreeView->RequestTreeRefresh();
+
+	UpdateActiveSelection(Results);
+
+	if (ShouldExpandNodes())
+	{
+		// Expand all
+		FilteredRootAction->ExpandAllChildren(TreeView);
+	}
+	else
+	{
+		// Get _all_ new nodes (flattened tree basically)
+		TArray< TSharedPtr<FGraphActionNode> > AllNodes;
+		FilteredRootAction->GetAllNodes(AllNodes);
+
+		// Expand to match the old state
+		RestoreExpansionState< TSharedPtr<FGraphActionNode> >(TreeView, AllNodes, OldExpansionState, CompareGraphActionNode);
+	}
+}
+
+// Returns true if the tree should be autoexpanded
+bool STG_ActionMenuTileView::ShouldExpandNodes() const
+{
+	// Expand all the categories that have filter results, or when there are only a few to show
+	const bool bFilterActive = !GetFilterText().IsEmpty();
+	const bool bOnlyAFewTotal = AllActions->GetNumActions() < 10;
+
+	return bFilterActive || bOnlyAFewTotal || bAutoExpandActionMenu;
+}
+
+bool STG_ActionMenuTileView::CanRenameNode(TWeakPtr<FGraphActionNode> InNode) const
+{
+	return false;
+}
+
+void STG_ActionMenuTileView::OnFilterTextChanged(const FText& InFilterText)
+{
+	// Reset the selection if the string is empty
+	if (InFilterText.IsEmpty() == true)
+	{
+		SelectedSuggestion = INDEX_NONE;
+	}
+	GenerateFilteredItems(false);
+}
+
+void STG_ActionMenuTileView::OnFilterTextCommitted(const FText& InText, ETextCommit::Type CommitInfo)
+{
+	if (CommitInfo == ETextCommit::OnEnter)
+	{
+		TryToSpawnActiveSuggestion();
+	}
+}
+
+bool STG_ActionMenuTileView::TryToSpawnActiveSuggestion()
+{
+	TArray< TSharedPtr<FGraphActionNode> > SelectionList = GetSelectedItems();
+
+	if (SelectionList.Num() == 1)
+	{
+		// This isnt really a keypress - its Direct, but its always called from a keypress function. (Maybe pass the selectinfo in ?)
+		OnItemSelected(SelectionList[0], ESelectInfo::OnKeyPress);
+		return true;
+	}
+	else if (FilteredActionNodes.Num() == 1)
+	{
+		OnItemSelected(FilteredActionNodes[0], ESelectInfo::OnKeyPress);
+		return true;
+	}
+
+	return false;
+}
+
+void STG_ActionMenuTileView::OnGetChildrenForCategory(TSharedPtr<FGraphActionNode> InItem, TArray< TSharedPtr<FGraphActionNode> >& OutChildren)
+{
+	//Intentionally leaving it unimplemented as we do not want the children to spawn for tree
+}
+
+TSharedRef<ITableRow> STG_ActionMenuTileView::MakeWidget(TSharedPtr<FGraphActionNode> InItem, const TSharedRef<STableViewBase>& OwnerTable, bool bIsReadOnly, bool bMakeTiles)
+{
+	TSharedPtr<IToolTip> SectionToolTip;
+
+	if (InItem->IsSectionHeadingNode())
+	{
+		if (OnGetSectionToolTip.IsBound())
+		{
+			SectionToolTip = OnGetSectionToolTip.Execute(InItem->SectionID);
+		}
+	}
+
+	// In the case of FGraphActionNodes that have multiple actions, all of the actions will
+	// have the same text as they will have been created at the same point - only the actual
+	// action itself will differ, which is why parts of this function only refer to InItem->Actions[0]
+	// rather than iterating over the array
+
+	// Create the widget but do not add any content, the widget is needed to pass the IsSelectedExclusively function down to the potential SInlineEditableTextBlock widget
+	TSharedPtr< STableRow< TSharedPtr<FGraphActionNode> > > TableRow;
+
+	if (InItem->IsSectionHeadingNode())
+	{
+		TableRow = SNew(STG_CategoryHeaderTableRow< TSharedPtr<FGraphActionNode> >, OwnerTable)
+			.ToolTip(SectionToolTip);
+	}
+	else
+	{
+		if (bMakeTiles)
+		{
+			TableRow = SNew(STableRow< TSharedPtr<FGraphActionNode> >, OwnerTable)
+			.Style(FAppStyle::Get(), "ContentBrowser.AssetListView.TileTableRow")
+			.OnDragDetected(this, &STG_ActionMenuTileView::OnItemDragDetected)
+			.ShowSelection(!InItem->IsSeparator())
+			.bAllowPreselectedItemActivation(bAllowPreselectedItemActivation)
+			.Cursor( EMouseCursor::GrabHand);
+		}
+		else
+		{
+			TableRow = SNew(STableRow< TSharedPtr<FGraphActionNode> >, OwnerTable)
+			.ShowSelection(!InItem->IsSeparator())
+			.bAllowPreselectedItemActivation(bAllowPreselectedItemActivation);
+		}
+	}
+
+	TSharedPtr<SHorizontalBox> RowContainer;
+	TableRow->SetRowContent
+	(
+		SAssignNew(RowContainer, SHorizontalBox)
+	);
+
+	TSharedPtr<SWidget> RowContent;
+	FMargin RowPadding = FMargin(0,4,0,0);
+
+	if (InItem->IsActionNode())
+	{
+		//When called for tree we are disabling the Action Nodes, As we are using tree to only create the category nodes
+		//Action nodes are spawned inside of the Tile view of Category node
+		if (!bMakeTiles)
+		{
+			TableRow->SetVisibility(EVisibility::Collapsed);
+		}
+		
+		check(InItem->HasValidAction());
+
+		FCreateWidgetForActionData CreateData(&InItem->OnRenameRequest());
+		CreateData.Action = InItem->GetPrimaryAction();
+		CreateData.HighlightText = TAttribute<FText>(this, &STG_ActionMenuTileView::GetFilterText);
+		CreateData.MouseButtonDownDelegate = FCreateWidgetMouseButtonDown::CreateSP(this, &STG_ActionMenuTileView::OnMouseButtonDownEvent);
+
+		if (OnCreateWidgetForAction.IsBound() && bMakeTiles)
+		{
+			CreateData.IsRowSelectedDelegate = FIsSelected::CreateSP(TableRow.Get(), &STableRow< TSharedPtr<FGraphActionNode> >::IsSelected);
+			CreateData.bIsReadOnly = bIsReadOnly;
+			CreateData.bHandleMouseButtonDown = false;		//Default to NOT using the delegate. OnCreateWidgetForAction can set to true if we need it
+			RowContent = OnCreateWidgetForAction.Execute(&CreateData);
+			RowPadding = FMargin(0, 0);
+		}
+		else
+		{
+			RowPadding = FMargin(0, 0);
+			RowContent = SNullWidget::NullWidget;
+		}
+	}
+	else if (InItem->IsCategoryNode())
+	{
+		TWeakPtr< FGraphActionNode > WeakItem = InItem;
+		// Hook up the delegate for verifying the category action is read only or not
+		STG_GraphActionCategoryWidget::FArguments ReadOnlyArgument;
+		if (bIsReadOnly)
+		{
+			ReadOnlyArgument.IsReadOnly(bIsReadOnly);
+		}
+		else
+		{
+			ReadOnlyArgument.IsReadOnly(this, &STG_ActionMenuTileView::CanRenameNode, WeakItem);
+		}
+
+		TSharedRef<STG_GraphActionCategoryWidget> CategoryWidget =
+			SNew(STG_GraphActionCategoryWidget, InItem)
+			.HighlightText(this, &STG_ActionMenuTileView::GetFilterText)
+			.IsSelected(TableRow.Get(), &STableRow< TSharedPtr<FGraphActionNode> >::IsSelectedExclusively)
+			.IsReadOnly(ReadOnlyArgument._IsReadOnly);
+
+		if (!bIsReadOnly)
+		{
+			InItem->OnRenameRequest().BindSP(CategoryWidget->InlineWidget.Pin().Get(), &SInlineEditableTextBlock::EnterEditingMode);
+		}
+
+		FString ChildCount = FString::Format(TEXT("{0} Items"), { InItem->Children.Num() });
+
+		RowContent = CategoryWidget;
+		auto TileView = SNew(STileView<TSharedPtr<FGraphActionNode>>)
+			.ListItemsSource(&InItem->Children)
+			.OnGenerateTile(this, &STG_ActionMenuTileView::GenerateTileForChildItem)
+			.OnSelectionChanged(this, &STG_ActionMenuTileView::OnTileItemSelected)
+			.ItemWidth(this,&STG_ActionMenuTileView::GetTileViewItemWidth)
+			.SelectionMode(ESelectionMode::Single)
+			.ItemHeight(TileHeight)
+			.ScrollbarVisibility(EVisibility::Collapsed);
+
+		auto ChildWidget = SNew(SBorder)
+			.BorderBackgroundColor(FStyleColors::Recessed)
+			.BorderImage(FCoreStyle::Get().GetBrush("WhiteBrush"))
+			.Padding(FMargin(6.0f, 2.0f, 5.0f, 2.0f))
+			.Visibility_Lambda([=]() {
+				if (TableRow->IsItemExpanded())
+				{
+					return EVisibility::Visible;
+				}
+				return EVisibility::Collapsed;
+			})
+			[
+				SNew(SVerticalBox)
+
+				+ SVerticalBox::Slot()
+				.AutoHeight()
+				[
+					TileView
+				]
+			];
+
+		TreeTileViewMap.Emplace(InItem, TileView);
+
+		TableRow->SetRowContent
+		(
+			SNew(SVerticalBox)
+
+			+ SVerticalBox::Slot()
+			.VAlign(VAlign_Top)
+			.AutoHeight()
+			[
+				SNew(SBorder)
+				.BorderImage(FCoreStyle::Get().GetBrush("WhiteBrush"))
+				.BorderBackgroundColor(EStyleColor::Recessed)
+				[
+					RowContainer.ToSharedRef()
+				]
+			]
+
+			+ SVerticalBox::Slot()
+			.VAlign(VAlign_Top)
+			.AutoHeight()
+			[
+				ChildWidget
+			]
+		);
+		
+	}
+
+	//Only add expander for the category node
+	if (!bMakeTiles && InItem->IsCategoryNode())
+	{
+		TSharedPtr<STG_CategoryHeaderExpanderArrow> ExpanderWidget = 
+			SNew(STG_CategoryHeaderExpanderArrow, TableRow , TreeView, InItem)
+			.AlwaysVisible(true);
+
+		RowContainer->AddSlot()
+		.AutoWidth()
+		.VAlign(VAlign_Fill)
+		.HAlign(HAlign_Right)
+		[
+			ExpanderWidget.ToSharedRef()
+		];
+	}
+
+	RowContainer->AddSlot()
+	.FillWidth(1.0)
+	.Padding(RowPadding)
+	[
+		RowContent.ToSharedRef()
+	];
+
+	return TableRow.ToSharedRef();
+}
+
+void STG_ActionMenuTileView::Tick(const FGeometry& AllottedGeometry, const double InCurrentTime, const float InDeltaTime)
+{
+	CalculateFillScale(AllottedGeometry);
+}
+
+void STG_ActionMenuTileView::CalculateFillScale(const FGeometry& AllottedGeometry)
+{
+	float ItemWidth = TileWidth;
+
+	// Scrollbars are 16, but we add 1 to deal with half pixels.
+	const float ScrollbarWidth = 16 + 1 + 8;//+8 for tile view left and right padding
+	float TotalWidth = AllottedGeometry.GetLocalSize().X - (ScrollbarWidth);
+	float Coverage = TotalWidth / ItemWidth;
+	int32 Items = (int)(TotalWidth / ItemWidth);
+
+	// If there isn't enough room to support even a single item, don't apply a fill scale.
+	if (Items > 0)
+	{
+		float GapSpace = ItemWidth * (Coverage - (float)Items);
+		float ExpandAmount = GapSpace / (float)Items;
+		FillScale = (ItemWidth + ExpandAmount) / ItemWidth;
+		FillScale = FMath::Max(1.0f, FillScale);
+	}
+	else
+	{
+		FillScale = 1.0f;
+	}
+}
+
+float STG_ActionMenuTileView::GetTileViewItemWidth() const
+{
+	return (float)(TileWidth * FillScale);
+}
+
+TSharedRef<ITableRow> STG_ActionMenuTileView::GenerateChildForTree(TSharedPtr<FGraphActionNode> InItem, const TSharedRef<STableViewBase>& OwnerTable, bool bIsReadOnly)
+{
+	return MakeWidget(InItem, OwnerTable, bIsReadOnly, false);
+}
+
+TSharedRef<ITableRow> STG_ActionMenuTileView::GenerateTileForChildItem(TSharedPtr<FGraphActionNode> InItem, const TSharedRef<STableViewBase>& OwnerTable)
+{
+	return MakeWidget(InItem, OwnerTable, false, true);
+}
+
+FText STG_ActionMenuTileView::GetFilterText() const
+{
+	// If there is an external source for the filter, use that text instead
+	if (OnGetFilterText.IsBound())
+	{
+		return OnGetFilterText.Execute();
+	}
+
+	return FilterTextBox->GetText();
+}
+
+void STG_ActionMenuTileView::OnItemSelected(TSharedPtr< FGraphActionNode > InSelectedItem, ESelectInfo::Type SelectInfo)
+{
+	if (!bIgnoreUIUpdate)
+	{
+		HandleSelection(InSelectedItem, SelectInfo);
+	}
+}
+
+void STG_ActionMenuTileView::OnTileItemSelected(TSharedPtr< FGraphActionNode > InSelectedItem, ESelectInfo::Type SelectInfo)
+{
+	if (!bIgnoreUIUpdate)
+	{
+		if (InSelectedItem.IsValid())
+		{
+			for (auto TileViewPair : TreeTileViewMap)
+			{
+				if (TileViewPair.Key != InSelectedItem->GetParentNode())
+					TileViewPair.Value->ClearSelection();
+			}
+			TreeView->ClearSelection();
+		}
+		HandleSelection(InSelectedItem, SelectInfo);
+	}
+}
+
+TArray< TSharedPtr<FGraphActionNode> > STG_ActionMenuTileView::GetSelectedItems() const
+{
+	TArray< TSharedPtr<FGraphActionNode> > SelectedItems;
+	for (auto TileViewPair : TreeTileViewMap)
+	{
+		for (auto SelectedItem : TileViewPair.Value->GetSelectedItems())
+		{
+			SelectedItems.Add(SelectedItem);
+		}
+	}
+	return SelectedItems;
+}
+
+void STG_ActionMenuTileView::OnItemDoubleClicked(TSharedPtr< FGraphActionNode > InClickedItem)
+{
+	if (InClickedItem.IsValid() && !bIgnoreUIUpdate)
+	{
+		if (InClickedItem->IsActionNode())
+		{
+			OnActionDoubleClicked.ExecuteIfBound(InClickedItem->Actions);
+		}
+		else if (InClickedItem->Children.Num())
+		{
+			TreeView->SetItemExpansion(InClickedItem, !TreeView->IsItemExpanded(InClickedItem));
+		}
+	}
+}
+
+FReply STG_ActionMenuTileView::OnItemDragDetected(const FGeometry& MyGeometry, const FPointerEvent& MouseEvent)
+{
+	// Start a function-call drag event for any entry that can be called by kismet
+	if (MouseEvent.IsMouseButtonDown(EKeys::LeftMouseButton))
+	{
+		TArray< TSharedPtr<FGraphActionNode> > SelectedNodes = GetSelectedItems();
+		if (SelectedNodes.Num() > 0)
+		{
+			TSharedPtr<FGraphActionNode> Node = SelectedNodes[0];
+			// Dragging a ctaegory
+			if (Node.IsValid() && Node->IsCategoryNode())
+			{
+				if (OnCategoryDragged.IsBound())
+				{
+					return OnCategoryDragged.Execute(Node->GetCategoryPath(), MouseEvent);
+				}
+			}
+			// Dragging an action
+			else
+			{
+				if (OnActionDragged.IsBound())
+				{
+					TArray< TSharedPtr<FEdGraphSchemaAction> > Actions;
+					GetSelectedActions(Actions);
+					return OnActionDragged.Execute(Actions, MouseEvent);
+				}
+			}
+		}
+	}
+
+	return FReply::Unhandled();
+}
+
+bool STG_ActionMenuTileView::OnMouseButtonDownEvent(TWeakPtr<FEdGraphSchemaAction> InAction)
+{
+	bool bResult = false;
+	if ((!bIgnoreUIUpdate) && InAction.IsValid())
+	{
+		TArray< TSharedPtr<FGraphActionNode> > SelectionList = GetSelectedItems();
+		TSharedPtr<FGraphActionNode> SelectedNode;
+		if (SelectionList.Num() == 1)
+		{
+			SelectedNode = SelectionList[0];
+		}
+		else if (FilteredActionNodes.Num() == 1)
+		{
+			SelectedNode = FilteredActionNodes[0];
+		}
+		if (SelectedNode.IsValid() && SelectedNode->HasValidAction())
+		{
+			if (SelectedNode->GetPrimaryAction().Get() == InAction.Pin().Get())
+			{
+				bResult = HandleSelection(SelectedNode, ESelectInfo::OnMouseClick);
+			}
+		}
+	}
+	return bResult;
+}
+
+FReply STG_ActionMenuTileView::OnKeyDown(const FGeometry& MyGeometry, const FKeyEvent& KeyEvent)
+{
+	int32 SelectionDelta = 0;
+
+	// Escape dismisses the menu without placing a node
+	if (KeyEvent.GetKey() == EKeys::Escape)
+	{
+		FSlateApplication::Get().DismissAllMenus();
+		return FReply::Handled();
+	}
+	else if ((KeyEvent.GetKey() == EKeys::Enter) && !bIgnoreUIUpdate)
+	{
+		return TryToSpawnActiveSuggestion() ? FReply::Handled() : FReply::Unhandled();
+	}
+	else if (!FilterTextBox->GetText().IsEmpty())
+	{
+		// Needs to be done here in order not to eat up the text navigation key events when list isn't populated
+		if (FilteredActionNodes.Num() == 0)
+		{
+			return FReply::Unhandled();
+		}
+
+		if (KeyEvent.GetKey() == EKeys::Up)
+		{
+			SelectedSuggestion = FMath::Max(0, SelectedSuggestion - 1);
+		}
+		else if (KeyEvent.GetKey() == EKeys::Down)
+		{
+			SelectedSuggestion = FMath::Min(FilteredActionNodes.Num() - 1, SelectedSuggestion + 1);
+		}
+		else if (KeyEvent.GetKey() == EKeys::PageUp)
+		{
+			const int32 NumItemsInAPage = 15; // arbitrary jump because we can't get at the visible item count from here
+			SelectedSuggestion = FMath::Max(0, SelectedSuggestion - NumItemsInAPage);
+		}
+		else if (KeyEvent.GetKey() == EKeys::PageDown)
+		{
+			const int32 NumItemsInAPage = 15; // arbitrary jump because we can't get at the visible item count from here
+			SelectedSuggestion = FMath::Min(FilteredActionNodes.Num() - 1, SelectedSuggestion + NumItemsInAPage);
+		}
+		else if (KeyEvent.GetKey() == EKeys::Home && KeyEvent.IsControlDown())
+		{
+			SelectedSuggestion = 0;
+		}
+		else if (KeyEvent.GetKey() == EKeys::End && KeyEvent.IsControlDown())
+		{
+			SelectedSuggestion = FilteredActionNodes.Num() - 1;
+		}
+		else
+		{
+			return FReply::Unhandled();
+		}
+
+		MarkActiveSuggestion();
+		return FReply::Handled();
+	}
+	else
+	{
+		// When all else fails, it means we haven't filtered the list and we want to handle it as if we were just scrolling through a normal tree view
+		return TreeView->OnKeyDown(FindChildGeometry(MyGeometry, TreeView.ToSharedRef()), KeyEvent);
+	}
+
+	return FReply::Unhandled();
+}
+
+void STG_ActionMenuTileView::MarkActiveSuggestion()
+{
+	TGuardValue<bool> PreventSelectionFromTriggeringCommit(bIgnoreUIUpdate, true);
+
+	if (SelectedSuggestion >= 0)
+	{
+		TSharedPtr<FGraphActionNode>& ActionToSelect = FilteredActionNodes[SelectedSuggestion];
+
+		TreeView->SetSelection(ActionToSelect);
+		TreeView->RequestScrollIntoView(ActionToSelect);
+	}
+	else
+	{
+		TreeView->ClearSelection();
+	}
+}
+
+void STG_ActionMenuTileView::AddReferencedObjects(FReferenceCollector& Collector)
+{
+	for (int32 CurTypeIndex = 0; CurTypeIndex < AllActions->GetNumActions(); ++CurTypeIndex)
+	{
+		FGraphActionListBuilderBase::ActionGroup& Action = AllActions->GetAction(CurTypeIndex);
+
+		for (int32 ActionIndex = 0; ActionIndex < Action.Actions.Num(); ActionIndex++)
+		{
+			Action.Actions[ActionIndex]->AddReferencedObjects(Collector);
+		}
+	}
+}
+
+FString STG_ActionMenuTileView::GetReferencerName() const
+{
+	return TEXT("STG_ActionMenuTileView");
+}
+
+bool STG_ActionMenuTileView::HandleSelection(TSharedPtr< FGraphActionNode >& InSelectedItem, ESelectInfo::Type InSelectionType)
+{
+	bool bResult = false;
+	if (OnActionSelected.IsBound())
+	{
+		if (InSelectedItem.IsValid() && InSelectedItem->IsActionNode())
+		{
+			OnActionSelected.Execute(InSelectedItem->Actions, InSelectionType);
+			bResult = true;
+		}
+		else
+		{
+			OnActionSelected.Execute(TArray< TSharedPtr<FEdGraphSchemaAction> >(), InSelectionType);
+			bResult = true;
+		}
+	}
+	return bResult;
+}
+
+void STG_ActionMenuTileView::OnSetExpansionRecursive(TSharedPtr<FGraphActionNode> InTreeNode, bool bInIsItemExpanded)
+{
+	if (InTreeNode.IsValid() && InTreeNode->Children.Num())
+	{
+		TreeView->SetItemExpansion(InTreeNode, bInIsItemExpanded);
+
+		for (TSharedPtr<FGraphActionNode> Child : InTreeNode->Children)
+		{
+			OnSetExpansionRecursive(Child, bInIsItemExpanded);
+		}
+	}
+}
+
+STG_ActionMenuTileView::FScoreResults STG_ActionMenuTileView::ScoreAndAddActions(int32 StartingIndex)
+{
+	// Trim and sanitized the filter text (so that it more likely matches the action descriptions)
+	FString TrimmedFilterString = FText::TrimPrecedingAndTrailing(GetFilterText()).ToString();
+
+	// Remember the last filter string to that external clients can access it
+	LastUsedFilterText = TrimmedFilterString;
+
+	// Tokenize the search box text into a set of terms; all of them must be present to pass the filter
+	TArray<FString> FilterTerms;
+	TrimmedFilterString.ParseIntoArray(FilterTerms, TEXT(" "), true);
+	for (FString& String : FilterTerms)
+	{
+		String = String.ToLower();
+	}
+
+	// Generate a list of sanitized versions of the strings
+	TArray<FString> SanitizedFilterTerms;
+	for (int32 iFilters = 0; iFilters < FilterTerms.Num(); iFilters++)
+	{
+		FString EachString = FName::NameToDisplayString(FilterTerms[iFilters], false);
+		EachString = EachString.Replace(TEXT(" "), TEXT(""));
+		SanitizedFilterTerms.Add(EachString);
+	}
+	ensure(SanitizedFilterTerms.Num() == FilterTerms.Num());// Both of these should match !
+
+	const bool bRequiresFiltering = FilterTerms.Num() > 0;
+	float BestMatchCount = SelectedSuggestionScore;
+	int32 BestMatchIndex = SelectedSuggestionSourceIndex;
+
+	// Get the schema of the graph that we are in so that we can correctly get the action weight
+	const UEdGraphSchema* ActionSchema = GraphObj ? GraphObj->GetSchema() : GetDefault<UEdGraphSchema>();
+	check(ActionSchema);
+
+	bool bIsPartialBuild = StartingIndex != INDEX_NONE;
+	const int32 NumActions = AllActions->GetNumActions();
+	for (int32 CurTypeIndex = bIsPartialBuild ? StartingIndex : 0; CurTypeIndex < NumActions; ++CurTypeIndex)
+	{
+		FGraphActionListBuilderBase::ActionGroup& CurrentAction = AllActions->GetAction(CurTypeIndex);
+
+		// If we're filtering, search check to see if we need to show this action
+		bool bShowAction = true;
+		float EachWeight = TNumericLimits<float>::Lowest();
+
+		const FString& SearchText = CurrentAction.GetSearchTextForFirstAction();
+		for (int32 FilterIndex = 0; (FilterIndex < FilterTerms.Num()) && bShowAction; ++FilterIndex)
+		{
+			const bool bMatchesTerm = (SearchText.Contains(FilterTerms[FilterIndex], ESearchCase::CaseSensitive) || (SearchText.Contains(SanitizedFilterTerms[FilterIndex], ESearchCase::CaseSensitive) == true));
+			bShowAction = bMatchesTerm;
+		}
+
+		if (!bShowAction)
+		{
+			continue;
+		}
+
+		if (bRequiresFiltering)
+		{
+			// Get the 'weight' of this in relation to the filter
+			EachWeight = ActionSchema->GetActionFilteredWeight(CurrentAction, FilterTerms, SanitizedFilterTerms, DraggedFromPins);
+			// If this action has a greater relevance than others, cache its index.
+			if (EachWeight > BestMatchCount)
+			{
+				BestMatchCount = EachWeight;
+				BestMatchIndex = CurTypeIndex;
+			}
+		}
+
+		if (bIsPartialBuild)
+		{
+			FilteredRootAction->AddChildAlphabetical(CurrentAction);
+		}
+		else
+		{
+			FilteredRootAction->AddChild(CurrentAction);
+		}
+	}
+
+	return { BestMatchIndex, BestMatchCount };
+}
+
+void STG_ActionMenuTileView::UpdateActiveSelection(STG_ActionMenuTileView::FScoreResults ForResults)
+{
+	int32 BestMatchIndex = ForResults.BestMatchIndex;
+	float BestMatchCount = ForResults.BestMatchScore;
+	// Update the filtered list (needs to be done in a separate pass because the list is sorted as items are inserted)
+	FilteredActionNodes.Reset();
+	FilteredRootAction->GetLeafNodes(FilteredActionNodes);
+
+	// If theres a BestMatchIndex find it in the actions nodes and select it (maybe this should check the current selected suggestion first ?)
+	if (BestMatchIndex != INDEX_NONE)
+	{
+		FGraphActionListBuilderBase::ActionGroup& FilterSelectAction = AllActions->GetAction(BestMatchIndex);
+		if (FilterSelectAction.Actions[0].IsValid() == true)
+		{
+			for (int32 iNode = 0; iNode < FilteredActionNodes.Num(); iNode++)
+			{
+				if (FilteredActionNodes[iNode].Get()->GetPrimaryAction() == FilterSelectAction.Actions[0])
+				{
+					SelectedSuggestion = iNode;
+					SelectedSuggestionScore = BestMatchCount;
+					SelectedSuggestionSourceIndex = BestMatchIndex;
+				}
+			}
+		}
+	}
+
+	// Make sure the selected suggestion stays within the filtered list
+	if ((SelectedSuggestion >= 0) && (FilteredActionNodes.Num() > 0))
+	{
+		//@TODO: Should try to actually maintain the highlight on the same item if it survived the filtering
+		SelectedSuggestion = FMath::Clamp<int32>(SelectedSuggestion, 0, FilteredActionNodes.Num() - 1);
+		MarkActiveSuggestion();
+	}
+	else
+	{
+		SelectedSuggestionScore = TNumericLimits<float>::Lowest();
+		SelectedSuggestionSourceIndex = INDEX_NONE;
+		SelectedSuggestion = INDEX_NONE;
+	}
+}
+#undef LOCTEXT_NAMESPACE
diff --git a/Engine/Plugins/Experimental/TextureGraph/Source/TextureGraphEditor/Private/Widgets/ActionMenuWidgets/STG_ActionMenuTileView.h b/Engine/Plugins/Experimental/TextureGraph/Source/TextureGraphEditor/Private/Widgets/ActionMenuWidgets/STG_ActionMenuTileView.h
new file mode 100644
--- /dev/null
+++ b/Engine/Plugins/Experimental/TextureGraph/Source/TextureGraphEditor/Private/Widgets/ActionMenuWidgets/STG_ActionMenuTileView.h
@@ -0,0 +1,365 @@
+// Copyright Epic Games, Inc. All Rights Reserved.
+
+
+#pragma once
+#include "SGraphActionMenu.h"
+#include "Containers/Array.h"
+#include "Containers/BitArray.h"
+#include "Containers/Map.h"
+#include "Containers/Set.h"
+#include "Containers/SparseArray.h"
+#include "Containers/UnrealString.h"
+#include "CoreMinimal.h"
+#include "CoreTypes.h"
+#include "Delegates/Delegate.h"
+#include "Framework/SlateDelegates.h"
+#include "HAL/PlatformCrt.h"
+#include "Input/Reply.h"
+#include "Internationalization/Text.h"
+#include "Misc/Attribute.h"
+#include "Misc/Optional.h"
+#include "SlateFwd.h"
+#include "Templates/SharedPointer.h"
+#include "Templates/TypeHash.h"
+#include "Templates/UnrealTemplate.h"
+#include "Types/SlateConstants.h"
+#include "Types/SlateEnums.h"
+#include "UObject/GCObject.h"
+#include "UObject/NameTypes.h"
+#include "Widgets/DeclarativeSyntaxSupport.h"
+#include "Widgets/SCompoundWidget.h"
+#include "Widgets/SWidget.h"
+#include "Widgets/Views/SExpanderArrow.h"
+#include "Widgets/Views/STableRow.h"
+#include "Widgets/Views/STableViewBase.h"
+#include "Widgets/Views/STreeView.h"
+#include "Widgets/Views/STileView.h"
+
+class FTG_Editor;
+class FReferenceCollector;
+class ITableRow;
+class IToolTip;
+class SEditableTextBox;
+class SPanel;
+class SSearchBox;
+class SWidget;
+class UEdGraph;
+class UEdGraphPin;
+struct FCreateWidgetForActionData;
+struct FGeometry;
+struct FGraphActionNode;
+struct FKeyEvent;
+struct FPointerEvent;
+struct FEdGraphSchemaAction;
+struct FGraphActionListBuilderBase;
+
+/* This class implements the hibrid Graph and Tile view of GraphActionMenu*/
+class STG_ActionMenuTileView : public SCompoundWidget, public FGCObject
+{
+public:
+	/** Delegate that can be used to create a widget for a particular action */
+	DECLARE_DELEGATE_RetVal_OneParam(TSharedRef<SWidget>, FOnCreateWidgetForAction, FCreateWidgetForActionData* const);
+	/** Delegate that can be used to create a custom "expander" widget for a particular row */
+	DECLARE_DELEGATE_RetVal_OneParam(TSharedRef<SExpanderArrow>, FOnCreateCustomRowExpander, FCustomExpanderData const&);
+	/** Delegate executed when an action is selected */
+	DECLARE_DELEGATE_TwoParams(FOnActionSelected, const TArray< TSharedPtr<FEdGraphSchemaAction> >&, ESelectInfo::Type);
+	/** Delegate executed when an action is double clicked */
+	DECLARE_DELEGATE_OneParam(FOnActionDoubleClicked, const TArray< TSharedPtr<FEdGraphSchemaAction> >&);
+	/** Delegate executed when an action is dragged */
+	DECLARE_DELEGATE_RetVal_TwoParams(FReply, FOnActionDragged, const TArray< TSharedPtr<FEdGraphSchemaAction> >&, const FPointerEvent&);
+	/** Delegate executed when a category is dragged */
+	DECLARE_DELEGATE_RetVal_TwoParams(FReply, FOnCategoryDragged, const FText&, const FPointerEvent&);
+	/** Delegate executed when the list of all actions needs to be refreshed */
+	DECLARE_DELEGATE_OneParam(FOnCollectAllActions, FGraphActionListBuilderBase&);
+	/** Delegate executed when the list of all actions needs to be refreshed, should return any sections that should always be visible, even if they don't have children. */
+	DECLARE_DELEGATE_OneParam(FOnCollectStaticSections, TArray<int32>&)
+		/** Delegate executed when a category is being renamed so any post-rename actions can be handled */
+		DECLARE_DELEGATE_ThreeParams(FOnCategoryTextCommitted, const FText&, ETextCommit::Type, TWeakPtr< FGraphActionNode >);
+	/** Delegate executed to check if the selected action is valid for renaming */
+	DECLARE_DELEGATE_RetVal_OneParam(bool, FCanRenameSelectedAction, TWeakPtr< FGraphActionNode >);
+	/** Delegate to get the name of a section if the widget is a section separator. */
+	DECLARE_DELEGATE_RetVal_OneParam(FText, FGetSectionTitle, int32);
+	/** Delegate to get the tooltip of a section if the widget is a section separator. */
+	DECLARE_DELEGATE_RetVal_OneParam(TSharedPtr<IToolTip>, FGetSectionToolTip, int32);
+	/** Delegate to get the widget that appears on the section bar in the section separator. */
+	DECLARE_DELEGATE_RetVal_TwoParams(TSharedRef<SWidget>, FGetSectionWidget, TSharedRef<SWidget>, int32);
+	/** Delegate to get the filter text */
+	DECLARE_DELEGATE_RetVal(FText, FGetFilterText);
+	/** Delegate to check if an action matches a specified name (used for renaming items etc.) */
+	DECLARE_DELEGATE_RetVal_TwoParams(bool, FOnActionMatchesName, FEdGraphSchemaAction*, const FName&);
+	/** Delegate that can be used to create and/or get a custom action list. */
+	DECLARE_DELEGATE_RetVal(TSharedRef<FGraphActionListBuilderBase>, FOnGetActionList);
+
+	SLATE_BEGIN_ARGS(STG_ActionMenuTileView)
+		: _AutoExpandActionMenu(false)
+		, _AlphaSortItems(true)
+		, _SortItemsRecursively(true)
+		, _ShowFilterTextBox(true)
+		, _UseSectionStyling(false)
+		, _bAllowPreselectedItemActivation(false)
+		, _GraphObj(nullptr)
+	{ }
+
+	SLATE_EVENT(FOnActionSelected, OnActionSelected)
+		SLATE_EVENT(FOnActionDoubleClicked, OnActionDoubleClicked)
+		SLATE_EVENT(FOnActionDragged, OnActionDragged)
+		SLATE_EVENT(FOnCategoryDragged, OnCategoryDragged)
+		SLATE_EVENT(FOnContextMenuOpening, OnContextMenuOpening)
+		SLATE_EVENT(FOnCreateWidgetForAction, OnCreateWidgetForAction)
+		SLATE_EVENT(FOnCreateCustomRowExpander, OnCreateCustomRowExpander)
+		SLATE_EVENT(FOnGetActionList, OnGetActionList)
+		SLATE_EVENT(FOnCollectAllActions, OnCollectAllActions)
+		SLATE_EVENT(FOnCollectStaticSections, OnCollectStaticSections)
+		SLATE_EVENT(FOnCategoryTextCommitted, OnCategoryTextCommitted)
+		SLATE_EVENT(FCanRenameSelectedAction, OnCanRenameSelectedAction)
+		SLATE_EVENT(FGetSectionTitle, OnGetSectionTitle)
+		SLATE_EVENT(FGetSectionToolTip, OnGetSectionToolTip)
+		SLATE_EVENT(FGetSectionWidget, OnGetSectionWidget)
+		SLATE_EVENT(FGetFilterText, OnGetFilterText)
+		SLATE_EVENT(FOnActionMatchesName, OnActionMatchesName)
+		SLATE_ARGUMENT(bool, AutoExpandActionMenu)
+		SLATE_ARGUMENT(bool, AlphaSortItems)
+		SLATE_ARGUMENT(bool, SortItemsRecursively)
+		SLATE_ARGUMENT(bool, ShowFilterTextBox)
+		SLATE_ARGUMENT(bool, UseSectionStyling)
+		SLATE_ARGUMENT(bool, bAllowPreselectedItemActivation)
+		SLATE_ARGUMENT(TArray<UEdGraphPin*>, DraggedFromPins)
+		SLATE_ARGUMENT(UEdGraph*, GraphObj)
+
+		SLATE_END_ARGS()
+
+		void Construct(const FArguments& InArgs, bool bIsReadOnly = true);
+
+	// FGCObject override
+	virtual void AddReferencedObjects(FReferenceCollector& Collector) override;
+	virtual FString GetReferencerName() const override;
+	virtual void Tick(const FGeometry& AllottedGeometry, const double InCurrentTime, const float InDeltaTime) override;
+
+	/**
+	 * Refreshes the actions that this widget should display
+	 *
+	 * @param bPreserveExpansion		TRUE if the expansion state of the tree should be preserved
+	 * @param bHandleOnSelectionEvent		TRUE if the item should be selected and any actions that occur with selection will be handled. FALSE and only selection will occur */
+	void RefreshAllActions(bool bPreserveExpansion, bool bHandleOnSelectionEvent = true);
+
+	/** Returns a map of all top level sections and their current expansion state. */
+	void GetSectionExpansion(TMap<int32, bool>& SectionExpansion) const;
+
+	/** Sets the sections to be expanded of all top level sections. */
+	void SetSectionExpansion(const TMap<int32, bool>& SectionExpansion);
+
+	void CalculateFillScale(const FGeometry& AllottedGeometry);
+
+	float GetTileViewItemWidth() const;
+
+protected:
+	/** Tree view for showing actions */
+	TSharedPtr< STreeView< TSharedPtr<FGraphActionNode> > > TreeView;
+
+	/** Tree tile view map for showing actions */
+	TMap< TSharedPtr<FGraphActionNode>, TSharedPtr<STileView<TSharedPtr<FGraphActionNode> > > > TreeTileViewMap;
+	
+	/** Text box used for searching for actions */
+	TSharedPtr<SSearchBox> FilterTextBox;
+
+	/** List of all actions we can browser */
+	TSharedPtr<FGraphActionListBuilderBase> AllActions;
+
+	/** Flattened list of all actions passing the filter */
+	TArray< TSharedPtr<FGraphActionNode> > FilteredActionNodes;
+
+	/** Root of filtered actions tree */
+	TSharedPtr<FGraphActionNode> FilteredRootAction;
+
+	/** Stored score for our current selection, so that we can quickly maintain selection when building the list asynchronously */
+	float SelectedSuggestionScore;
+	/** Stored index in the AllActions list builder that we think is the best fit */
+	int32 SelectedSuggestionSourceIndex;
+	/** Used to track selected action for keyboard interaction */
+	int32 SelectedSuggestion;
+	/** Allows us to set selection (via keyboard) without triggering action */
+	bool bIgnoreUIUpdate;
+	/** Should we auto-expand categories */
+	bool bAutoExpandActionMenu;
+	/** Should we display the filter text box */
+	bool bShowFilterTextBox;
+	/** Don't sort items alphabetically */
+	bool bAlphaSortItems;
+	/** If only the top entries should be sorted or subentries as well */
+	bool bSortItemsRecursively;
+	/** Should the rows and sections be styled like the details panel? */
+	bool bUseSectionStyling;
+	/** Whether we allow pre-selected items to be activated with a left-click */
+	bool bAllowPreselectedItemActivation;
+
+	/** Delegate to call when action is selected */
+	FOnActionSelected OnActionSelected;
+	/** Delegate to call when action is double clicked */
+	FOnActionDoubleClicked OnActionDoubleClicked;
+	/** Delegate to call when an action is dragged. */
+	FOnActionDragged OnActionDragged;
+	/** Delegate to call when a category is dragged. */
+	FOnCategoryDragged OnCategoryDragged;
+	/** Delegate to call to create widget for an action */
+	FOnCreateWidgetForAction OnCreateWidgetForAction;
+	/** Delegate to call for creating a custom "expander" widget for indenting a menu row with */
+	FOnCreateCustomRowExpander OnCreateCustomRowExpander;
+	/** Delegate to call to get a custom action list */
+	FOnGetActionList OnGetActionList;
+	/** Delegate to call to collect all actions */
+	FOnCollectAllActions OnCollectAllActions;
+	/** Delegate to call to collect all always visible sections */
+	FOnCollectStaticSections OnCollectStaticSections;
+	/** Delegate to call to handle any post-category rename events */
+	FOnCategoryTextCommitted OnCategoryTextCommitted;
+	/** Delegate to call to check if a selected action is valid for renaming */
+	FCanRenameSelectedAction OnCanRenameSelectedAction;
+	/** Delegate to get the name of a section separator. */
+	FGetSectionTitle OnGetSectionTitle;
+	/** Delegate to get the tooltip of a section separator. */
+	FGetSectionToolTip OnGetSectionToolTip;
+	/** Delegate to get the widgets of a section separator. */
+	FGetSectionWidget OnGetSectionWidget;
+	/** Delegate to get the filter text if supplied from an external source */
+	FGetFilterText OnGetFilterText;
+	/** Delegate to check if an action matches a specified name (used for renaming items etc.) */
+	FOnActionMatchesName OnActionMatchesName;
+
+public:
+	// SWidget interface
+	virtual FReply OnKeyDown(const FGeometry& MyGeometry, const FKeyEvent& KeyEvent) override;
+	// End of SWidget interface
+
+	/** Get filter text box widget */
+	TSharedRef<SEditableTextBox> GetFilterTextBox();
+
+	/** Get action that is currently selected */
+	void GetSelectedActions(TArray< TSharedPtr<FEdGraphSchemaAction> >& OutSelectedActions) const;
+
+	/** Initiates a rename on the selected action node, if possible */
+	void OnRequestRenameOnActionNode();
+
+	/** Queries if a rename on the selected action node is possible */
+	bool CanRequestRenameOnActionNode() const;
+
+	/** Get category that is currently selected */
+	FString GetSelectedCategoryName() const;
+
+	/** Get category child actions that is currently selected */
+	void GetSelectedCategorySubActions(TArray<TSharedPtr<FEdGraphSchemaAction>>& OutActions) const;
+
+	/** Get category child actions for the passed in action */
+	void GetCategorySubActions(TWeakPtr<FGraphActionNode> InAction, TArray<TSharedPtr<FEdGraphSchemaAction>>& OutActions) const;
+
+	/**
+	 * Selects an non-creation item in the list, searching by FName, deselects if name is none
+	 *
+	 * @param	ItemName		The name of the item to select
+	 * @param	SelectInfo		The selection type
+	 * @param	SectionId		If known, the section Id to restrict the selection to, useful in the case of categories where they can exist multiple times
+	 * @param	bIsCategory		TRUE if the selection is a category, categories obey different rules and it's hard to re-select properly without this knowledge
+	 * @return					TRUE if the item was successfully selected or the tree cleared, FALSE if unsuccessful
+	 */
+	bool SelectItemByName(const FName& ItemName, ESelectInfo::Type SelectInfo = ESelectInfo::Direct, int32 SectionId = INDEX_NONE, bool bIsCategory = false);
+
+	/** Expands any category with the associated name */
+	void ExpandCategory(const FText& CategoryName);
+
+	/* Handler for mouse button going down */
+	bool OnMouseButtonDownEvent(TWeakPtr<FEdGraphSchemaAction> InAction);
+
+	/** Updates the displayed list starting from IdxStart, useful for async building the display list of actions */
+	void UpdateForNewActions(int32 IdxStart);
+	/** Regenerated filtered results (FilteredRootAction and FilteredActionNodes) based on filter text  */
+	void GenerateFilteredItems(bool bPreserveExpansion);
+
+	/** The last typed action within the graph action menu */
+	static FString LastUsedFilterText;
+
+protected:
+	/** Get current filter text */
+	FText GetFilterText() const;
+
+	/** Change the selection to reflect the active suggestion */
+	void MarkActiveSuggestion();
+
+	/** Try to spawn the node reflected by the active suggestion */
+	bool TryToSpawnActiveSuggestion();
+
+	/** Returns true if the tree should be autoexpanded */
+	bool ShouldExpandNodes() const;
+
+	/** Checks if the passed in node is safe for renaming */
+	bool CanRenameNode(TWeakPtr<FGraphActionNode> InNode) const;
+
+	// Delegates
+
+	/** Called when filter text changes */
+	void OnFilterTextChanged(const FText& InFilterText);
+	/** Called when enter is hit in search box */
+	void OnFilterTextCommitted(const FText& InText, ETextCommit::Type CommitInfo);
+	/** Get children  */
+	void OnGetChildrenForCategory(TSharedPtr<FGraphActionNode> InItem, TArray< TSharedPtr<FGraphActionNode> >& OutChildren);
+	/** Create widget for the supplied node */
+	TSharedRef<ITableRow> MakeWidget(TSharedPtr<FGraphActionNode> InItem, const TSharedRef<STableViewBase>& OwnerTable, bool bIsReadOnly,bool bMakeTiles);
+	/*Create the widget for the Child node of the Tree hooked up with treeview*/
+	TSharedRef<ITableRow> GenerateChildForTree(TSharedPtr<FGraphActionNode> InItem, const TSharedRef<STableViewBase>& OwnerTable, bool bIsReadOnly);
+	/*Create the widget for the child of tile view inside the tree's category node, hooked up with tile view of each category node*/
+	TSharedRef<ITableRow> GenerateTileForChildItem(TSharedPtr<FGraphActionNode> InItem, const TSharedRef<STableViewBase>& OwnerTable);
+	/**
+	 * Called when tree item is selected
+	 *
+	 * @param	InSelectedItem	The action node that is being selected
+	 * @param	SelectInfo		Selection type - Only OnMouseClick and OnKeyPress will trigger a call to HandleSelection
+	 */
+	void OnItemSelected(TSharedPtr< FGraphActionNode > InSelectedItem, ESelectInfo::Type SelectInfo);
+	/**
+	 * Called when Tile view item is selected
+	 *
+	 * @param	InSelectedItem	The action node that is being selected
+	 * @param	SelectInfo		Selection type - Only OnMouseClick and OnKeyPress will trigger a call to HandleSelection
+	 */
+	void OnTileItemSelected(TSharedPtr< FGraphActionNode > InSelectedItem, ESelectInfo::Type SelectInfo);
+
+	TArray< TSharedPtr<FGraphActionNode> > GetSelectedItems() const;
+	/**
+	 * Executes the selection delegate providing it has been bound, and the provided action node given is valid and is an action node
+	 *
+	 * @param InselectedItem	The graph action node selected
+	 *
+	 * @return true if item selection delegate was executed
+	 */
+	bool HandleSelection(TSharedPtr< FGraphActionNode >& InSelectedItem, ESelectInfo::Type InSelectionType);
+
+	/** Called when tree item is double clicked */
+	void OnItemDoubleClicked(TSharedPtr< FGraphActionNode > InClickedItem);
+	/** Called when tree item dragged */
+	FReply OnItemDragDetected(const FGeometry& MyGeometry, const FPointerEvent& MouseEvent);
+	/** Callback when rename text is committed */
+	void OnNameTextCommitted(const FText& NewText, ETextCommit::Type InTextCommit, TWeakPtr< FGraphActionNode > InAction);
+	/** Handler for when an item has scrolled into view after having been requested to do so */
+	void OnItemScrolledIntoView(TSharedPtr<FGraphActionNode> InActionNode, const TSharedPtr<ITableRow>& InWidget);
+	/** Callback for expanding tree items recursively */
+	void OnSetExpansionRecursive(TSharedPtr<FGraphActionNode> InTreeNode, bool bInIsItemExpanded);
+	/** Helper function for adding and scoring actions from our builder */
+	struct FScoreResults
+	{
+		int32 BestMatchIndex;
+		float BestMatchScore;
+	};
+	FScoreResults ScoreAndAddActions(int32 StartingIndex = INDEX_NONE);
+	/** Helper function to update the active selection after updating the displayed tree */
+	void UpdateActiveSelection(FScoreResults ForResults);
+private:
+	/** The pins that have been dragged off of to prompt the creation of this action menu. */
+	TArray<UEdGraphPin*> DraggedFromPins;
+
+	/** The graph that this menu is being constructed in */
+	UEdGraph* GraphObj;
+
+	float FillScale;
+
+	static const int TileWidth;
+	
+	static const int TileHeight;
+};
diff --git a/Engine/Plugins/Experimental/TextureGraph/Source/TextureGraphEditor/Private/Widgets/ActionMenuWidgets/STG_GraphActionMenu.cpp b/Engine/Plugins/Experimental/TextureGraph/Source/TextureGraphEditor/Private/Widgets/ActionMenuWidgets/STG_GraphActionMenu.cpp
new file mode 100644
--- /dev/null
+++ b/Engine/Plugins/Experimental/TextureGraph/Source/TextureGraphEditor/Private/Widgets/ActionMenuWidgets/STG_GraphActionMenu.cpp
@@ -0,0 +1,126 @@
+// Copyright Epic Games, Inc. All Rights Reserved.
+
+#include "STG_GraphActionMenu.h"
+
+#include "EdGraph/EdGraph.h"
+#include "EdGraph/EdGraphSchema.h"
+#include "Framework/Application/SlateApplication.h"
+#include "HAL/PlatformMath.h"
+#include "Layout/Margin.h"
+#include "Misc/Attribute.h"
+#include "Styling/AppStyle.h"
+#include "Types/SlateStructs.h"
+#include "Widgets/Layout/SBox.h"
+#include "Widgets/Text/SInlineEditableTextBlock.h"
+#include "Widgets/SBoxPanel.h"
+#include "TG_Style.h"
+#include "EdGraph/TG_EdGraphSchema.h"
+#include "STG_Palette.h"
+
+class SEditableTextBox;
+
+STG_GraphActionMenu::~STG_GraphActionMenu()
+{
+	OnClosedCallback.ExecuteIfBound();
+}
+
+void STG_GraphActionMenu::Construct(const FArguments& InArgs)
+{
+	this->GraphObj = InArgs._GraphObj;
+	this->bSpwanOnSelect = InArgs._bSpwanOnSelect;
+	this->DraggedFromPins = InArgs._DraggedFromPins;
+	this->NewNodePosition = InArgs._NewNodePosition;
+	this->OnClosedCallback = InArgs._OnClosedCallback;
+	this->AutoExpandActionMenu = InArgs._AutoExpandActionMenu;
+
+	FLinearColor FillColor = UTG_EdGraphSchema::NodeBodyColor;
+
+	if (bSpwanOnSelect)
+	{
+		Brush = *FTG_Style::Get().GetBrush("TG.Palette.Backgorund");
+		Brush.OutlineSettings.Color = FSlateColor(FillColor * UTG_EdGraphSchema::NodeOutlineColorMultiplier);
+	}
+	else
+	{
+		Brush = *FCoreStyle::Get().GetBrush("WhiteBrush");
+		Brush.OutlineSettings.Color = FSlateColor(FillColor);
+	}
+
+	// Build the widget layout
+	SBorder::Construct(SBorder::FArguments()
+		.BorderImage(&Brush)
+		.BorderBackgroundColor(FillColor)
+		.Padding(2)
+		[
+			// Achieving fixed width by nesting items within a fixed width box.
+			SNew(SBox)
+			.WidthOverride(400)
+			.HeightOverride(400)
+			[
+				SAssignNew(GraphActionMenu, SGraphActionMenu)
+				.OnActionSelected(this, &STG_GraphActionMenu::OnActionSelected)
+				.OnCollectAllActions(this, &STG_GraphActionMenu::CollectAllActions)
+				.OnCreateWidgetForAction(this, &STG_GraphActionMenu::OnCreateWidgetForAction)
+				.AutoExpandActionMenu(AutoExpandActionMenu)
+				.DraggedFromPins(DraggedFromPins)
+				.OnActionDragged(InArgs._OnActionDragged)
+				.GraphObj(GraphObj)
+				.DefaultRowExpanderBaseIndentLevel(10)
+			]
+		]);
+}
+
+void STG_GraphActionMenu::CollectAllActions(FGraphActionListBuilderBase& OutAllActions)
+{
+	// Build up the context object
+	FGraphContextMenuBuilder ContextMenuBuilder(GraphObj);
+	if (DraggedFromPins.Num() > 0)
+	{
+		ContextMenuBuilder.FromPin = DraggedFromPins[0];
+	}
+
+	// Determine all possible actions
+	GraphObj->GetSchema()->GetGraphContextActions(ContextMenuBuilder);
+
+	// Copy the added options back to the main list
+	//@TODO: Avoid this copy
+	OutAllActions.Append(ContextMenuBuilder);
+}
+
+TSharedRef<SWidget> STG_GraphActionMenu::OnCreateWidgetForAction(FCreateWidgetForActionData* const InCreateData)
+{
+	return	SNew(STG_PaletteItem, InCreateData);
+}
+
+TSharedRef<SEditableTextBox> STG_GraphActionMenu::GetFilterTextBox()
+{
+	return GraphActionMenu->GetFilterTextBox();
+}
+
+
+void STG_GraphActionMenu::OnActionSelected(const TArray< TSharedPtr<FEdGraphSchemaAction> >& SelectedAction, ESelectInfo::Type InSelectionType)
+{
+	if (bSpwanOnSelect && (InSelectionType == ESelectInfo::OnMouseClick || InSelectionType == ESelectInfo::OnKeyPress || SelectedAction.Num() == 0))
+	{
+		bool bDoDismissMenus = true;
+
+		if (GraphObj != NULL)
+		{
+			for (int32 ActionIndex = 0; ActionIndex < SelectedAction.Num(); ActionIndex++)
+			{
+				TSharedPtr<FEdGraphSchemaAction> CurrentAction = SelectedAction[ActionIndex];
+
+				if (CurrentAction.IsValid())
+				{
+					if (bDoDismissMenus)
+					{
+						FSlateApplication::Get().DismissAllMenus();
+						bDoDismissMenus = false;
+					}
+
+					CurrentAction->PerformAction(GraphObj, DraggedFromPins, NewNodePosition);
+				}
+			}
+		}
+	}
+}
diff --git a/Engine/Plugins/Experimental/TextureGraph/Source/TextureGraphEditor/Private/Widgets/ActionMenuWidgets/STG_GraphActionMenu.h b/Engine/Plugins/Experimental/TextureGraph/Source/TextureGraphEditor/Private/Widgets/ActionMenuWidgets/STG_GraphActionMenu.h
new file mode 100644
--- /dev/null
+++ b/Engine/Plugins/Experimental/TextureGraph/Source/TextureGraphEditor/Private/Widgets/ActionMenuWidgets/STG_GraphActionMenu.h
@@ -0,0 +1,79 @@
+// Copyright Epic Games, Inc. All Rights Reserved.
+
+#pragma once
+
+#include "Containers/Array.h"
+#include "CoreMinimal.h"
+#include "EdGraph/EdGraphSchema.h"
+#include "GraphEditor.h"
+#include "HAL/PlatformCrt.h"
+#include "Math/Vector2D.h"
+#include "Templates/SharedPointer.h"
+#include "Types/SlateEnums.h"
+#include "Widgets/DeclarativeSyntaxSupport.h"
+#include "Widgets/Input/SEditableTextBox.h"
+#include "Widgets/Layout/SBorder.h"
+#include "SGraphActionMenu.h"
+#include "SGraphPalette.h"
+
+class SEditableTextBox;
+class UEdGraph;
+class UEdGraphPin;
+struct FEdGraphSchemaAction;
+struct FGraphActionListBuilderBase;
+class SInlineEditableTextBlock;
+
+/////////////////////////////////////////////////////////////////////////////////////////////////
+/* This class implements the customization to the SGraphActionMenu for Texture Graph*/
+class STG_GraphActionMenu : public SBorder
+{
+public:
+	SLATE_BEGIN_ARGS(STG_GraphActionMenu)
+		: _GraphObj(static_cast<UEdGraph*>(NULL))
+		, _bSpwanOnSelect(true)
+		, _NewNodePosition(FVector2D::ZeroVector)
+		, _AutoExpandActionMenu(false)
+	{}
+
+	SLATE_ARGUMENT(UEdGraph*, GraphObj)
+		SLATE_ARGUMENT(bool, bSpwanOnSelect)
+		SLATE_ARGUMENT(FVector2D, NewNodePosition)
+		SLATE_ARGUMENT(TArray<UEdGraphPin*>, DraggedFromPins)
+		SLATE_ARGUMENT(SGraphEditor::FActionMenuClosed, OnClosedCallback)
+		SLATE_ARGUMENT(bool, AutoExpandActionMenu)
+
+		SLATE_EVENT(SGraphActionMenu::FOnActionDragged, OnActionDragged)
+		SLATE_EVENT(SGraphActionMenu::FOnCollectAllActions, OnCollectAllActions)
+		SLATE_EVENT(SGraphActionMenu::FOnCreateWidgetForAction, OnCreateWidgetForAction)
+		SLATE_END_ARGS()
+
+		void Construct(const FArguments& InArgs);
+
+	~STG_GraphActionMenu();
+
+	TSharedRef<SEditableTextBox> GetFilterTextBox();
+	TSharedRef<SWidget> OnCreateWidgetForAction(FCreateWidgetForActionData* const InCreateData);
+
+protected:
+	//From Right click menu action is executed when selected
+	bool bSpwanOnSelect;
+
+	FSlateBrush Brush;
+
+	UEdGraph* GraphObj;
+
+	TArray<UEdGraphPin*> DraggedFromPins;
+
+	FVector2D NewNodePosition;
+
+	bool AutoExpandActionMenu;
+
+	SGraphEditor::FActionMenuClosed OnClosedCallback;
+
+	TSharedPtr<SGraphActionMenu> GraphActionMenu;
+
+	void OnActionSelected(const TArray< TSharedPtr<FEdGraphSchemaAction> >& SelectedAction, ESelectInfo::Type InSelectionType);
+
+	/** Callback used to populate all actions list in SGraphActionMenu */
+	void CollectAllActions(FGraphActionListBuilderBase& OutAllActions);
+};
