Fixed Issue: Texture parameter is not being displayed in parameters view

Fix/How: In the case of texture expression, Param is TG_Texture but it cannot be modified here. Use the UTexture source because it is the only way we can change the parameter value via the parameter panel.

diff --git a/Engine/Plugins/Experimental/TextureGraph/Source/TextureGraphEditor/Private/Customizations/TG_ParameterCustomization.h b/Engine/Plugins/Experimental/TextureGraph/Source/TextureGraphEditor/Private/Customizations/TG_ParameterCustomization.h
--- a/Engine/Plugins/Experimental/TextureGraph/Source/TextureGraphEditor/Private/Customizations/TG_ParameterCustomization.h
+++ b/Engine/Plugins/Experimental/TextureGraph/Source/TextureGraphEditor/Private/Customizations/TG_ParameterCustomization.h
@@ -18,6 +18,7 @@
 #include "TG_Pin.h"
 #include "TG_Graph.h"
 #include "Expressions/TG_Expression.h"
+#include "Expressions/Input/TG_Expression_Texture.h"
 #include "Widgets/Layout/SBox.h"
 #include "Widgets/Text/STextBlock.h"
 
@@ -63,7 +64,21 @@ public:
 					TArray<UObject*> ExternalObjects;
 					ExternalObjects.Add(ExpressionPtr);
 
-					auto PinName = PinPtr->GetArgumentName();
+					FName PinName;
+
+					// In case of texture expression
+					// Param is TG_Texture but it cannot be modified here
+					// use UTexture source because its the only way we can change the parameter value via parameter panel.
+					UTG_Expression_Texture* TextureExpression = Cast<UTG_Expression_Texture>(ExpressionPtr);
+					if(TextureExpression)
+					{
+						PinName = GET_MEMBER_NAME_CHECKED(UTG_Expression_Texture, Source);
+					}
+					else
+					{
+						PinName = PinPtr->GetArgumentName();
+					}
+					
 					IDetailPropertyRow* ParameterDetailsPropertyRow = ChildBuilder.AddExternalObjectProperty(ExternalObjects, PinName, FAddPropertyParams().HideRootObjectNode(true));
 					if (ParameterDetailsPropertyRow)
 					{
