crash on Thumbnail creation when closing a graph

diff --git a/Engine/Plugins/Experimental/TextureGraph/Source/TextureGraphEditor/Private/Customizations/STG_TextureHistogram.cpp b/Engine/Plugins/Experimental/TextureGraph/Source/TextureGraphEditor/Private/Customizations/STG_TextureHistogram.cpp
--- a/Engine/Plugins/Experimental/TextureGraph/Source/TextureGraphEditor/Private/Customizations/STG_TextureHistogram.cpp
+++ b/Engine/Plugins/Experimental/TextureGraph/Source/TextureGraphEditor/Private/Customizations/STG_TextureHistogram.cpp
@@ -42,8 +42,11 @@ void STG_TextureHistogram::SetTexture(FTG_Texture& Source, UTextureGraph* InText
 		})
 		.then([this, Source]() mutable
 		{
-			TiledBlobPtr FinalizedHistogram = std::static_pointer_cast<TiledBlob>(Source.RasterBlob->GetHistogram());
-			HistogramBars->Update(FinalizedHistogram);
+			if (DoesSharedInstanceExist())
+            {
+				TiledBlobPtr FinalizedHistogram = std::static_pointer_cast<TiledBlob>(Source.RasterBlob->GetHistogram());
+				HistogramBars->Update(FinalizedHistogram);
+			}
 		});
 	}
 }
\ No newline at end of file
diff --git a/Engine/Plugins/Experimental/TextureGraph/Source/TextureGraphEditor/Private/EdGraph/STG_EditorGraphNode.cpp b/Engine/Plugins/Experimental/TextureGraph/Source/TextureGraphEditor/Private/EdGraph/STG_EditorGraphNode.cpp
--- a/Engine/Plugins/Experimental/TextureGraph/Source/TextureGraphEditor/Private/EdGraph/STG_EditorGraphNode.cpp
+++ b/Engine/Plugins/Experimental/TextureGraph/Source/TextureGraphEditor/Private/EdGraph/STG_EditorGraphNode.cpp
@@ -885,7 +885,10 @@ TSharedRef<SWidget> STG_EditorGraphNode::CreateTitleRightWidget()
 			{
 				CachedThumb->OnFinalise().then([CachedThumb, ThumbnailWidget]()
 				{
-					ThumbnailWidget->UpdateBlob(CachedThumb);
+					if (ThumbnailWidget.IsValid())
+					{
+						ThumbnailWidget->UpdateBlob(CachedThumb);
+					}
 				});
 			}
 			
@@ -1013,7 +1016,11 @@ void STG_EditorGraphNode::UpdateThumbnail(const FTG_EvaluationContext* InContext
 					.then([=, this](BufferResultPtr) mutable
 						{
 							auto ThumbWidget = FindOrCreateThumbWidget(Pin->GetId());
-							ThumbWidget->UpdateBlob(ThumbBlob);
+
+							if (ThumbWidget.IsValid())
+							{
+								ThumbWidget->UpdateBlob(ThumbBlob);
+							}
 						});
 			}
 		}
diff --git a/Engine/Plugins/Experimental/TextureGraph/Source/TextureGraphEditor/Private/EdGraph/STG_NodeThumbnail.cpp b/Engine/Plugins/Experimental/TextureGraph/Source/TextureGraphEditor/Private/EdGraph/STG_NodeThumbnail.cpp
--- a/Engine/Plugins/Experimental/TextureGraph/Source/TextureGraphEditor/Private/EdGraph/STG_NodeThumbnail.cpp
+++ b/Engine/Plugins/Experimental/TextureGraph/Source/TextureGraphEditor/Private/EdGraph/STG_NodeThumbnail.cpp
@@ -83,35 +83,38 @@ void STG_NodeThumbnail::UpdateParams(TiledBlobPtr InBlob)
 	/// We don't assume that te blob would be finalised when we get to this point
 	InBlob->OnFinalise().then([=, this]()
 		{
-			UTexture* BlobTexture = GetTextureFromBlob(InBlob);
-			check(BlobTexture);
-
-			//FRHITexture* SlateTexture = (FRHITexture*)(AssetThumbnail->GetViewportRenderTargetTexture());
-			UTexture2D* Texture2D = Cast<UTexture2D>(BlobTexture);
-			UTextureRenderTarget2D* TextureRT2D = Cast<UTextureRenderTarget2D>(BlobTexture);
-
-			float ShowChecker = 1.0;
-			if (BlobTexture)
-			{
-				ShowChecker = 0.0;
-				BrushMaterial->SetTextureParameterValue("ThumbTex", BlobTexture);
-			}
-
-			float SingleChannel = 0.0;
-
-			if (Texture2D)
+			if (DoesSharedInstanceExist())
 			{
-				SingleChannel = GPixelFormats[Texture2D->GetPixelFormat()].NumComponents == 1 ? 1.0 : 0.0;
+				UTexture* BlobTexture = GetTextureFromBlob(InBlob);
+				check(BlobTexture);
+
+				//FRHITexture* SlateTexture = (FRHITexture*)(AssetThumbnail->GetViewportRenderTargetTexture());
+				UTexture2D* Texture2D = Cast<UTexture2D>(BlobTexture);
+				UTextureRenderTarget2D* TextureRT2D = Cast<UTextureRenderTarget2D>(BlobTexture);
+				
+				float ShowChecker = 1.0;
+				if (BlobTexture)
+				{
+					ShowChecker = 0.0;
+					BrushMaterial->SetTextureParameterValue("ThumbTex", BlobTexture);
+				}
+
+				float SingleChannel = 0.0;
+
+				if (Texture2D)
+				{
+					SingleChannel = GPixelFormats[Texture2D->GetPixelFormat()].NumComponents == 1 ? 1.0 : 0.0;
+				}
+
+				if (TextureRT2D)
+				{
+					SingleChannel = GPixelFormats[TextureRT2D->GetFormat()].NumComponents == 1 ? 1.0 : 0.0;
+				}
+
+				BrushMaterial->SetScalarParameterValue("ShowChecker", ShowChecker);
+				BrushMaterial->SetScalarParameterValue("SingleChannel", SingleChannel);
+				Brush->SetResourceObject(BrushMaterial);
 			}
-
-			if (TextureRT2D)
-			{
-				SingleChannel = GPixelFormats[TextureRT2D->GetFormat()].NumComponents == 1 ? 1.0 : 0.0;
-			}
-
-			BrushMaterial->SetScalarParameterValue("ShowChecker", ShowChecker);
-			BrushMaterial->SetScalarParameterValue("SingleChannel", SingleChannel);
-			Brush->SetResourceObject(BrushMaterial);
 		});
 }
 
diff --git a/Engine/Plugins/Experimental/TextureGraph/Source/TextureGraphEditor/Private/STG_OutputSelectionDlg.cpp b/Engine/Plugins/Experimental/TextureGraph/Source/TextureGraphEditor/Private/STG_OutputSelectionDlg.cpp
--- a/Engine/Plugins/Experimental/TextureGraph/Source/TextureGraphEditor/Private/STG_OutputSelectionDlg.cpp
+++ b/Engine/Plugins/Experimental/TextureGraph/Source/TextureGraphEditor/Private/STG_OutputSelectionDlg.cpp
@@ -97,7 +97,10 @@ void STG_OutputSelectionDlg::AddExportItems()
 				ThumbBlob->OnFinalise()
 					.then([ThumbBlob, NodeThumbnail]
 					{
-						NodeThumbnail->UpdateBlob(ThumbBlob);
+						if (NodeThumbnail.IsValid())
+						{
+							NodeThumbnail->UpdateBlob(ThumbBlob);
+						}
 					});
 
 				ThumbnailWidget = NodeThumbnail;
diff --git a/Engine/Plugins/Experimental/TextureGraph/Source/TextureGraphEditor/Private/STG_TextureDetails.cpp b/Engine/Plugins/Experimental/TextureGraph/Source/TextureGraphEditor/Private/STG_TextureDetails.cpp
--- a/Engine/Plugins/Experimental/TextureGraph/Source/TextureGraphEditor/Private/STG_TextureDetails.cpp
+++ b/Engine/Plugins/Experimental/TextureGraph/Source/TextureGraphEditor/Private/STG_TextureDetails.cpp
@@ -159,10 +159,29 @@ void STG_TextureDetails::CalculateHistogram(BlobPtr InBlob, UTextureGraph* InTex
 	})
 	.then([this, InBlob]() mutable
 	{
-		TiledBlobPtr FinalizedHistogram = std::static_pointer_cast<TiledBlob>(InBlob->GetHistogram());
-		HistogramBlobWidgetR->Update(FinalizedHistogram);
-		HistogramBlobWidgetG->Update(FinalizedHistogram);
-		HistogramBlobWidgetB->Update(FinalizedHistogram);
-		HistogramBlobWidgetLuma->Update(FinalizedHistogram);
+		if (DoesSharedInstanceExist())
+		{
+			TiledBlobPtr FinalizedHistogram = std::static_pointer_cast<TiledBlob>(InBlob->GetHistogram());
+			if (HistogramBlobWidgetR.IsValid())
+			{
+				HistogramBlobWidgetR->Update(FinalizedHistogram);
+			}
+
+			if (HistogramBlobWidgetG.IsValid())
+			{
+				HistogramBlobWidgetG->Update(FinalizedHistogram);
+			}
+		
+			if (HistogramBlobWidgetB.IsValid())
+			{
+				HistogramBlobWidgetB->Update(FinalizedHistogram);
+			}
+		
+			if (HistogramBlobWidgetLuma.IsValid())
+			{
+				HistogramBlobWidgetLuma->Update(FinalizedHistogram);
+			}
+		}
+		
 	});
 }
