Browse button bug fixes

diff --git a/Engine/Plugins/Experimental/TextureGraph/Source/TextureGraphEditor/Private/EdGraph/TG_EdGraphSchema.cpp b/Engine/Plugins/Experimental/TextureGraph/Source/TextureGraphEditor/Private/EdGraph/TG_EdGraphSchema.cpp
--- a/Engine/Plugins/Experimental/TextureGraph/Source/TextureGraphEditor/Private/EdGraph/TG_EdGraphSchema.cpp
+++ b/Engine/Plugins/Experimental/TextureGraph/Source/TextureGraphEditor/Private/EdGraph/TG_EdGraphSchema.cpp
@@ -157,6 +157,20 @@ void UTG_EdGraphSchema::TrySetDefaultValue(UEdGraphPin& Pin, const FString& NewD
 
 #endif	//#if WITH_EDITOR
 }
+
+void UTG_EdGraphSchema::TrySetDefaultObject(UEdGraphPin& Pin, UObject* NewDefaultObject, bool bMarkAsModified) const
+{
+	Pin.DefaultObject = NewDefaultObject;
+#if WITH_EDITOR
+	UTG_EdGraphNode* Node = Cast<UTG_EdGraphNode>(Pin.GetOwningNode());
+	check(Node);
+
+	// get file path from Object changed and set as default value
+	TrySetDefaultValue(Pin, NewDefaultObject ? NewDefaultObject->GetPathName() : FString(), bMarkAsModified);
+
+#endif	//#if WITH_EDITOR
+}
+
 const FPinConnectionResponse UTG_EdGraphSchema::CanCreateConnection(const UEdGraphPin* A, const UEdGraphPin* B) const
 {
 //	bool bPreventInvalidConnections = CVarPreventInvalidMaterialConnections.GetValueOnGameThread() != 0;
diff --git a/Engine/Plugins/Experimental/TextureGraph/Source/TextureGraphEditor/Private/EdGraph/TG_EdGraphSchema.h b/Engine/Plugins/Experimental/TextureGraph/Source/TextureGraphEditor/Private/EdGraph/TG_EdGraphSchema.h
--- a/Engine/Plugins/Experimental/TextureGraph/Source/TextureGraphEditor/Private/EdGraph/TG_EdGraphSchema.h
+++ b/Engine/Plugins/Experimental/TextureGraph/Source/TextureGraphEditor/Private/EdGraph/TG_EdGraphSchema.h
@@ -74,6 +74,7 @@ public:
 	virtual void GetAssetsGraphHoverMessage(const TArray<FAssetData>& Assets, const UEdGraph* HoverGraph, FString& OutTooltipText, bool& OutOkIcon) const override;
 	virtual void GetAssetsNodeHoverMessage(const TArray<struct FAssetData>& Assets, const UEdGraphNode* HoverNode, FString& OutTooltipText, bool& OutOkIcon) const override;
 	virtual void TrySetDefaultValue(UEdGraphPin& Pin, const FString& NewDefaultValue, bool bMarkAsModified = true) const override;
+	virtual void TrySetDefaultObject(UEdGraphPin& Pin, UObject* NewDefaultObject, bool bMarkAsModified = true) const;
 	//~ End EdGraphSchema Interface
 
 	UTG_Pin* GetTGPinFromEdPin(const UEdGraphPin* InPin) const;
