Recursive TextureGraphNode causes crash

diff --git a/Engine/Plugins/Experimental/TextureGraph/Source/TextureGraph/Private/Expressions/Input/TG_Expression_Graph.cpp b/Engine/Plugins/Experimental/TextureGraph/Source/TextureGraph/Private/Expressions/Input/TG_Expression_Graph.cpp
--- a/Engine/Plugins/Experimental/TextureGraph/Source/TextureGraph/Private/Expressions/Input/TG_Expression_Graph.cpp
+++ b/Engine/Plugins/Experimental/TextureGraph/Source/TextureGraph/Private/Expressions/Input/TG_Expression_Graph.cpp
@@ -2,7 +2,9 @@
 #include "Expressions/Input/TG_Expression_Graph.h"
 
 #include "TG_Graph.h"
+#include "UObject/Package.h"
 #include "TextureGraph.h"
+#include "UObject/ObjectSaveContext.h"
 
 //////////////////////////////////////////////////////////////////////////
 //// Generic subgraph
@@ -50,6 +52,21 @@ void UTG_Expression_Graph::SetupAndEvaluate(FTG_EvaluationContext* InContext)
 	Evaluate(InContext);
 }
 
+UTG_Expression_TextureGraph::UTG_Expression_TextureGraph()
+{
+#if WITH_EDITOR
+	// listener for UObject saves, so we can synchronise when linked TextureGraphs get updated 
+	PreSaveHandle = FCoreUObjectDelegates::OnObjectPreSave.AddUObject(this, &UTG_Expression_TextureGraph::OnTextureGraphPreSave);
+#endif
+}
+
+UTG_Expression_TextureGraph::~UTG_Expression_TextureGraph()
+{
+#if WITH_EDITOR
+	FCoreUObjectDelegates::OnObjectPreSave.Remove(PreSaveHandle);
+#endif
+}
+
 void UTG_Expression_Graph::Evaluate(FTG_EvaluationContext* InContext)
 {
 	auto Graph = GetGraph();
@@ -72,15 +89,17 @@ void UTG_Expression_Graph::Evaluate(FTG_EvaluationContext* InContext)
 	}
 }
 
+
 //////////////////////////////////////////////////////////////////////////
 // TextureTextureGraph Expression
 //////////////////////////////////////////////////////////////////////////
 
 #if WITH_EDITOR
+
 void UTG_Expression_TextureGraph::PostEditChangeProperty(FPropertyChangedEvent& PropertyChangedEvent)
 {
 	// if Graph changes catch it first
-	if (PropertyChangedEvent.GetPropertyName() == FName(TEXT("TextureGraph")))
+	if (PropertyChangedEvent.GetPropertyName() == GET_MEMBER_NAME_CHECKED(UTG_Expression_TextureGraph, TextureGraph))
 	{
 		UE_LOG(LogTextureGraph, Log, TEXT("TextureGraph Expression PostEditChangeProperty."));
 		SetTextureGraphInternal(TextureGraph);
@@ -88,20 +107,87 @@ void UTG_Expression_TextureGraph::PostEditChangeProperty(FPropertyChangedEvent&
 
 	Super::PostEditChangeProperty(PropertyChangedEvent);
 }
+
 #endif
 
+void UTG_Expression_TextureGraph::OnTextureGraphPreSave(UObject* Object, FObjectPreSaveContext SaveContext)
+{
+	// every editor should check if your texture graph is dependent on the object being saved
+	UTextureGraph* GraphBeingSaved = Cast<UTextureGraph>(Object);
+	// if object being saved is our linked TextureGraph, we re-create
+	if (GraphBeingSaved && TextureGraph)
+	{
+		UTextureGraph* ParentGraph = CastChecked<UTextureGraph>(GetOutermostObject());
+
+		// update if the graph being saved is the one we're using in this expression directly
+		// Or our parent graph is indirectly dependent on the graph being saved 
+		if (GraphBeingSaved == TextureGraph ||
+			ParentGraph->IsDependent(GraphBeingSaved))
+		{
+			// Recreate the Runtime copy
+			RuntimeGraph = DuplicateObject(TextureGraph->Graph(), this);	
+
+			// Graph is reset, notify the owning node / graph to update itself
+			NotifyGraphChanged();
+		}
+	}
+}
 void UTG_Expression_TextureGraph::Initialize()
 {
 	if (TextureGraph && !RuntimeGraph)
 	{
 		RuntimeGraph = DuplicateObject(TextureGraph->Graph(), this);
 	}
+
+}
+
+// CheckDependencies: Returns true if no unwanted dependencies found
+bool UTG_Expression_TextureGraph::CheckDependencies(const UTextureGraph* InTextureGraph) const
+{
+	if (InTextureGraph == nullptr)
+		return true;
+		
+	// check if InTextureGraph doesn't have "this" texture graph as a dependency
+	UTextureGraph* ParentGraph = CastChecked<UTextureGraph>(GetOutermostObject());
+	const UPackage* Package = this->GetPackage();  
+	const UPackage* SecondPackage = InTextureGraph->GetPackage();  
+	const bool bIsTransientPackage = Package->HasAnyFlags(RF_Transient) || Package == GetTransientPackage();
+		
+	if (ParentGraph == InTextureGraph ||
+		(Package == SecondPackage && !bIsTransientPackage))	 // early out if we're trying to assign our own texture graph)
+	{
+		auto ErrorType = static_cast<int32>(ETextureGraphErrorType::UNSUPPORTED_TYPE);
+			
+		TextureGraphEngine::GetErrorReporter(ParentGraph)->ReportError(ErrorType, FString::Printf(TEXT("Can not assign graph to itself")));
+		return false;
+	}
+
+	if (InTextureGraph->IsDependent(ParentGraph))
+	{
+		auto ErrorType = static_cast<int32>(ETextureGraphErrorType::RECURSIVE_CALL);
+	
+		TextureGraphEngine::GetErrorReporter(ParentGraph)->ReportError(ErrorType, FString::Printf(TEXT("Can not use graph as a cyclic dependency is detected")));
+		return false;
+	}
+
+	return true;
 }
 
 void UTG_Expression_TextureGraph::SetTextureGraphInternal(UTextureGraph* InTextureGraph)
 {
-	Modify();
+	// validate that we're allowed to use InTextureGraph as input here
+	if (!CheckDependencies(InTextureGraph))
+	{
+		// if dependencies exist, we don't want to keep using the problematic TextureGraph
+		// TODO: At the moment we reset TextureGraph to nullptr when a dependency check fails.
+		// Ideally we'd like to revert to the previous valid value. The problem here is that
+		// TextureGraph at this point comes pre-filled with the problematic value 
+		TextureGraph = nullptr;
+		RuntimeGraph = nullptr;
+		return;
+	}
 	
+	Modify();
 	if (!InTextureGraph)
 	{
 		TextureGraph = nullptr;
@@ -174,10 +260,10 @@ void UTG_Expression_TextureGraph::Evaluate(FTG_EvaluationContext* InContext)
 bool UTG_Expression_TextureGraph::Validate(MixUpdateCyclePtr Cycle)
 {
 	// if TextureGraph is already in execution list then
-	if(TextureGraph && Cycle->ContainsMix(TextureGraph))
+	if(TextureGraph &&  Cycle->ContainsMix(TextureGraph))
 	{
 		auto ErrorType = static_cast<int32>(ETextureGraphErrorType::RECURSIVE_CALL);
-
+		
 		const FString ErrorMsg = FString::Format(TEXT("Recursive Texture Graph call found in '{0}' asset."), {Cycle->TopMix()->GetName()} );
 		TextureGraphEngine::GetErrorReporter(Cycle->GetMix())->ReportError(ErrorType, ErrorMsg, GetParentNode());
 		return false;
@@ -190,7 +276,7 @@ bool UTG_Expression_TextureGraph::CanHandleAsset(UObject* Asset)
 {
 	const UTextureGraph* TSTextureGraph = Cast<UTextureGraph>(Asset);
 	
-	return TSTextureGraph != nullptr;
+	return TSTextureGraph != nullptr && this->GetPackage() != Asset->GetPackage();
 }
 
 void UTG_Expression_TextureGraph::SetAsset(UObject* Asset)
diff --git a/Engine/Plugins/Experimental/TextureGraph/Source/TextureGraph/Private/TextureGraph.cpp b/Engine/Plugins/Experimental/TextureGraph/Source/TextureGraph/Private/TextureGraph.cpp
--- a/Engine/Plugins/Experimental/TextureGraph/Source/TextureGraph/Private/TextureGraph.cpp
+++ b/Engine/Plugins/Experimental/TextureGraph/Source/TextureGraph/Private/TextureGraph.cpp
@@ -7,8 +7,10 @@
 #include "TG_CustomVersion.h"
 
 #include "UObject/ObjectSaveContext.h"
+#include "UObject/Package.h"
 
 #include "2D/TextureHelper.h"
+#include "Expressions/Input/TG_Expression_Graph.h"
 #include "FxMat/MaterialManager.h"
 #include "Model/Mix/MixManager.h"
 #include "Model/Mix/MixSettings.h"
@@ -18,6 +20,66 @@
 #include "Transform/Mix/T_InvalidateTiles.h"
 #include "Transform/Mix/T_UpdateTargets.h"
 
+
+bool UTextureGraph::CheckCyclicDependency(const UTextureGraph* InTextureGraph) const
+{
+	TArray<UTextureGraph*> DependentGraphs;
+	GatherAllDependentGraphs(DependentGraphs);
+
+	// if after exhausting all nodes and their dependent graphs recursively, we found our source graph, we have a dependency
+	return DependentGraphs.ContainsByPredicate([&InTextureGraph](const UTextureGraph* CurrentTextureGraph)
+	{
+		const UPackage* Package = CurrentTextureGraph->GetPackage();  
+		const UPackage* SecondPackage = InTextureGraph->GetPackage();  
+		const bool bIsTransientPackage = Package->HasAnyFlags(RF_Transient) || Package == GetTransientPackage();
+		return CurrentTextureGraph->GetOutermostObject() == InTextureGraph->GetOutermostObject() ||
+				(Package == SecondPackage && !bIsTransientPackage);
+	});
+}
+void UTextureGraph::GatherAllDependentGraphs(TArray<UTextureGraph*>& DependentGraphs) const
+{
+	Graph()->ForEachNodes(
+			[&](const UTG_Node* Node, uint32 Index)
+		{
+				if (Node)
+				{
+					if (UTG_Expression_TextureGraph* TextureGraphExpr = Cast<UTG_Expression_TextureGraph>(Node->GetExpression()))
+					{
+						const auto OriginalAssetForGraphInExpression = TextureGraphExpr->TextureGraph;
+
+						if (OriginalAssetForGraphInExpression != nullptr)
+						{
+							// save it in the list
+							if(!DependentGraphs.Contains(OriginalAssetForGraphInExpression))
+								DependentGraphs.Add(OriginalAssetForGraphInExpression);
+
+							// recursively gather graphs for all GraphExpressions encountered
+							OriginalAssetForGraphInExpression->GatherAllDependentGraphs(DependentGraphs);
+						}
+					}
+				}
+		});
+}
+bool UTextureGraph::IsDependent(const UTextureGraph* InTextureGraph) const
+{
+	
+	// check if we're trying to assign our own TextureGraph to this expression
+	
+	if (this->GetOutermostObject() == InTextureGraph->GetOutermostObject())
+	{
+		return true;
+	}
+			
+	// check for cyclic dependency
+	if (CheckCyclicDependency(InTextureGraph))
+	{
+		return true;
+	}
+
+	// default to false (Not dependent)
+	return false;
+}
+
 void UTextureGraph::Construct(FString InName)
 {
 	// On the first new texture script we set the engine in run mode
diff --git a/Engine/Plugins/Experimental/TextureGraph/Source/TextureGraph/Public/Expressions/Input/TG_Expression_Graph.h b/Engine/Plugins/Experimental/TextureGraph/Source/TextureGraph/Public/Expressions/Input/TG_Expression_Graph.h
--- a/Engine/Plugins/Experimental/TextureGraph/Source/TextureGraph/Public/Expressions/Input/TG_Expression_Graph.h
+++ b/Engine/Plugins/Experimental/TextureGraph/Source/TextureGraph/Public/Expressions/Input/TG_Expression_Graph.h
@@ -23,6 +23,7 @@ public:
 	virtual void Evaluate(FTG_EvaluationContext* InContext) override;
 
 protected:
+	
 	mutable TArray<FTG_Id> InParamIds;
 	mutable TArray<FTG_Id> OutParamIds;
 
@@ -42,6 +43,9 @@ class TEXTUREGRAPH_API UTG_Expression_TextureGraph : public UTG_Expression_Graph
 	GENERATED_BODY()
 public:
 
+	UTG_Expression_TextureGraph();
+	virtual ~UTG_Expression_TextureGraph();
+	
 #if WITH_EDITOR
 	virtual void PostEditChangeProperty(FPropertyChangedEvent& PropertyChangedEvent) override;
 #endif
@@ -51,6 +55,7 @@ public:
 
 	UPROPERTY(EditAnywhere, Setter, Category = NoCategory, meta = (TGType = "TG_Setting"))
 	TObjectPtr<UTextureGraph>	TextureGraph;				// The TextureGraph to use for the sub-graph
+
 	void SetTextureGraph(UTextureGraph* InTextureGraph);
 
 	virtual bool CanHandleAsset(UObject* Asset) override;
@@ -62,10 +67,14 @@ protected:
 
 	virtual UTG_Graph* GetGraph() const override;
 
+	void OnTextureGraphPreSave(UObject* Object, FObjectPreSaveContext SaveContext);
 private:
 	UPROPERTY(Transient)
 	TObjectPtr<UTG_Graph>	RuntimeGraph;		// The runtime graph instance to use for the sub-graph
 
+	FDelegateHandle PreSaveHandle;				// delegate handle for global ObjectPreSave event 
+	
+	bool CheckDependencies(const UTextureGraph* InTextureGraph) const;
 	void SetTextureGraphInternal(UTextureGraph* InTextureGraph);
 
 public:
diff --git a/Engine/Plugins/Experimental/TextureGraph/Source/TextureGraph/Public/TG_Graph.h b/Engine/Plugins/Experimental/TextureGraph/Source/TextureGraph/Public/TG_Graph.h
--- a/Engine/Plugins/Experimental/TextureGraph/Source/TextureGraph/Public/TG_Graph.h
+++ b/Engine/Plugins/Experimental/TextureGraph/Source/TextureGraph/Public/TG_Graph.h
@@ -283,6 +283,7 @@ protected:
 	void EvalInOutPins() const;
 	void EvalTraverseOrder() const;
 
+	bool IsDependentInternal(const UTG_Graph* SourceGraph, TArray<UTG_Graph*>& DependentGraphs);
 
 #if WITH_EDITORONLY_DATA
 	// Extra data to hold information that is useful only in editor (like comments)
diff --git a/Engine/Plugins/Experimental/TextureGraph/Source/TextureGraph/Public/TextureGraph.h b/Engine/Plugins/Experimental/TextureGraph/Source/TextureGraph/Public/TextureGraph.h
--- a/Engine/Plugins/Experimental/TextureGraph/Source/TextureGraph/Public/TextureGraph.h
+++ b/Engine/Plugins/Experimental/TextureGraph/Source/TextureGraph/Public/TextureGraph.h
@@ -21,8 +21,13 @@ protected:
 	UPROPERTY(EditAnywhere, Category = NoCategory)
 	TObjectPtr<UTG_OutputSettingsSet> OutputSettingsSet;
 
+	bool CheckCyclicDependency(const UTextureGraph* InTextureGraph) const;
+	void GatherAllDependentGraphs(TArray<UTextureGraph*>& DependentGraphs) const;
+	
 public:
 
+	bool IsDependent(const UTextureGraph* TextureGraph) const;
+	
 	// Construct the script giving it its name Initialize to a default one output script
 	virtual void Construct(FString Name);
 
diff --git a/Engine/Plugins/Experimental/TextureGraph/Source/TextureGraphEditor/Private/EdGraph/TG_EditorErrorReporter.cpp b/Engine/Plugins/Experimental/TextureGraph/Source/TextureGraphEditor/Private/EdGraph/TG_EditorErrorReporter.cpp
--- a/Engine/Plugins/Experimental/TextureGraph/Source/TextureGraphEditor/Private/EdGraph/TG_EditorErrorReporter.cpp
+++ b/Engine/Plugins/Experimental/TextureGraph/Source/TextureGraphEditor/Private/EdGraph/TG_EditorErrorReporter.cpp
@@ -2,10 +2,10 @@
 
 #include "TG_EditorErrorReporter.h"
 
-#include "SGraphNode.h"
 #include "TG_Node.h"
 #include "EdGraph/TG_EdGraph.h"
 #include "EdGraph/TG_EdGraphNode.h"
+#include "Misc/MessageDialog.h"
 
 FTextureGraphErrorReport FTG_EditorErrorReporter::ReportLog(int32 ErrorId, const FString& ErrorMsg, UObject* ReferenceObj /*= nullptr*/)
 {
@@ -33,6 +33,14 @@ FTextureGraphErrorReport FTG_EditorErrorReporter::Report(int32 ErrorId, const FS
 			}
 		}
 	}
+	else if (ReferenceObj == nullptr)
+	{
+		// if no node provided we do a popup on the Editor
+		// Editor->
+		const FText Message = FText::FromString(ErrorMsg);
+		EAppMsgCategory MsgCategory = GetMsgAppCategoryFromEMessageSeverity(ErrorType);
+		FMessageDialog::Open(MsgCategory, EAppMsgType::Ok, Message);
+	}
 	TArray<FTextureGraphErrorReport>& ErrorEntry = CompilationErrors.FindOrAdd(ErrorId);
 
 	// check to see if this error on this object is already reported
@@ -84,3 +92,18 @@ void FTG_EditorErrorReporter::Clear()
 	FTextureGraphErrorReporter::Clear();
 }
 
+EAppMsgCategory FTG_EditorErrorReporter::GetMsgAppCategoryFromEMessageSeverity(EMessageSeverity::Type ErrorType)
+{
+	switch (ErrorType)
+	{
+	case EMessageSeverity::Info:
+	default:
+		return EAppMsgCategory::Info;
+
+	case EMessageSeverity::Warning:
+		return EAppMsgCategory::Warning;
+		
+	case EMessageSeverity::Error:
+   		return EAppMsgCategory::Error;
+	}
+}
\ No newline at end of file
diff --git a/Engine/Plugins/Experimental/TextureGraph/Source/TextureGraphEditor/Private/EdGraph/TG_EditorErrorReporter.h b/Engine/Plugins/Experimental/TextureGraph/Source/TextureGraphEditor/Private/EdGraph/TG_EditorErrorReporter.h
--- a/Engine/Plugins/Experimental/TextureGraph/Source/TextureGraphEditor/Private/EdGraph/TG_EditorErrorReporter.h
+++ b/Engine/Plugins/Experimental/TextureGraph/Source/TextureGraphEditor/Private/EdGraph/TG_EditorErrorReporter.h
@@ -24,8 +24,9 @@ public:
 	virtual void Clear() override;
 
 private:
+	static EAppMsgCategory GetMsgAppCategoryFromEMessageSeverity(EMessageSeverity::Type ErrorType);
+	
 	FTextureGraphErrorReport Report(int32 ErrorId, const FString& ErrorMsg, UObject* ReferenceObj, EMessageSeverity::Type ErrorType);
 
 	FTG_Editor* Editor;
 };
-
diff --git a/Engine/Plugins/Experimental/TextureGraph/Source/TextureGraphEngine/TextureGraphErrorReporter.cpp b/Engine/Plugins/Experimental/TextureGraph/Source/TextureGraphEngine/TextureGraphErrorReporter.cpp
--- a/Engine/Plugins/Experimental/TextureGraph/Source/TextureGraphEngine/TextureGraphErrorReporter.cpp
+++ b/Engine/Plugins/Experimental/TextureGraph/Source/TextureGraphEngine/TextureGraphErrorReporter.cpp
@@ -1,5 +1,7 @@
 // Copyright Epic Games, Inc. All Rights Reserved.
 #include "TextureGraphErrorReporter.h"
+
+#include "TextureGraphEngine.h"
 #include "Logging/LogMacros.h"
 
 DEFINE_LOG_CATEGORY(LogTextureGraphError);
@@ -14,14 +16,22 @@ FTextureGraphErrorReport FTextureGraphErrorReporter::ReportLog(int32 ErrorId, co
 FTextureGraphErrorReport FTextureGraphErrorReporter::ReportWarning(int32 ErrorId, const FString& ErrorMsg, UObject* ReferenceObj /*= nullptr*/)
 {
 	FTextureGraphErrorReport Report{ErrorId, ErrorMsg, ReferenceObj};
-	UE_LOG(LogTextureGraphError, Warning, TEXT("ErrorReporter: %s"), *Report.GetFormattedMessage());
+	// Don't throw Warning assert when in test mode
+	if (!TextureGraphEngine::IsTestMode())
+	{
+		UE_LOG(LogTextureGraphError, Warning, TEXT("ErrorReporter: %s"), *Report.GetFormattedMessage());
+	}
 	return Report; 
 }
 
 FTextureGraphErrorReport FTextureGraphErrorReporter::ReportError(int32 ErrorId, const FString& ErrorMsg, UObject* ReferenceObj /*= nullptr*/)
 {
 	FTextureGraphErrorReport Report{ErrorId, ErrorMsg, ReferenceObj};
-	UE_LOG(LogTextureGraphError, Error, TEXT("ErrorReporter: %s"), *Report.GetFormattedMessage());
+	// Don't throw Error assert when in test mode
+	if (!TextureGraphEngine::IsTestMode())
+	{
+		UE_LOG(LogTextureGraphError, Error, TEXT("ErrorReporter: %s"), *Report.GetFormattedMessage());
+	}
 	return Report; 
 }
 
