Fixed Texture Graph node action sometime do not work on click but work on enter
Fixed the typo with settings node tooltip
node Input style changes

diff --git a/Engine/Plugins/Experimental/TextureGraph/Source/TextureGraph/Public/Expressions/Input/TG_Expression_OutputSettings.h b/Engine/Plugins/Experimental/TextureGraph/Source/TextureGraph/Public/Expressions/Input/TG_Expression_OutputSettings.h
--- a/Engine/Plugins/Experimental/TextureGraph/Source/TextureGraph/Public/Expressions/Input/TG_Expression_OutputSettings.h
+++ b/Engine/Plugins/Experimental/TextureGraph/Source/TextureGraph/Public/Expressions/Input/TG_Expression_OutputSettings.h
@@ -36,7 +36,7 @@ public:
 	virtual FTG_Name GetDefaultName() const override { return TEXT("Settings"); }
 	virtual void SetTitleName(FName NewName) override;
 	virtual FName GetTitleName() const override;
-	virtual FText GetTooltipText() const override { return FText::FromString(TEXT("Makes an the settings for the output.")); }
+	virtual FText GetTooltipText() const override { return FText::FromString(TEXT("Makes the settings for the output.")); }
 
 private:
 	FTG_OutputSettings PreviousInput;
diff --git a/Engine/Plugins/Experimental/TextureGraph/Source/TextureGraphEditor/Private/STG_Palette.cpp b/Engine/Plugins/Experimental/TextureGraph/Source/TextureGraphEditor/Private/STG_Palette.cpp
--- a/Engine/Plugins/Experimental/TextureGraph/Source/TextureGraphEditor/Private/STG_Palette.cpp
+++ b/Engine/Plugins/Experimental/TextureGraph/Source/TextureGraphEditor/Private/STG_Palette.cpp
@@ -26,6 +26,7 @@ void STG_PaletteItem::Construct(const FArguments& InArgs, FCreateWidgetForAction
 
 	TSharedPtr<FEdGraphSchemaAction> GraphAction = InCreateData->Action;
 	ActionPtr = InCreateData->Action;
+	MouseButtonDownDelegate = InCreateData->MouseButtonDownDelegate;
 
 	const FSlateBrush* IconBrush = GetIconBrush();
 	FSlateColor IconColor = FSlateColor::UseForeground();
@@ -118,6 +119,16 @@ FText STG_PaletteItem::GetItemTooltip() const
 	return ActionPtr.Pin()->GetTooltipDescription();
 }
 
+FReply STG_PaletteItem::OnMouseButtonDown(const FGeometry& MyGeometry, const FPointerEvent& MouseEvent)
+{
+	if (MouseButtonDownDelegate.IsBound() && MouseButtonDownDelegate.Execute(ActionPtr))
+	{
+		return FReply::Handled();
+	}
+
+	return FReply::Unhandled();
+}
+
 ////////////////////////////////////////////////////////////////////
 
 void STG_PaletteTileItem::Construct(const FArguments& InArgs, FCreateWidgetForActionData* const InCreateData)
diff --git a/Engine/Plugins/Experimental/TextureGraph/Source/TextureGraphEditor/Private/STG_Palette.h b/Engine/Plugins/Experimental/TextureGraph/Source/TextureGraphEditor/Private/STG_Palette.h
--- a/Engine/Plugins/Experimental/TextureGraph/Source/TextureGraphEditor/Private/STG_Palette.h
+++ b/Engine/Plugins/Experimental/TextureGraph/Source/TextureGraphEditor/Private/STG_Palette.h
@@ -25,7 +25,9 @@ public:
 	SLATE_BEGIN_ARGS(STG_PaletteItem) {};
 	SLATE_END_ARGS()
 
-		void Construct(const FArguments& InArgs, FCreateWidgetForActionData* const InCreateData);
+	void Construct(const FArguments& InArgs, FCreateWidgetForActionData* const InCreateData);
+
+	virtual FReply OnMouseButtonDown(const FGeometry& MyGeometry, const FPointerEvent& MouseEvent) override;
 
 protected:
 	virtual FString GetIconBrushName();
@@ -37,6 +39,9 @@ protected:
 	const FSlateBrush* GetIconBrush();
 
 	virtual FText GetItemTooltip() const override;
+
+	/** Delegate executed when mouse button goes down */
+	FCreateWidgetMouseButtonDown MouseButtonDownDelegate;
 };
 
 /** Widget for displaying a single item  */
diff --git a/Engine/Plugins/Experimental/TextureGraph/Source/TextureGraphEditor/Private/TG_Style.cpp b/Engine/Plugins/Experimental/TextureGraph/Source/TextureGraphEditor/Private/TG_Style.cpp
--- a/Engine/Plugins/Experimental/TextureGraph/Source/TextureGraphEditor/Private/TG_Style.cpp
+++ b/Engine/Plugins/Experimental/TextureGraph/Source/TextureGraphEditor/Private/TG_Style.cpp
@@ -13,6 +13,7 @@
 #include "HAL/PlatformFileManager.h"
 #include "GenericPlatform/GenericPlatformFile.h"
 #include "HAL/FileManagerGeneric.h"
+#include "Styling/StyleColors.h"
 
 #define FONT(...) FSlateFontInfo(FCoreStyle::GetDefaultFont(), __VA_ARGS__)
 
@@ -123,10 +124,12 @@ FTG_Style::FTG_Style() : FSlateStyleSet("TG_Style")
 
 		Set(FName(PropertyName), new IMAGE_BRUSH_SVG(Path, Icon16x16));
 	}
-	FLinearColor NoSpillColor(1, 1, 1, 1.0);
-	int BodyRadius = 10;
-	int NodeHeaderRadius = 7;
-	int PalleteRadius = 4;
+	
+	const FLinearColor NoSpillColor(1, 1, 1, 1.0);
+	const int BodyRadius = 10;
+	const int NodeHeaderRadius = 7;
+	const int PalleteRadius = 4;
+	const int NodeTitleEdiboxRadius = 4;
 
 	Set("TG.Graph.Node.BodyBackground", new FSlateRoundedBoxBrush(NoSpillColor, BodyRadius));
 	Set("TG.Graph.Node.BodyBorder", new FSlateRoundedBoxBrush(NoSpillColor, BodyRadius));
@@ -148,14 +151,25 @@ FTG_Style::FTG_Style() : FSlateStyleSet("TG_Style")
 		.SetShadowColorAndOpacity(FLinearColor(0.f, 0.f, 0.f, 0.7f))
 	);
 
-	FTextBlockStyle GraphNodeTitleText = FTextBlockStyle(/*NormalText*/)
+	FTextBlockStyle GraphNodeTitleText = FTextBlockStyle()
 		.SetColorAndOpacity(FLinearColor(1.0, 1.0,1.0))
-		.SetFont(FCoreStyle::GetDefaultFontStyle("Normal", 12));
+		.SetFont(FCoreStyle::GetDefaultFontStyle("Bold", 12))
+		.SetShadowOffset(FVector2D::ZeroVector)
+		.SetShadowColorAndOpacity(FLinearColor(0.f, 0.f, 0.f, 0.7f));
 	Set("TG.Graph.Node.Title", GraphNodeTitleText);
 
 	FEditableTextBoxStyle GraphActionNodeTitleEditableText = FEditableTextBoxStyle()
-		.SetTextStyle(FAppStyle::Get().GetWidgetStyle<FTextBlockStyle>("NormalText"))
-		.SetFont(GraphNodeTitleText.Font);
+		.SetFont(NormalText.Font)
+		.SetForegroundColor(FStyleColors::Input)
+		.SetBackgroundImageNormal(FSlateRoundedBoxBrush(FStyleColors::Foreground, NodeTitleEdiboxRadius, FStyleColors::Secondary, 1.0f))
+		.SetBackgroundImageHovered(FSlateRoundedBoxBrush(FStyleColors::Foreground, NodeTitleEdiboxRadius, FStyleColors::Hover, 1.0f))
+		.SetBackgroundImageFocused(FSlateRoundedBoxBrush(FStyleColors::Foreground, NodeTitleEdiboxRadius, FStyleColors::Primary, 1.0f))
+		.SetBackgroundImageReadOnly(FSlateRoundedBoxBrush(FStyleColors::Header, NodeTitleEdiboxRadius, FStyleColors::InputOutline, 1.0f))
+		.SetForegroundColor(FStyleColors::Background)
+		.SetBackgroundColor(FStyleColors::White)
+		.SetReadOnlyForegroundColor(FStyleColors::Foreground)
+		.SetFocusedForegroundColor(FStyleColors::Background);
+
 	Set("TG.Graph.Node.NodeTitleEditableText", GraphActionNodeTitleEditableText);
 
 	Set("TG.Graph.Node.NodeTitleInlineEditableText", FInlineEditableTextBlockStyle()
