Drag Drop 3D mesh on viewport doesn't work

diff --git a/Engine/Plugins/Experimental/TextureGraph/Source/TextureGraphEditor/Private/STG_EditorViewport.cpp b/Engine/Plugins/Experimental/TextureGraph/Source/TextureGraphEditor/Private/STG_EditorViewport.cpp
--- a/Engine/Plugins/Experimental/TextureGraph/Source/TextureGraphEditor/Private/STG_EditorViewport.cpp
+++ b/Engine/Plugins/Experimental/TextureGraph/Source/TextureGraphEditor/Private/STG_EditorViewport.cpp
@@ -325,8 +325,8 @@ TSharedRef<FEditorViewportClient> STG_EditorViewport::MakeEditorViewportClient()
 	EditorViewportClient = MakeShareable(new FTG_EditorViewportClient(TG_EditorPtr, *PreviewScene.Get(), SharedThis(this)));
 	UAssetViewerSettings::Get()->OnAssetViewerSettingsChanged().AddRaw(this, &STG_EditorViewport::OnAssetViewerSettingsChanged);
 	EditorViewportClient->SetViewLocation(FVector::ZeroVector);
-	EditorViewportClient->SetViewRotation(FRotator(-15.0f, -90.0f, 0.0f));
-	EditorViewportClient->SetViewLocationForOrbiting(FVector::ZeroVector);
+	EditorViewportClient->SetViewRotation(FRotator(-25.0f, -135.0f, 0.0f));
+	EditorViewportClient->SetViewLocationForOrbiting(FVector::ZeroVector, 500);
 	EditorViewportClient->bSetListenerPosition = false;
 	EditorViewportClient->EngineShowFlags.EnableAdvancedFeatures();
 	EditorViewportClient->EngineShowFlags.SetLighting(true);
@@ -729,4 +729,11 @@ void STG_EditorViewport::GenerateRenderModesList()
 			RenderModesList.Add(MaterialMappingInfo.MaterialInput);
 		}
 	}
-}
\ No newline at end of file
+}
+
+FReply STG_EditorViewport::OnDrop(const FGeometry& MyGeometry, const FDragDropEvent& DragDropEvent)
+{
+	// assign dropped asset as preview mesh
+	OnSetPreviewMeshFromSelection();
+	return SAssetEditorViewport::OnDrop(MyGeometry, DragDropEvent);
+}
diff --git a/Engine/Plugins/Experimental/TextureGraph/Source/TextureGraphEditor/Private/STG_EditorViewport.h b/Engine/Plugins/Experimental/TextureGraph/Source/TextureGraphEditor/Private/STG_EditorViewport.h
--- a/Engine/Plugins/Experimental/TextureGraph/Source/TextureGraphEditor/Private/STG_EditorViewport.h
+++ b/Engine/Plugins/Experimental/TextureGraph/Source/TextureGraphEditor/Private/STG_EditorViewport.h
@@ -52,6 +52,7 @@ private:
 protected:
 	void									GenerateRenderModesList();
 	
+	virtual FReply							OnDrop( const FGeometry& MyGeometry, const FDragDropEvent& DragDropEvent ) override;
 public:
 	void									Construct(const FArguments& InArgs);
 
