Material Crash

diff --git a/Engine/Plugins/Experimental/TextureGraph/Shaders/TextureGraphMaterialShader.usf b/Engine/Plugins/Experimental/TextureGraph/Shaders/TextureGraphMaterialShader.usf
--- a/Engine/Plugins/Experimental/TextureGraph/Shaders/TextureGraphMaterialShader.usf
+++ b/Engine/Plugins/Experimental/TextureGraph/Shaders/TextureGraphMaterialShader.usf
@@ -1,4 +1,7 @@
 // Copyright Epic Games, Inc. All Rights Reserved.
+#undef SCENE_TEXTURES_DISABLED
+#define SCENE_TEXTURES_DISABLED 1
+
 #include "/Engine/Public/Platform.ush"
 #include "/Engine/Private/Common.ush"
 #include "/Engine/Generated/Material.ush"
