Fixed issue: Disconnecting material id from texture node does not update the details tab

Fix/How: Resetting the lists when no material id is available

diff --git a/Engine/Plugins/Experimental/TextureGraph/Source/TextureGraph/Private/Expressions/Utilities/TG_Expression_MaterialID.cpp b/Engine/Plugins/Experimental/TextureGraph/Source/TextureGraph/Private/Expressions/Utilities/TG_Expression_MaterialID.cpp
--- a/Engine/Plugins/Experimental/TextureGraph/Source/TextureGraph/Private/Expressions/Utilities/TG_Expression_MaterialID.cpp
+++ b/Engine/Plugins/Experimental/TextureGraph/Source/TextureGraph/Private/Expressions/Utilities/TG_Expression_MaterialID.cpp
@@ -10,21 +10,26 @@ void UTG_Expression_MaterialID::Evaluate(FTG_EvaluationContext* InContext)
 {
 	Super::Evaluate(InContext);
 
-	if(!MaterialIDMap)
+	if(MaterialIDMap)
 	{
-		MaterialIDMap = FTG_Texture::GetBlack();
-	}
-	
-	T_ExtractMaterialIds::Create(InContext->Cycle, MaterialIDMap, MaterialIDInfoCollection, MaterialIDMaskInfos, ActiveColorsCount, InContext->TargetId);
+		T_ExtractMaterialIds::Create(InContext->Cycle, MaterialIDMap, MaterialIDInfoCollection, MaterialIDMaskInfos, ActiveColorsCount, InContext->TargetId);
 
-	if(ActiveColorsCount > 0)
-	{
-		// Generate the mask based on active colors.
-		Output = T_MaterialIDMask::Create(InContext->Cycle, MaterialIDMap, ActiveColors, ActiveColorsCount, Output.GetBufferDescriptor(), InContext->TargetId);
+		if(ActiveColorsCount > 0)
+		{
+			// Generate the mask based on active colors.
+			Output = T_MaterialIDMask::Create(InContext->Cycle, MaterialIDMap, ActiveColors, ActiveColorsCount, Output.GetBufferDescriptor(), InContext->TargetId);
+		}
+		else
+		{
+			Output = TextureHelper::GetBlack();
+		}
 	}
 	else
 	{
 		Output = TextureHelper::GetBlack();
+		MaterialIDInfoCollection.Infos.Empty();
+		MaterialIDMaskInfos.Empty();
+		ActiveColorsCount = 0;
 	}
 }
 
diff --git a/Engine/Plugins/Experimental/TextureGraph/Source/TextureGraphEngine/Transform/Expressions/T_ExtractMaterialIds.cpp b/Engine/Plugins/Experimental/TextureGraph/Source/TextureGraphEngine/Transform/Expressions/T_ExtractMaterialIds.cpp
--- a/Engine/Plugins/Experimental/TextureGraph/Source/TextureGraphEngine/Transform/Expressions/T_ExtractMaterialIds.cpp
+++ b/Engine/Plugins/Experimental/TextureGraph/Source/TextureGraphEngine/Transform/Expressions/T_ExtractMaterialIds.cpp
@@ -587,6 +587,12 @@ cti::continuable<int32> Job_ExtractMaterialIds::ExecAsync(ENamedThreads::Type Ex
 					}
 					
 				}
+				else
+				{
+					MaterialIDCollection.Infos.Empty();
+					MaterialIDMaskInfo.Empty();
+					ActiveColorsCount = 0;
+				}
 			}
 			
 			EndNative();
