Sporadic crash fix when tweaking a slider.

diff --git a/Engine/Plugins/Experimental/TextureGraph/Source/TextureGraph/Public/TG_GraphEvaluation.h b/Engine/Plugins/Experimental/TextureGraph/Source/TextureGraph/Public/TG_GraphEvaluation.h
--- a/Engine/Plugins/Experimental/TextureGraph/Source/TextureGraph/Public/TG_GraphEvaluation.h
+++ b/Engine/Plugins/Experimental/TextureGraph/Source/TextureGraph/Public/TG_GraphEvaluation.h
@@ -72,7 +72,11 @@ struct TEXTUREGRAPH_API FTG_EvaluationContext
 
 	// For debug, let's log the evaluation call sequence while it happens
 	bool				bDoLog = true;
-	bool				Tweaking = false;
+
+	FORCEINLINE bool	IsTweaking() const
+	{
+		return Cycle ? Cycle->IsTweaking() : false;
+	}
 };
 
 class UTG_Pin;
diff --git a/Engine/Plugins/Experimental/TextureGraph/Source/TextureGraphEditor/Private/EdGraph/STG_EditorGraphNode.cpp b/Engine/Plugins/Experimental/TextureGraph/Source/TextureGraphEditor/Private/EdGraph/STG_EditorGraphNode.cpp
--- a/Engine/Plugins/Experimental/TextureGraph/Source/TextureGraphEditor/Private/EdGraph/STG_EditorGraphNode.cpp
+++ b/Engine/Plugins/Experimental/TextureGraph/Source/TextureGraphEditor/Private/EdGraph/STG_EditorGraphNode.cpp
@@ -961,9 +961,8 @@ void STG_EditorGraphNode::OnNodeReconstruct()
 
 void STG_EditorGraphNode::OnNodePostEvaluate(const FTG_EvaluationContext* InContext)
 {
-	if (!InContext->Tweaking)
+	if (!InContext->IsTweaking())
 		UpdateThumbnail(InContext);
-	
 }
 
 void STG_EditorGraphNode::OnPinSelectionChanged(UEdGraphPin* EdPin)
diff --git a/Engine/Plugins/Experimental/TextureGraph/Source/TextureGraphEditor/Private/STG_TextureDetails.cpp b/Engine/Plugins/Experimental/TextureGraph/Source/TextureGraphEditor/Private/STG_TextureDetails.cpp
--- a/Engine/Plugins/Experimental/TextureGraph/Source/TextureGraphEditor/Private/STG_TextureDetails.cpp
+++ b/Engine/Plugins/Experimental/TextureGraph/Source/TextureGraphEditor/Private/STG_TextureDetails.cpp
@@ -143,6 +143,7 @@ void STG_TextureDetails::ClearHistogramWidgets()
 
 void STG_TextureDetails::CalculateHistogram(BlobPtr InBlob, UTextureGraph* InTextureGraph)
 {
+	/// If the given blob is invalid then we clear the histogram widget
 	if (!InBlob)
 	{
 		ClearHistogramWidgets();
@@ -150,38 +151,44 @@ void STG_TextureDetails::CalculateHistogram(BlobPtr InBlob, UTextureGraph* InTex
 		return;
 	}
 
-	InBlob->OnFinalise()
-	.then([this, InTextureGraph, InBlob]() mutable
-	{
-		T_TextureHistogram::Create(InTextureGraph, std::static_pointer_cast<TiledBlob>(InBlob), 0);
+	/// If it's a transient blob then we don't calculate histogram, but retain the existing one
+	if (InBlob->IsTransient())
+		return;
 
-		return InBlob->GetHistogram()->OnFinalise();
-	})
-	.then([this, InBlob]() mutable
-	{
-		if (DoesSharedInstanceExist())
+	InBlob->OnFinalise()
+		.then([this, InTextureGraph, InBlob]() mutable
 		{
-			TiledBlobPtr FinalizedHistogram = std::static_pointer_cast<TiledBlob>(InBlob->GetHistogram());
-			if (HistogramBlobWidgetR.IsValid())
+			if (!InBlob->IsTransient())
 			{
-				HistogramBlobWidgetR->Update(FinalizedHistogram);
+				T_TextureHistogram::Create(InTextureGraph, std::static_pointer_cast<TiledBlob>(InBlob), 0);
+				return InBlob->GetHistogram()->OnFinalise();
 			}
-
-			if (HistogramBlobWidgetG.IsValid())
+			return (AsyncBlobResultPtr)(cti::make_ready_continuable<const Blob*>(nullptr));
+		})
+		.then([this, InBlob]() mutable
+		{
+			if (DoesSharedInstanceExist() && !InBlob->IsTransient())
 			{
-				HistogramBlobWidgetG->Update(FinalizedHistogram);
-			}
+				TiledBlobPtr FinalizedHistogram = std::static_pointer_cast<TiledBlob>(InBlob->GetHistogram());
+				if (HistogramBlobWidgetR.IsValid())
+				{
+					HistogramBlobWidgetR->Update(FinalizedHistogram);
+				}
+
+				if (HistogramBlobWidgetG.IsValid())
+				{
+					HistogramBlobWidgetG->Update(FinalizedHistogram);
+				}
 		
-			if (HistogramBlobWidgetB.IsValid())
-			{
-				HistogramBlobWidgetB->Update(FinalizedHistogram);
-			}
+				if (HistogramBlobWidgetB.IsValid())
+				{
+					HistogramBlobWidgetB->Update(FinalizedHistogram);
+				}
 		
-			if (HistogramBlobWidgetLuma.IsValid())
-			{
-				HistogramBlobWidgetLuma->Update(FinalizedHistogram);
+				if (HistogramBlobWidgetLuma.IsValid())
+				{
+					HistogramBlobWidgetLuma->Update(FinalizedHistogram);
+				}
 			}
-		}
-		
-	});
+		});
 }
diff --git a/Engine/Plugins/Experimental/TextureGraph/Source/TextureGraphEditor/Private/TG_Editor.cpp b/Engine/Plugins/Experimental/TextureGraph/Source/TextureGraphEditor/Private/TG_Editor.cpp
--- a/Engine/Plugins/Experimental/TextureGraph/Source/TextureGraphEditor/Private/TG_Editor.cpp
+++ b/Engine/Plugins/Experimental/TextureGraph/Source/TextureGraphEditor/Private/TG_Editor.cpp
@@ -563,11 +563,11 @@ void FTG_Editor::OnRunGraph_Clicked()
 	}
 }
 
-void FTG_Editor::OnRenderingDone(UMixInterface* TextureGraph)
+void FTG_Editor::OnRenderingDone(UMixInterface* TextureGraph, const FInvalidationDetails* Details)
 {
 	if (TextureGraph != nullptr)
 	{
-		RefreshSelectionPreview(GraphEditorWidget->GetSelectedNodes());
+		RefreshSelectionPreview(GraphEditorWidget->GetSelectedNodes(), Details);
 	}
 }
 
@@ -1106,7 +1106,7 @@ void FTG_Editor::OnSelectedNodesChanged(const TSet<class UObject*>& NewSelection
 	FocusDetailsPanel();
 
 	RefreshPreviewViewport();
-	RefreshSelectionPreview(NewSelection);
+	RefreshSelectionPreview(NewSelection, nullptr);
 }
 
 void FTG_Editor::OnNodeDoubleClicked(UEdGraphNode* Node)
@@ -1768,7 +1768,7 @@ void FTG_Editor::RefreshPreviewViewport()
 	//TODO: refresh the viewport here
 }
 
-void FTG_Editor::RefreshSelectionPreview(const TSet<class UObject*>& NewSelection)
+void FTG_Editor::RefreshSelectionPreview(const TSet<class UObject*>& NewSelection, const FInvalidationDetails* Details)
 {
 	TArray<UTG_EdGraphNode*> NodesForSelectionPreview;
 	for (TSet<class UObject*>::TConstIterator SetIt(NewSelection); SetIt; ++SetIt)
@@ -1780,8 +1780,8 @@ void FTG_Editor::RefreshSelectionPreview(const TSet<class UObject*>& NewSelectio
 	}
 
 	//TODO: Handle multiple selection case
-	UTG_EdGraphNode* selectedNode = NodesForSelectionPreview.Num() > 0 ? NodesForSelectionPreview[0] : nullptr;
-	SelectionPreview->OnSelectionChanged(selectedNode);
+	UTG_EdGraphNode* SelectedNode = NodesForSelectionPreview.Num() > 0 ? NodesForSelectionPreview[0] : nullptr;
+	SelectionPreview->OnSelectionChanged(SelectedNode);
 }
 
 void FTG_Editor::SetMesh(class UMeshComponent* InPreviewMesh, class UWorld* InWorld)
diff --git a/Engine/Plugins/Experimental/TextureGraph/Source/TextureGraphEditor/Private/TG_Editor.h b/Engine/Plugins/Experimental/TextureGraph/Source/TextureGraphEditor/Private/TG_Editor.h
--- a/Engine/Plugins/Experimental/TextureGraph/Source/TextureGraphEditor/Private/TG_Editor.h
+++ b/Engine/Plugins/Experimental/TextureGraph/Source/TextureGraphEditor/Private/TG_Editor.h
@@ -103,7 +103,7 @@ public:
 	void											RefreshPreviewViewport();
 	
 	/** Called to update the selection view */
-	void											RefreshSelectionPreview(const TSet<class UObject*>& NewSelection);
+	void											RefreshSelectionPreview(const TSet<class UObject*>& NewSelection, const FInvalidationDetails* Details);
 
 	void											SetMesh(class UMeshComponent* InPreviewMesh, class UWorld* InWorld) override;
 	bool 											SetPreviewAsset(UObject* InAsset);
@@ -160,7 +160,7 @@ private:
 	ECheckBoxState									HandleAutoRunActionIsChecked() const;
 	void											HandleAutoRunActionExecute(ECheckBoxState InNewState);
 
-	void											OnRenderingDone(UMixInterface* script);
+	void											OnRenderingDone(UMixInterface* TextureGraph, const FInvalidationDetails* Details);
 	void											OnViewportSettingsChanged();
 	void											OnMaterialMappingChanged();
 	/** Log Graph in console */
diff --git a/Engine/Plugins/Experimental/TextureGraph/Source/TextureGraphEngine/2D/Tex.cpp b/Engine/Plugins/Experimental/TextureGraph/Source/TextureGraphEngine/2D/Tex.cpp
--- a/Engine/Plugins/Experimental/TextureGraph/Source/TextureGraphEngine/2D/Tex.cpp
+++ b/Engine/Plugins/Experimental/TextureGraph/Source/TextureGraphEngine/2D/Tex.cpp
@@ -348,7 +348,8 @@ void Tex::InitRT(bool ForceFloat /* = false */)
 		}
 	}
 
-	const FName Name = *FString::Printf(TEXT("%s [RT]"), *Desc.Name);
+	FString DescName = Desc.Name.Left(std::min(NAME_SIZE / 2, Desc.Name.Len()));
+	const FName Name = *FString::Printf(TEXT("%s [RT]"), *DescName);
 	const auto Package = Util::GetRenderTargetPackage();
 	const FName UniqueName = MakeUniqueObjectName(Package, UTextureRenderTarget2D::StaticClass(), Name);
 	RT = NewObject<UTextureRenderTarget2D>(Package, UniqueName);
@@ -358,11 +359,10 @@ void Tex::InitRT(bool ForceFloat /* = false */)
 	check(RT);
 	
 	RT->ClearColor = Desc.ClearColor;
-	
 
-	if (Desc.bMipMaps && Desc.bAutoGenerateMipMaps) {
+	if (Desc.bMipMaps && Desc.bAutoGenerateMipMaps) 
+	{
 		RT->bAutoGenerateMips = true;
-		//_rt->MipGenSettings = TMGS_Sharpen4; /// Not for RT ?
 	}
 	
 	RT->MipsSamplerFilter = Filter;
@@ -373,15 +373,9 @@ void Tex::InitRT(bool ForceFloat /* = false */)
 	RT->OverrideFormat = Desc.Format;
 	RT->SizeX = Desc.Width;
 	RT->SizeY = Desc.Height;
-	//UE_LOG(LogTexture, Log, TEXT("Mips number: %d"), _rt->GetNumMips());
 	RT->LODBias = 0;
 
-	//_rt->InitCustomFormat(_desc.Width, _desc.Height, _desc.Format, !_desc.sRGB);
-	//UE_LOG(LogTexture, Log, TEXT("Creating render target: %s"), *_desc.Name);
-
 	RT->UpdateResource();
-	//_rt->UpdateResourceImmediate(false);
-
 }
 
 FString Tex::GetReferencerName() const
diff --git a/Engine/Plugins/Experimental/TextureGraph/Source/TextureGraphEngine/Data/TiledBlob.cpp b/Engine/Plugins/Experimental/TextureGraph/Source/TextureGraphEngine/Data/TiledBlob.cpp
--- a/Engine/Plugins/Experimental/TextureGraph/Source/TextureGraphEngine/Data/TiledBlob.cpp
+++ b/Engine/Plugins/Experimental/TextureGraph/Source/TextureGraphEngine/Data/TiledBlob.cpp
@@ -130,7 +130,13 @@ void TiledBlob::FinaliseFrom(const Blob* RHS)
 
 	Tiles = RHSTiled->Tiles;
 	HashValue = RHSTiled->HashValue;
+
+	bool bIsTransient = Desc.bIsTransient;
 	Desc = RHSTiled->Desc;
+
+	/// Retain the transient information as cached tiled blob could be transient and this one isn't or vice-versa
+	Desc.bIsTransient = bIsTransient; 
+
 	bReady = RHSTiled->bReady;
 	bTiledTarget = RHSTiled->bTiledTarget;
 	SingleBlob = RHSTiled->SingleBlob;
@@ -1090,37 +1096,21 @@ void TiledBlob_Promise::FinaliseNow(bool bNoCalcHash, CHashPtr FixedHash)
 	UpdateLinkedBlobs(true);
 }
 
-AsyncBufferResultPtr TiledBlob_Promise::FinaliseFrom(std::shared_ptr<TiledBlob_Promise> RHS)
+void TiledBlob_Promise::FinaliseFrom(TiledBlobPtr RHS)
 {
-	if (RHS->bMakeSingleBlob && !bMakeSingleBlob)
-	{
-		check(RHS->IsFinalised());
-		check(RHS->Rows() == 1 && RHS->Cols() == 1);
-
-		auto RHSBuffer = RHS->GetTile(0, 0)->GetBufferRef();
-		check(RHSBuffer);
-		Buffer = RHSBuffer;
-
-		check(Buffer);
-		check(!bMakeSingleBlob && Desc.Width >= Tiles.Rows() && Desc.Height >= Tiles.Cols());
-		bIsFinalised = true;
-		
-		return TileBuffer(Buffer, Tiles);
-	}
-
-	*this = *RHS.get();
-	return cti::make_ready_continuable<BufferResultPtr>(std::make_shared<BufferResult>());
+	FinaliseFrom(std::static_pointer_cast<Blob>(RHS).get());
 }
 
 void TiledBlob_Promise::FinaliseFrom(const Blob* RHS)
 {
 	check(RHS->IsTiled());
 	const TiledBlob* RHSTiled = static_cast<const TiledBlob*>(RHS);
-	check(RHSTiled->IsPromise());
 
-	const TiledBlob_Promise* RHSTiledPromise = static_cast<const TiledBlob_Promise*>(RHS);
-
-	bMakeSingleBlob = RHSTiledPromise->bMakeSingleBlob;
+	if (RHSTiled->IsPromise())
+	{
+		const TiledBlob_Promise* RHSTiledPromise = static_cast<const TiledBlob_Promise*>(RHS);
+		bMakeSingleBlob = RHSTiledPromise->bMakeSingleBlob;
+	}
 
 	TiledBlob::FinaliseFrom(RHS);
 }
diff --git a/Engine/Plugins/Experimental/TextureGraph/Source/TextureGraphEngine/Data/TiledBlob.h b/Engine/Plugins/Experimental/TextureGraph/Source/TextureGraphEngine/Data/TiledBlob.h
--- a/Engine/Plugins/Experimental/TextureGraph/Source/TextureGraphEngine/Data/TiledBlob.h
+++ b/Engine/Plugins/Experimental/TextureGraph/Source/TextureGraphEngine/Data/TiledBlob.h
@@ -182,7 +182,7 @@ public:
 
 	virtual void					FinaliseNow(bool bNoCalcHash, CHashPtr FixedHash) override;
 	virtual AsyncBufferResultPtr	Finalise(bool bNoCalcHash, CHashPtr FixedHash) override;
-	AsyncBufferResultPtr			FinaliseFrom(std::shared_ptr<TiledBlob_Promise> RHS);
+	void							FinaliseFrom(TiledBlobPtr RHS);
 
 	void							ResetForReplay(); /// For debug purpose, reset the state of the tile as a promise to NOT finalised, increment the replayCount
 
diff --git a/Engine/Plugins/Experimental/TextureGraph/Source/TextureGraphEngine/Device/Device.cpp b/Engine/Plugins/Experimental/TextureGraph/Source/TextureGraphEngine/Device/Device.cpp
--- a/Engine/Plugins/Experimental/TextureGraph/Source/TextureGraphEngine/Device/Device.cpp
+++ b/Engine/Plugins/Experimental/TextureGraph/Source/TextureGraphEngine/Device/Device.cpp
@@ -417,8 +417,11 @@ void Device::Collect(DeviceBuffer* Buffer)
 	Cache.RemoveThreadSafe(Hash->Value());
 
 	/// Make sure that its not marked for transfer if its getting collected
-	DeviceTransferServicePtr transferSvc = TextureGraphEngine::GetScheduler()->GetDeviceTransferService().lock();
-	transferSvc->AbortTransfer(Hash->Value());
+	DeviceTransferServicePtr TransferSvc = TextureGraphEngine::GetScheduler()->GetDeviceTransferService().lock();
+	if (TransferSvc)
+	{
+		TransferSvc->AbortTransfer(Hash->Value());
+	}
 
 	BlobRef blob = TextureGraphEngine::GetBlobber()->Find(Hash->Value());
 
@@ -613,7 +616,7 @@ AsyncJobResultPtr Device::UpdateDeviceTransfers()
 							/// If we find a blob against it then we start the transfer
 							if (BlobObj)
 							{
-								//transferSvc->QueueTransfer(blob, nextDevice);
+								//TransferSvc->QueueTransfer(blob, nextDevice);
 								//auto promise = std::move(blob->TransferTo(nextDevice));
 								Promises.emplace_back(BlobObj->TransferTo(NextDevice));
 								bShouldCollect = false;
diff --git a/Engine/Plugins/Experimental/TextureGraph/Source/TextureGraphEngine/Job/Job.cpp b/Engine/Plugins/Experimental/TextureGraph/Source/TextureGraphEngine/Job/Job.cpp
--- a/Engine/Plugins/Experimental/TextureGraph/Source/TextureGraphEngine/Job/Job.cpp
+++ b/Engine/Plugins/Experimental/TextureGraph/Source/TextureGraphEngine/Job/Job.cpp
@@ -392,7 +392,19 @@ bool Job::CheckCached()
 	{
 		check(ExistingResult->IsTiled());
 
-		Result = std::static_pointer_cast<TiledBlob>(ExistingResult.get());
+		BufferDescriptor ResultDesc = ExistingResult->GetDescriptor();
+
+		if (ResultDesc.bIsTransient == IsDiscard())
+		{
+			Result = ExistingResult;
+		}
+		else
+		{
+			ResultOrg = std::make_shared<TiledBlob_Promise>(ExistingResult->GetDescriptor(), ExistingResult->Rows(), ExistingResult->Cols(), ExistingResult->Hash());
+			ResultOrg->FinaliseFrom(ExistingResult.get());
+			Result = TiledBlobRef(std::static_pointer_cast<TiledBlob>(ResultOrg), true, false);
+		}
+
 		FinalJobResult = std::make_shared<JobResult>(GetResultRef(), nullptr);
 		bIsCulled = true;
 
@@ -461,7 +473,9 @@ TiledBlobPtr Job::InitResult(FString InNewName, const BufferDescriptor* InDesire
 	check(!ResultDesc.Name.IsEmpty());
 
 	if (ResultDesc.IsLateBound())
+	{
 		return InitLateBoundResult(NewName, ResultDesc, NumInputBlobs);
+	}
 	else if (DesiredResultDesc && ArgsDescCombined.IsLateBound() && DesiredResultDesc->IsAutoSize() && ArgsDescCombined.IsAutoSize())
 	{
 		ResultDesc.Format = BufferFormat::LateBound;
@@ -497,6 +511,7 @@ TiledBlobPtr Job::InitResult(FString InNewName, const BufferDescriptor* InDesire
 
 	check(Result->GetDescriptor().Width >= Result->Rows());
 	check(Result->GetDescriptor().Height >= Result->Cols());
+	check(Result->IsTransient() == IsDiscard());
 
 	return Result;
 }
@@ -505,11 +520,15 @@ TiledBlobPtr Job::InitLateBoundResult(FString NewName, BufferDescriptor DesiredD
 {
 	/// We can only check for a cached result for late bound blobs if there are no input blobs
 	if (!NumInputBlobs && CheckCached())
+	{
+		check(Result->IsTransient() == IsDiscard());
 		return Result;
+	}
 
 	BufferDescriptor Desc = DesiredDesc;
 	Desc.Name = NewName;
 	Desc.Format = BufferFormat::LateBound;
+	Desc.bIsTransient = IsDiscard();
 
 	CHashPtr TempHash = Hash(); ///std::make_shared<CHash>(Desc.HashValue(), false);
 	check(MixObj);
@@ -518,8 +537,8 @@ TiledBlobPtr Job::InitLateBoundResult(FString NewName, BufferDescriptor DesiredD
 	int32 NumTilesY = MixObj->GetNumYTiles();
 
 	ResultOrg = std::make_shared<TiledBlob_Promise>(Desc, NumTilesX, NumTilesY, TempHash);
-
 	Result = TiledBlobRef(std::static_pointer_cast<TiledBlob>(ResultOrg), true, false);
+	check(Result->IsTransient() == IsDiscard());
 
 	return Result;
 }
@@ -531,16 +550,16 @@ AsyncJobResultPtr Job::FinaliseTiles(JobRunInfo InRunInfo)
 
 	std::vector<std::decay_t<AsyncBufferResultPtr>, std::allocator<std::decay_t<AsyncBufferResultPtr>>> Promises;
 
-	bool discard = RunInfo.Cycle->GetDetails().IsDiscard();
+	bool bIsDiscard = RunInfo.Cycle->GetDetails().IsDiscard();
 
-	/// only flush if these have not been marked to discard
+	/// only flush if these have not been marked to bIsDiscard
 	AsyncJobResultPtr FinalPromise = cti::make_ready_continuable(FinalJobResult);
 	double flushStartTime = Util::Time();
 
 	check(Result->IsPromise());
 	TiledBlob_PromisePtr result = GetResultPromise();
 
-	if (!discard)
+	if (!bIsDiscard)
 	{
 		if (CanHandleTiles())
 		{
@@ -803,7 +822,7 @@ AsyncPrepareResult Job::PrepareTargets(JobBatch* Batch)
 		}
 	}
 
-	/// Set transient
+	/// Transient flag should've been set correctly earlier on
 	if (bIsDiscard)
 		Result->SetTransient();
 
diff --git a/Engine/Plugins/Experimental/TextureGraph/Source/TextureGraphEngine/Job/Scheduler.cpp b/Engine/Plugins/Experimental/TextureGraph/Source/TextureGraphEngine/Job/Scheduler.cpp
--- a/Engine/Plugins/Experimental/TextureGraph/Source/TextureGraphEngine/Job/Scheduler.cpp
+++ b/Engine/Plugins/Experimental/TextureGraph/Source/TextureGraphEngine/Job/Scheduler.cpp
@@ -32,7 +32,7 @@ Scheduler::Scheduler()
 	HistogramServiceObj = HistogramSPtr;
 
 	AddIdleService(BlobHasherSPtr);
-	AddIdleService(DeviceTransferSPtr);
+	//AddIdleService(DeviceTransferSPtr);
 	AddIdleService(ThumbnailsSPtr);
 	AddIdleService(MipmapSPtr);
 	AddIdleService(MinmaxSPtr);
@@ -83,7 +83,7 @@ void Scheduler::Update(float dt)
 		}
 	}
 
-	JobBatchPtr batchToRun = nullptr;
+	JobBatchPtr BatchToRun = nullptr;
 	{
 		FScopeLock Lock(&BatchMutex);
 		if (Batches.size())
@@ -91,7 +91,7 @@ void Scheduler::Update(float dt)
 			/// TODO: Enable this at some point
 			if (Batches.size() == 1)
 			{
-				batchToRun = Batches.front();
+				BatchToRun = Batches.front();
 				Batches.pop_front();
 			}
 			else
@@ -101,51 +101,51 @@ void Scheduler::Update(float dt)
 				/// OK here, we're going to merge batches belonging to the same mix together
 				
 				// figure out the prioritized mix from the Batch
-				batchToRun = Batches.front();
-				UMixInterface* currentBatchMix = batchToRun->GetCycle()->GetMix();
+				BatchToRun = Batches.front();
+				UMixInterface* currentBatchMix = BatchToRun->GetCycle()->GetMix();
 
 				// find the latest Batch of required mix
-				auto iter = Batches.begin();
-				JobBatchPtr lastMixbatch = batchToRun;
+				auto Iter = Batches.begin();
+				JobBatchPtr lastMixbatch = BatchToRun;
 
-				while (iter != Batches.end())
+				while (Iter != Batches.end())
 				{
-					JobBatchPtr futureBatch = *iter;
+					JobBatchPtr futureBatch = *Iter;
 
 					if (futureBatch->GetCycle()->GetMix() == currentBatchMix)
 					{
-						lastMixbatch = *iter;
+						lastMixbatch = *Iter;
 					}
 
-					iter++;
+					Iter++;
 				}
 
 				// we get latest Batch to run for the mix
-				batchToRun = lastMixbatch;
+				BatchToRun = lastMixbatch;
 
 				// Merge invalidation details of the mix from multiple batches.
 				// Remove the batches from the list after merging.
-				iter = Batches.begin();
-				while (iter != Batches.end())
+				Iter = Batches.begin();
+				while (Iter != Batches.end())
 				{
-					JobBatchPtr futureBatch = *iter;
+					JobBatchPtr futureBatch = *Iter;
 
 					if (futureBatch->GetCycle()->GetMix() == currentBatchMix)
 					{
-						batchToRun->GetCycle()->MergeDetails(futureBatch->GetCycle()->GetDetails());
-						iter = Batches.erase(iter);
+						BatchToRun->GetCycle()->MergeDetails(futureBatch->GetCycle()->GetDetails());
+						Iter = Batches.erase(Iter);
 					}
 					else
-						iter++;
+						Iter++;
 				}
 			}
 		}
 	}
 
-	if (batchToRun)
+	if (BatchToRun)
 	{
 		FScopeLock Lock(&CurrentBatchMutex);
-		CurrentBatch = batchToRun;
+		CurrentBatch = BatchToRun;
 		CurrentBatchStartTime = Util::Time();
 
 		UE_LOG(LogBatch, Log, TEXT("Next Batch: %llu [Prev: %llu]"), CurrentBatch->GetBatchId(), PreviousBatch ? PreviousBatch->GetBatchId() : 0);
@@ -186,39 +186,11 @@ void Scheduler::Update(float dt)
 
 					ObserverSource->BatchDone(CurrentBatch);
 
+					/// Mark the current batch as done
+					CurrentBatch->GetCycle()->GetDetails().BroadcastOnDone();
+
 					UE_LOG(LogBatch, Log, TEXT("Scheduler Observer::BatchDone finished for Batch: %llu"), PreviousBatch ? PreviousBatch->GetBatchId() : -1);
-				})
-				.then([this]()
-					{
-						// Trigger cycle->invalidateDetails.DoneCallBackLambda()
-						
-						//Adding logs to debug and safty checks here to avoid the crash
-						//Crash was happening in cases when we do not have a valid Cycle or Mix
-						//This could happen when we open a script run it and then delete it and then try to run another script.
-						// 
-						//TODO: Investigate why we need to call the PreviousBatch->BroadcastOnDone instead of calling Currentbatch->BroadcastOnDone
-						//As we are setting the PreviousBatch = CurrentBatch in the callback of Exec function and that gets called from EndBatch after the
-						//execution of this .then block 
-						//looks like a bug though 
-						//fix would be to remove this .then block and call PreviousBatch->GetCycle()->GetDetails().BroadcastOnDone(); in the callback of Exec function
-						if (!PreviousBatch)
-						{
-							UE_LOG(LogBatch, Log, TEXT("Scheduler Previous Batch is null"));
-							return;
-						}
-						if (!PreviousBatch->GetCycle())
-						{
-							UE_LOG(LogBatch, Warning, TEXT("Scheduler Previous Batch : %llu Cycle is null") , PreviousBatch->GetBatchId());
-							return;
-						}
-						if (!PreviousBatch->GetCycle()->GetDetails().Mix.IsValid())
-						{
-							UE_LOG(LogBatch, Warning, TEXT("Scheduler Previous Batch : %llu Cycle Details Mix is null"), PreviousBatch->GetBatchId());
-							return;
-						}
-
-						PreviousBatch->GetCycle()->GetDetails().BroadcastOnDone();
-					});
+				});
 	}
 	else
 	{
@@ -264,6 +236,14 @@ AsyncJobResultPtr Scheduler::UpdateIdleBatch(size_t Index)
 
 void Scheduler::UpdateIdle()
 {
+	{
+		FScopeLock Lock(&CurrentBatchMutex);
+
+		/// If there's an active batch then we don't run this at all
+		if (CurrentBatch)
+			return;
+	}
+
 	/// Currently we don't do the idle update in test mode
 	if (TextureGraphEngine::IsTestMode())
 		return;
diff --git a/Engine/Plugins/Experimental/TextureGraph/Source/TextureGraphEngine/Model/Mix/MixInterface.cpp b/Engine/Plugins/Experimental/TextureGraph/Source/TextureGraphEngine/Model/Mix/MixInterface.cpp
--- a/Engine/Plugins/Experimental/TextureGraph/Source/TextureGraphEngine/Model/Mix/MixInterface.cpp
+++ b/Engine/Plugins/Experimental/TextureGraph/Source/TextureGraphEngine/Model/Mix/MixInterface.cpp
@@ -98,9 +98,9 @@ int32 UMixInterface::GetNumYTiles() const
 	return GetSettings()->GetYTiles();
 }
 
-void UMixInterface::BroadcastOnRenderingDone()
+void UMixInterface::BroadcastOnRenderingDone(const FInvalidationDetails* Details)
 {
-	OnRenderDone.ExecuteIfBound(this);
+	OnRenderDone.ExecuteIfBound(this, Details);
 }
 
 UMixSettings* UMixInterface::GetSettings() const 
diff --git a/Engine/Plugins/Experimental/TextureGraph/Source/TextureGraphEngine/Model/Mix/MixInterface.h b/Engine/Plugins/Experimental/TextureGraph/Source/TextureGraphEngine/Model/Mix/MixInterface.h
--- a/Engine/Plugins/Experimental/TextureGraph/Source/TextureGraphEngine/Model/Mix/MixInterface.h
+++ b/Engine/Plugins/Experimental/TextureGraph/Source/TextureGraphEngine/Model/Mix/MixInterface.h
@@ -65,7 +65,7 @@ public:
 	UPROPERTY(EditAnywhere, Category = Previewing, meta = (AllowedClasses = "/Script/Engine.StaticMesh,/Script/Engine.SkeletalMesh", ExactClass = "true"))
 	FSoftObjectPath						PreviewMesh;  // Adding it here because this will be the base class for TS_Script and TS_ScriptInstance.
 
-	DECLARE_DELEGATE_OneParam(FOnRenderDone, UMixInterface*);
+	DECLARE_DELEGATE_TwoParams(FOnRenderDone, UMixInterface*, const FInvalidationDetails*);
 	FOnRenderDone						OnRenderDone;
 	virtual void						PostMeshLoad();
 
@@ -97,7 +97,7 @@ public:
 	virtual void						InvalidateWithDetails(const FInvalidationDetails& Details);
 	virtual void						InvalidateAll();
 
-	virtual void						BroadcastOnRenderingDone();
+	virtual void						BroadcastOnRenderingDone(const FInvalidationDetails* Details);
 	
 	//////////////////////////////////////////////////////////////////////////
 	/// Inline functions
diff --git a/Engine/Plugins/Experimental/TextureGraph/Source/TextureGraphEngine/Model/Mix/MixUpdateCycle.cpp b/Engine/Plugins/Experimental/TextureGraph/Source/TextureGraphEngine/Model/Mix/MixUpdateCycle.cpp
--- a/Engine/Plugins/Experimental/TextureGraph/Source/TextureGraphEngine/Model/Mix/MixUpdateCycle.cpp
+++ b/Engine/Plugins/Experimental/TextureGraph/Source/TextureGraphEngine/Model/Mix/MixUpdateCycle.cpp
@@ -127,24 +127,28 @@ JobPtrW MixUpdateCycle::AddJob(int32 InTargetId, JobUPtr JobObj)
 
 	JobPtrW JobW = Batch->AddJob(std::move(JobObj));
 
-	JobRunInfo runInfo;
-	runInfo.JobScheduler = TextureGraphEngine::GetScheduler();
-	runInfo.Batch = Batch.get();
-	runInfo.Cycle = Batch->GetCycle();
-	runInfo.ThisJob = JobW;
+	JobRunInfo RunInfo;
+	RunInfo.JobScheduler = TextureGraphEngine::GetScheduler();
+	RunInfo.Batch = Batch.get();
+	RunInfo.Cycle = Batch->GetCycle();
+	RunInfo.ThisJob = JobW;
 
 	JobPtr JobS = JobW.lock();
 
 	/// If the job is fully culled, then we don't need to do anything
-	if (JobS->CheckCulled(runInfo))
+	if (JobS->CheckCulled(RunInfo))
 	{
 		UE_LOG(LogBatch, Log, TEXT("MixUpdateCycle::AddJob  Job [%llu]: IsCulled"), Batch->GetBatchId());
 	}
 
-	SceneTargetUpdatePtr target = Targets[InTargetId];
+	SceneTargetUpdatePtr Target = Targets[InTargetId];
 
 	if (JobS->GetResult())
+	{
 		JobS->GetResult()->Job() = JobW;
+		if (Details.IsDiscard())
+			JobS->GetResult()->SetTransient();
+	}
 
 	/// If the job is culled, then we remove it
 	if (JobS->IsCulled())
diff --git a/Engine/Plugins/Experimental/TextureGraph/Source/TextureGraphEngine/Model/Mix/MixUpdateCycle.h b/Engine/Plugins/Experimental/TextureGraph/Source/TextureGraphEngine/Model/Mix/MixUpdateCycle.h
--- a/Engine/Plugins/Experimental/TextureGraph/Source/TextureGraphEngine/Model/Mix/MixUpdateCycle.h
+++ b/Engine/Plugins/Experimental/TextureGraph/Source/TextureGraphEngine/Model/Mix/MixUpdateCycle.h
@@ -121,6 +121,7 @@ public:
 	FORCEINLINE const JobBatch*		GetBatch() const { return Batch.get(); }
 	FORCEINLINE size_t				GetNumTargets() const { return Targets.size(); }
 	FORCEINLINE UMixInterface*		GetMix() { return MixObj.IsValid() ? MixObj.Get() : nullptr; }
+	FORCEINLINE bool				IsTweaking() const { return Details.bTweaking; }
 };
 
 typedef std::shared_ptr<MixUpdateCycle>		MixUpdateCyclePtr;
diff --git a/Engine/Plugins/Experimental/TextureGraph/Source/TextureGraphEngine/Model/ModelObject.cpp b/Engine/Plugins/Experimental/TextureGraph/Source/TextureGraphEngine/Model/ModelObject.cpp
--- a/Engine/Plugins/Experimental/TextureGraph/Source/TextureGraphEngine/Model/ModelObject.cpp
+++ b/Engine/Plugins/Experimental/TextureGraph/Source/TextureGraphEngine/Model/ModelObject.cpp
@@ -12,12 +12,12 @@ FInvalidationDetails::FInvalidationDetails(const TWeakObjectPtr<UMixInterface> m
 
 FInvalidationDetails& FInvalidationDetails::All()
 {
-	auto c = OnDone;
+	auto OnDoneCurrent = OnDone;
 	*this = FInvalidationDetails(Mix);
 	if (OnDone.IsBound())
 	{
 		/// Create a fresh copy
-		this->OnDone = c;
+		OnDone = OnDoneCurrent;
 	}
 
 	bRender = true;
@@ -52,25 +52,25 @@ CHashPtr FInvalidationDetails::Hash() const
 	return HashValue;
 }
 
-FInvalidationDetails& FInvalidationDetails::Merge(const FInvalidationDetails& details)
+FInvalidationDetails& FInvalidationDetails::Merge(const FInvalidationDetails& Details)
 {
-	FInvalidationDetails curr = *this;
-	*this = details;
+	FInvalidationDetails CurrentDetails = *this;
+	*this = Details;
 
 	// Current contains all the accumulated delegates
 	// the parameter Details is expected to NOT contain any delegates in _onDoneMerged
 	check(!OnDoneMergedInternal.IsBound());
-	OnDoneMergedInternal = curr.OnDoneMergedInternal;
-	// Accumulate potentially curr.OnDone
-	if (curr.OnDone.IsBound())
-		OnDoneMergedInternal.Add(curr.OnDone);
-	// At this point this contains the delegate OnDone from details
+	OnDoneMergedInternal = CurrentDetails.OnDoneMergedInternal;
+	// Accumulate potentially CurrentDetails.OnDone
+	if (CurrentDetails.OnDone.IsBound())
+		OnDoneMergedInternal.Add(CurrentDetails.OnDone);
+	// At this point this contains the delegate OnDone from Details
 	// And accumulated in _internal_onDoneMerged all the previously accumulated delegates in this
 
-	bTweaking = curr.bTweaking & details.bTweaking;
-	bRender = curr.bRender | details.bRender;
+	bTweaking = CurrentDetails.bTweaking && Details.bTweaking;
+	bRender = CurrentDetails.bRender || Details.bRender;
 
-	bSelective &= curr.bSelective;
+	bSelective &= CurrentDetails.bSelective;
 
 	// reset hash to recalculate with updated values when requested
 	HashValue = nullptr;
@@ -80,10 +80,10 @@ FInvalidationDetails& FInvalidationDetails::Merge(const FInvalidationDetails& de
 
 void FInvalidationDetails::BroadcastOnDone() const
 {
-	OnDone.ExecuteIfBound();
-	OnDoneMergedInternal.Broadcast();
+	OnDone.ExecuteIfBound(this);
+	OnDoneMergedInternal.Broadcast(this);
 	check(Mix.Get());
-	Mix->BroadcastOnRenderingDone();
+	Mix->BroadcastOnRenderingDone(this);
 }
 
 bool FInvalidationDetails::IsDiscard() const
diff --git a/Engine/Plugins/Experimental/TextureGraph/Source/TextureGraphEngine/Model/ModelObject.h b/Engine/Plugins/Experimental/TextureGraph/Source/TextureGraphEngine/Model/ModelObject.h
--- a/Engine/Plugins/Experimental/TextureGraph/Source/TextureGraphEngine/Model/ModelObject.h
+++ b/Engine/Plugins/Experimental/TextureGraph/Source/TextureGraphEngine/Model/ModelObject.h
@@ -21,7 +21,7 @@ struct TEXTUREGRAPHENGINE_API FInvalidationDetails
 
 	TWeakObjectPtr<UMixInterface>		Mix = nullptr;								/// The mix that is being invalidated
 	
-	DECLARE_DELEGATE(FOnDone);
+	DECLARE_DELEGATE_OneParam(FOnDone, const FInvalidationDetails*);
 	FOnDone								OnDone;										/// Delegate allowing to provide a user callback once the invalidation is executed
 
 	bool								bReload = false;							/// Whether to reload the sources or not
@@ -76,7 +76,7 @@ struct TEXTUREGRAPHENGINE_API FInvalidationDetails
 	}
 
 private:
-	DECLARE_MULTICAST_DELEGATE(FOnDoneMulti);
+	DECLARE_MULTICAST_DELEGATE_OneParam(FOnDoneMulti, const FInvalidationDetails*);
 	FOnDoneMulti				OnDoneMergedInternal; /// "Private" multicast delegate where we merge OnDone delegates for merged details
 };
 
diff --git a/Engine/Plugins/Experimental/TextureGraph/Source/TextureGraphEngine/Transform/Layer/T_Thumbnail.cpp b/Engine/Plugins/Experimental/TextureGraph/Source/TextureGraphEngine/Transform/Layer/T_Thumbnail.cpp
--- a/Engine/Plugins/Experimental/TextureGraph/Source/TextureGraphEngine/Transform/Layer/T_Thumbnail.cpp
+++ b/Engine/Plugins/Experimental/TextureGraph/Source/TextureGraphEngine/Transform/Layer/T_Thumbnail.cpp
@@ -26,6 +26,7 @@ JobBatchPtr AddThumbJobToCycle(TiledBlobPtr InBlobToBind, JobUPtr JobPtr, UObjec
 TiledBlobRef T_Thumbnail::Bind(UMixInterface* Mix, UObject* Model, TiledBlobPtr InBlobToBind, int32 InTargetId)
 {
 	check(InBlobToBind);
+	check(!InBlobToBind->IsTransient());
 
 	UE_LOG(LogJob, VeryVerbose, TEXT("T_Thumbnail::Bind [%s]"), *InBlobToBind->Name());
 	//RenderMaterial_ThumbPtr MatThumb = TextureGraphEngine::GetMaterialManager()->CreateMaterial_Thumbnail(TEXT("T_Thumbnail"), TEXT("Util/CopyUnlit"));
diff --git a/Engine/Plugins/Experimental/TextureGraph/Source/TextureGraphEngine/Transform/Utility/T_TextureHistogram.cpp b/Engine/Plugins/Experimental/TextureGraph/Source/TextureGraphEngine/Transform/Utility/T_TextureHistogram.cpp
--- a/Engine/Plugins/Experimental/TextureGraph/Source/TextureGraphEngine/Transform/Utility/T_TextureHistogram.cpp
+++ b/Engine/Plugins/Experimental/TextureGraph/Source/TextureGraphEngine/Transform/Utility/T_TextureHistogram.cpp
@@ -46,6 +46,8 @@ TiledBlobPtr T_TextureHistogram::Create(UMixInterface* InMix, TiledBlobPtr Sourc
 	}
 
 	check(SourceTex);
+	check(!SourceTex->IsTransient());
+
 	HistogramServicePtr Service = TextureGraphEngine::GetScheduler()->GetHistogramService().lock();
 	check(Service);
 
@@ -55,15 +57,12 @@ TiledBlobPtr T_TextureHistogram::Create(UMixInterface* InMix, TiledBlobPtr Sourc
 
 	FIntVector4 SrcDimensions(SourceTex->GetWidth(), SourceTex->GetHeight(), 1, 1);
 
-	auto SrcTiles = std::make_shared<JobArg_Blob>(SourceTex, "SourceTiles");
-
 	FString Name = FString::Printf(TEXT("[%s].[%d] Histogram"), *SourceTex->DisplayName(), TargetId);
 	RenderMaterial_FXPtr Transform = T_TextureHistogram::CreateMaterial_Histogram(
 		TEXT("T_Histogram"), TEXT("Result"), PermutationVector, SrcDimensions.X, SrcDimensions.Y, 1);
 
 	JobUPtr JobObj = std::make_unique<Job>(InMix, TargetId, std::static_pointer_cast<BlobTransform>(Transform));
-	JobObj
-		->AddArg(SrcTiles);
+	JobObj->AddArg(ARG_BLOB(SourceTex, "SourceTiles"));
 
 	BufferDescriptor Desc;
 	Desc.Width = NumBins;
@@ -82,7 +81,6 @@ TiledBlobPtr T_TextureHistogram::Create(UMixInterface* InMix, TiledBlobPtr Sourc
 	//Add the job using histogram idle service
 	AddHistogramJobToCycle(Batch->GetCycle(),std::move(JobObj), TargetId, InMix);
 
-
 	if (!SourceTex->HasHistogram())
 	{
 		//setting it as the histogram of source so it is retained untill the life cycle of source blob.
