Fixed issue: All input nodes have input pin

diff --git a/Engine/Plugins/Experimental/TextureGraph/Source/TextureGraph/Public/Expressions/Input/TG_Expression_Material.h b/Engine/Plugins/Experimental/TextureGraph/Source/TextureGraph/Public/Expressions/Input/TG_Expression_Material.h
--- a/Engine/Plugins/Experimental/TextureGraph/Source/TextureGraph/Public/Expressions/Input/TG_Expression_Material.h
+++ b/Engine/Plugins/Experimental/TextureGraph/Source/TextureGraph/Public/Expressions/Input/TG_Expression_Material.h
@@ -22,7 +22,7 @@ public:
 #endif
 
 	// The input material to employ for rendering
-	UPROPERTY(EditAnywhere, BlueprintReadWrite, Setter, Category = NoCategory, meta = (TGType = "TG_Setting"))
+	UPROPERTY(EditAnywhere, BlueprintReadWrite, Setter, Category = NoCategory, meta = (TGType = "TG_Setting", TGPinNotConnectable))
 	TObjectPtr<UMaterialInterface> Material;
 	
 	void SetMaterial(UMaterialInterface* InMaterial);
diff --git a/Engine/Plugins/Experimental/TextureGraph/Source/TextureGraph/Public/Expressions/Input/TG_Expression_MaterialFunction.h b/Engine/Plugins/Experimental/TextureGraph/Source/TextureGraph/Public/Expressions/Input/TG_Expression_MaterialFunction.h
--- a/Engine/Plugins/Experimental/TextureGraph/Source/TextureGraph/Public/Expressions/Input/TG_Expression_MaterialFunction.h
+++ b/Engine/Plugins/Experimental/TextureGraph/Source/TextureGraph/Public/Expressions/Input/TG_Expression_MaterialFunction.h
@@ -21,7 +21,7 @@ public:
 #endif
 	
 	// The material function to employ for rendering
-	UPROPERTY(EditAnywhere, BlueprintReadWrite, Setter, Category = NoCategory, meta = (TGType = "TG_Setting"))
+	UPROPERTY(EditAnywhere, BlueprintReadWrite, Setter, Category = NoCategory, meta = (TGType = "TG_Setting", TGPinNotConnectable))
 	TObjectPtr<UMaterialFunctionInterface> MaterialFunction;
 	void SetMaterialFunction(UMaterialFunctionInterface* InMaterialFunction);
 
