Fixed issue: Assertion was hit when trying to create a texture node

Fix/How: Added a new meta to make the pin not connectable. Removed the pin from the texture parameter. Disabled a connection with the new node when the pin is not connectable or param.

#virtualized

diff --git a/Engine/Plugins/Experimental/TextureGraph/Source/TextureGraph/Private/Expressions/TG_Expression.cpp b/Engine/Plugins/Experimental/TextureGraph/Source/TextureGraph/Private/Expressions/TG_Expression.cpp
--- a/Engine/Plugins/Experimental/TextureGraph/Source/TextureGraph/Private/Expressions/TG_Expression.cpp
+++ b/Engine/Plugins/Experimental/TextureGraph/Source/TextureGraph/Private/Expressions/TG_Expression.cpp
@@ -63,41 +63,45 @@ FTG_Signature::FInit UTG_Expression::GetSignatureInitArgsFromClass() const
 	for (TFieldIterator<FProperty> Prop(GetClass()); Prop; ++Prop) {
 		FName PropName = FName(Prop->GetNameCPP());
 		FName PropTypeName = FName(Prop->GetCPPType());
+		uint8 MaskNotConnectable = 0;
 #if WITH_EDITORONLY_DATA
+		
 		auto PropertyTGType = Prop->GetMetaData(TEXT("TGType"));
+		bool PropertyPinNotConnectable = Prop->HasMetaData(TEXT("TGPinNotConnectable"));
+		MaskNotConnectable = PropertyPinNotConnectable ? static_cast<uint8>(ETG_Access::NotConnectableFlag) : 0;
 #else
 		FString PropertyTGType;
 #endif
-
+		
 		if (PropertyTGType.Compare(TEXT("TG_InputParam")) == 0)
 		{
 			// Add the new argument to the signature
-			SignatureInit.Arguments.Add(FTG_Argument{ PropName, PropTypeName, { ETG_Access::InParam } });
+			SignatureInit.Arguments.Add(FTG_Argument{ PropName, PropTypeName, { ETG_Access::InParam, MaskNotConnectable} });
 		}																								   
 		else if (PropertyTGType.Compare(TEXT("TG_OutputParam")) == 0)									   
 		{																									   
 			// Add the new argument to the signature													   
-			SignatureInit.Arguments.Add(FTG_Argument{ PropName, PropTypeName, { ETG_Access::OutParam } });
+			SignatureInit.Arguments.Add(FTG_Argument{ PropName, PropTypeName, { ETG_Access::OutParam, MaskNotConnectable } });
 		}																								   
 		else if (PropertyTGType.Compare(TEXT("TG_Input")) == 0)										   
 		{													   
 			// Add the new argument to the signature													   
-			SignatureInit.Arguments.Add(FTG_Argument{ PropName, PropTypeName, { ETG_Access::In } });
+			SignatureInit.Arguments.Add(FTG_Argument{ PropName, PropTypeName, { ETG_Access::In, MaskNotConnectable } });
 		}																								   
 		else if (PropertyTGType.Compare(TEXT("TG_Output")) == 0)										   
 		{																	   
 			// Add the new argument to the signature													   
-			SignatureInit.Arguments.Add(FTG_Argument{ PropName, PropTypeName, { ETG_Access::Out } });
+			SignatureInit.Arguments.Add(FTG_Argument{ PropName, PropTypeName, { ETG_Access::Out, MaskNotConnectable } });
 		}
 		else if (PropertyTGType.Compare(TEXT("TG_Setting")) == 0)
 		{
 			// Add the new argument to the signature as a Setting													   
-			SignatureInit.Arguments.Add(FTG_Argument{ PropName, PropTypeName, { ETG_Access::InSetting } });
+			SignatureInit.Arguments.Add(FTG_Argument{ PropName, PropTypeName, { ETG_Access::InSetting, MaskNotConnectable } });
 		}
 		else if (PropertyTGType.Compare(TEXT("TG_Private")) == 0)
 		{
 			// Add the new argument to the signature as a Private													   
-			SignatureInit.Arguments.Add(FTG_Argument{ PropName, PropTypeName, { ETG_Access::Private } });
+			SignatureInit.Arguments.Add(FTG_Argument{ PropName, PropTypeName, { ETG_Access::Private, MaskNotConnectable } });
 		}
 		else
 		{
diff --git a/Engine/Plugins/Experimental/TextureGraph/Source/TextureGraph/Public/Expressions/Input/TG_Expression_Texture.h b/Engine/Plugins/Experimental/TextureGraph/Source/TextureGraph/Public/Expressions/Input/TG_Expression_Texture.h
--- a/Engine/Plugins/Experimental/TextureGraph/Source/TextureGraph/Public/Expressions/Input/TG_Expression_Texture.h
+++ b/Engine/Plugins/Experimental/TextureGraph/Source/TextureGraph/Public/Expressions/Input/TG_Expression_Texture.h
@@ -27,7 +27,7 @@ public:
 	FTG_Texture Output;
 
 	// The source asset to be used to generate the Output
-	UPROPERTY(EditAnywhere, BlueprintReadWrite, Category = NoCategory, meta = (TGType = "TG_Input", PinNotConnectable = true))
+	UPROPERTY(EditAnywhere, BlueprintReadWrite, Category = NoCategory, meta = (TGType = "TG_Input", TGPinNotConnectable) )
 	TObjectPtr<UTexture> Source;
 
 	// The input texture that was loaded from the asset
diff --git a/Engine/Plugins/Experimental/TextureGraph/Source/TextureGraph/Public/TG_Signature.h b/Engine/Plugins/Experimental/TextureGraph/Source/TextureGraph/Public/TG_Signature.h
--- a/Engine/Plugins/Experimental/TextureGraph/Source/TextureGraph/Public/TG_Signature.h
+++ b/Engine/Plugins/Experimental/TextureGraph/Source/TextureGraph/Public/TG_Signature.h
@@ -46,7 +46,8 @@ struct TEXTUREGRAPH_API FTG_ArgumentType
 
 	FTG_ArgumentType() {}
 	FTG_ArgumentType(ETG_Access InAccess) : Flags(static_cast<uint8>(InAccess)) {}
-
+	FTG_ArgumentType(ETG_Access InAccess, uint8 Flags) : Flags(static_cast<uint8>(InAccess) | Flags) {}
+	
 	UPROPERTY()
 	uint32			Flags = static_cast<uint8>(ETG_Access::In);
 
diff --git a/Engine/Plugins/Experimental/TextureGraph/Source/TextureGraphEditor/Private/EdGraph/TG_EdGraphNode.cpp b/Engine/Plugins/Experimental/TextureGraph/Source/TextureGraphEditor/Private/EdGraph/TG_EdGraphNode.cpp
--- a/Engine/Plugins/Experimental/TextureGraph/Source/TextureGraphEditor/Private/EdGraph/TG_EdGraphNode.cpp
+++ b/Engine/Plugins/Experimental/TextureGraph/Source/TextureGraphEditor/Private/EdGraph/TG_EdGraphNode.cpp
@@ -530,12 +530,15 @@ void UTG_EdGraphNode::AutowireNewNode(UEdGraphPin* FromPin)
 		TObjectPtr<UTG_Pin> OtherPin = Node->GetGraph()->GetPin(OtherPinId);
 		check(OtherPin);
 
-		const FName& OtherPinName = OtherPin->GetArgumentName();
-		UEdGraphPin* ToPin = FindPinChecked(OtherPinName, bFromPinIsInput ? EEdGraphPinDirection::EGPD_Output : EEdGraphPinDirection::EGPD_Input);
-		if (ToPin && GetSchema()->TryCreateConnection(FromPin, ToPin))
+		if(!OtherPin->IsNotConnectable() && !OtherPin->IsParam())
 		{
-			// Connection succeeded
-			break;
+			const FName& OtherPinName = OtherPin->GetArgumentName();
+			UEdGraphPin* ToPin = FindPinChecked(OtherPinName, bFromPinIsInput ? EEdGraphPinDirection::EGPD_Output : EEdGraphPinDirection::EGPD_Input);
+			if (ToPin && GetSchema()->TryCreateConnection(FromPin, ToPin))
+			{
+				// Connection succeeded
+				break;
+			}	
 		}
 	}
 
