Accidental checkin fix

diff --git a/Engine/Plugins/Experimental/TextureGraph/Source/TextureGraphEngine/Job/Job.cpp b/Engine/Plugins/Experimental/TextureGraph/Source/TextureGraphEngine/Job/Job.cpp
--- a/Engine/Plugins/Experimental/TextureGraph/Source/TextureGraphEngine/Job/Job.cpp
+++ b/Engine/Plugins/Experimental/TextureGraph/Source/TextureGraphEngine/Job/Job.cpp
@@ -511,7 +511,6 @@ TiledBlobPtr Job::InitResult(FString InNewName, const BufferDescriptor* InDesire
 
 	check(Result->GetDescriptor().Width >= Result->Rows());
 	check(Result->GetDescriptor().Height >= Result->Cols());
-	check(Result->IsTransient() == IsDiscard());
 
 	return Result;
 }
@@ -521,7 +520,6 @@ TiledBlobPtr Job::InitLateBoundResult(FString NewName, BufferDescriptor DesiredD
 	/// We can only check for a cached result for late bound blobs if there are no input blobs
 	if (!NumInputBlobs && CheckCached())
 	{
-		check(Result->IsTransient() == IsDiscard());
 		return Result;
 	}
 
@@ -538,7 +536,6 @@ TiledBlobPtr Job::InitLateBoundResult(FString NewName, BufferDescriptor DesiredD
 
 	ResultOrg = std::make_shared<TiledBlob_Promise>(Desc, NumTilesX, NumTilesY, TempHash);
 	Result = TiledBlobRef(std::static_pointer_cast<TiledBlob>(ResultOrg), true, false);
-	check(Result->IsTransient() == IsDiscard());
 
 	return Result;
 }
diff --git a/Engine/Plugins/Experimental/TextureGraph/Source/TextureGraphEngine/Transform/Layer/T_Thumbnail.cpp b/Engine/Plugins/Experimental/TextureGraph/Source/TextureGraphEngine/Transform/Layer/T_Thumbnail.cpp
--- a/Engine/Plugins/Experimental/TextureGraph/Source/TextureGraphEngine/Transform/Layer/T_Thumbnail.cpp
+++ b/Engine/Plugins/Experimental/TextureGraph/Source/TextureGraphEngine/Transform/Layer/T_Thumbnail.cpp
@@ -26,7 +26,6 @@ JobBatchPtr AddThumbJobToCycle(TiledBlobPtr InBlobToBind, JobUPtr JobPtr, UObjec
 TiledBlobRef T_Thumbnail::Bind(UMixInterface* Mix, UObject* Model, TiledBlobPtr InBlobToBind, int32 InTargetId)
 {
 	check(InBlobToBind);
-	check(!InBlobToBind->IsTransient());
 
 	UE_LOG(LogJob, VeryVerbose, TEXT("T_Thumbnail::Bind [%s]"), *InBlobToBind->Name());
 	//RenderMaterial_ThumbPtr MatThumb = TextureGraphEngine::GetMaterialManager()->CreateMaterial_Thumbnail(TEXT("T_Thumbnail"), TEXT("Util/CopyUnlit"));
