fix Pow expression

Pow is now using the correct code path

diff --git a/Engine/Plugins/Experimental/TextureGraph/Source/TextureGraphEngine/Transform/Expressions/T_Maths_TwoInputs.cpp b/Engine/Plugins/Experimental/TextureGraph/Source/TextureGraphEngine/Transform/Expressions/T_Maths_TwoInputs.cpp
--- a/Engine/Plugins/Experimental/TextureGraph/Source/TextureGraphEngine/Transform/Expressions/T_Maths_TwoInputs.cpp
+++ b/Engine/Plugins/Experimental/TextureGraph/Source/TextureGraphEngine/Transform/Expressions/T_Maths_TwoInputs.cpp
@@ -96,7 +96,7 @@ TiledBlobPtr T_Maths_TwoInputs::CreateCross(MixUpdateCyclePtr Cycle, BufferDescr
 
 TiledBlobPtr T_Maths_TwoInputs::CreatePow(MixUpdateCyclePtr Cycle, BufferDescriptor DesiredOutputDesc, int32 TargetId, TiledBlobPtr Operand1, TiledBlobPtr Operand2)
 {
-	return CreateGenericMathOp<FSH_Divide>(Cycle, DesiredOutputDesc, TargetId, Operand1, Operand2, "T_Pow");
+	return CreateGenericMathOp<FSH_Pow>(Cycle, DesiredOutputDesc, TargetId, Operand1, Operand2, "T_Pow");
 }
 
 RenderMaterial_FXPtr GetIfThenElseMaterial(EIfThenElseOperator Operator, EIfThenElseType Type)
