Improve HIstogram to contain min and max values
Separate the invocation of the T_Histogram transform into 2 differnt calls, one for on service and the other in the CYlcle queue.
this will be usefull for expressions requiring the Histogram (Levels, ...)

#tests 65b19fccc14595787f92d738

diff --git a/Engine/Plugins/Experimental/TextureGraph/Shaders/Utils/Histogram_comp.usf b/Engine/Plugins/Experimental/TextureGraph/Shaders/Utils/Histogram_comp.usf
--- a/Engine/Plugins/Experimental/TextureGraph/Shaders/Utils/Histogram_comp.usf
+++ b/Engine/Plugins/Experimental/TextureGraph/Shaders/Utils/Histogram_comp.usf
@@ -37,17 +37,48 @@ void CSH_Histogram(
 
 	GroupMemoryBarrierWithGroupSync();
 
-	if(GroupIndex == 0)
+	if (GroupIndex == 0)
 	{
+		// Now only in first group, let's go through the histogram bins SEQUENTIALLY and extract the min and max count and min and max values
+		float4 MinCount = float4(MAX_flt,MAX_flt,MAX_flt,MAX_flt);
+		float4 MaxCount = float4(0,0,0,0);
+		
 		float4 MinValue = float4(MAX_flt,MAX_flt,MAX_flt,MAX_flt);
-		float4 MaxValue = float4(0,0,0,0);
+		float4 MaxValue = float4(0, 0, 0, 0);
+			
 		LOOP for(uint y = 0; y < 256 ; y++)
-		{
-			MinValue = min(Histogram[y],MinValue);
-			MaxValue = max(Histogram[y],MaxValue);
-			Result[uint2(y,0)] = float4(Histogram[y].r,Histogram[y].g,Histogram[y].b,Histogram[y].a);
+		{		
+			
+			float4 HistogramVal = float4(Histogram[y]);
+	
+			MinCount = min(HistogramVal, MinCount);
+			MaxCount = max(HistogramVal, MaxCount);
+			
+			Result[uint2(y,0)] = HistogramVal;
+			
+			if (any(HistogramVal))
+			{
+				// Assigned min value for the first HistogramVal found, never after
+				MinValue.x = (MinValue.x != MAX_flt ? MinValue.x : (HistogramVal.x == 0 ? MAX_flt : y) );
+				MinValue.y = (MinValue.y != MAX_flt ? MinValue.y : (HistogramVal.y == 0 ? MAX_flt : y) );
+				MinValue.z = (MinValue.z != MAX_flt ? MinValue.z : (HistogramVal.z == 0 ? MAX_flt : y) );
+				MinValue.w = (MinValue.w != MAX_flt ? MinValue.w : (HistogramVal.w == 0 ? MAX_flt : y) );
+				
+				// Assigned max value, always the current Y of the loop if the Histogram is not 0.
+				MaxValue.x = (HistogramVal.x == 0 ? MaxValue.x : y + 1); 
+				MaxValue.y = (HistogramVal.y == 0 ? MaxValue.y : y + 1); 
+				MaxValue.z = (HistogramVal.z == 0 ? MaxValue.z : y + 1); 
+				MaxValue.w = (HistogramVal.w == 0 ? MaxValue.w : y + 1); 
+			}			
 		}
-		Result[uint2(0,1)] = MinValue;//2nd row first pixel is the min value
-		Result[uint2(1,1)] = MaxValue;//2nd row second pixel is the max value
+		
+		MaxValue = max(MaxValue, MinValue); // Make sure Max was assigned and coherent with min
+		MinValue /= 256.0;
+		MaxValue /= 256.0;
+		
+		Result[uint2(0,1)] = MinCount;//2nd row first pixel is the min count
+		Result[uint2(1,1)] = MaxCount;//2nd row second pixel is the max count
+		Result[uint2(2,1)] = MinValue;//2nd row third pixel is the min value
+		Result[uint2(3,1)] = MaxValue;//2nd row fourth pixel is the max value
 	}
 }
diff --git a/Engine/Plugins/Experimental/TextureGraph/Source/TextureGraphEditor/Private/Customizations/STG_TextureHistogram.cpp b/Engine/Plugins/Experimental/TextureGraph/Source/TextureGraphEditor/Private/Customizations/STG_TextureHistogram.cpp
--- a/Engine/Plugins/Experimental/TextureGraph/Source/TextureGraphEditor/Private/Customizations/STG_TextureHistogram.cpp
+++ b/Engine/Plugins/Experimental/TextureGraph/Source/TextureGraphEditor/Private/Customizations/STG_TextureHistogram.cpp
@@ -37,7 +37,7 @@ void STG_TextureHistogram::SetTexture(FTG_Texture& Source, UTextureGraph* InText
 		Source.RasterBlob->OnFinalise()
 		.then([this, InTextureGraph, Source]() mutable
 		{
-			T_TextureHistogram::Create(InTextureGraph, std::static_pointer_cast<TiledBlob>(Source.RasterBlob), 0);
+			T_TextureHistogram::CreateOnService(InTextureGraph, std::static_pointer_cast<TiledBlob>(Source.RasterBlob), 0);
 			return  Source->GetHistogram()->OnFinalise();
 		})
 		.then([this, Source]() mutable
diff --git a/Engine/Plugins/Experimental/TextureGraph/Source/TextureGraphEditor/Private/STG_TextureDetails.cpp b/Engine/Plugins/Experimental/TextureGraph/Source/TextureGraphEditor/Private/STG_TextureDetails.cpp
--- a/Engine/Plugins/Experimental/TextureGraph/Source/TextureGraphEditor/Private/STG_TextureDetails.cpp
+++ b/Engine/Plugins/Experimental/TextureGraph/Source/TextureGraphEditor/Private/STG_TextureDetails.cpp
@@ -160,7 +160,7 @@ void STG_TextureDetails::CalculateHistogram(BlobPtr InBlob, UTextureGraph* InTex
 		{
 			if (!InBlob->IsTransient())
 			{
-				T_TextureHistogram::Create(InTextureGraph, std::static_pointer_cast<TiledBlob>(InBlob), 0);
+				T_TextureHistogram::CreateOnService(InTextureGraph, std::static_pointer_cast<TiledBlob>(InBlob), 0);
 				return InBlob->GetHistogram()->OnFinalise();
 			}
 			return (AsyncBlobResultPtr)(cti::make_ready_continuable<const Blob*>(nullptr));
diff --git a/Engine/Plugins/Experimental/TextureGraph/Source/TextureGraphEngine/Transform/Utility/T_TextureHistogram.cpp b/Engine/Plugins/Experimental/TextureGraph/Source/TextureGraphEngine/Transform/Utility/T_TextureHistogram.cpp
--- a/Engine/Plugins/Experimental/TextureGraph/Source/TextureGraphEngine/Transform/Utility/T_TextureHistogram.cpp
+++ b/Engine/Plugins/Experimental/TextureGraph/Source/TextureGraphEngine/Transform/Utility/T_TextureHistogram.cpp
@@ -37,21 +37,10 @@ RenderMaterial_FXPtr T_TextureHistogram::CreateMaterial_Histogram(FString Name,
 	return std::make_shared<RenderMaterial_FX>(Name, std::static_pointer_cast<FxMaterial>(Mat));
 }
 
-TiledBlobPtr T_TextureHistogram::Create(UMixInterface* InMix, TiledBlobPtr SourceTex, int32 TargetId)
-{
-	if (SourceTex->HasHistogram())
-	{
-		TiledBlobPtr SourceHistogram = std::static_pointer_cast<TiledBlob>(SourceTex->GetHistogram());
-		return SourceHistogram;
-	}
-
-	check(SourceTex);
-	check(!SourceTex->IsTransient());
-
-	HistogramServicePtr Service = TextureGraphEngine::GetScheduler()->GetHistogramService().lock();
-	check(Service);
 
-	JobBatchPtr Batch = Service->GetOrCreateNewBatch(InMix);
+TiledBlobPtr T_TextureHistogram::CreateJobAndResult(JobUPtr& OutJob, MixUpdateCyclePtr Cycle, TiledBlobPtr SourceTex, int32 TargetId)
+{
+	check(!SourceTex->HasHistogram())
 
 	CSH_Histogram::FPermutationDomain PermutationVector;
 
@@ -61,8 +50,8 @@ TiledBlobPtr T_TextureHistogram::Create(UMixInterface* InMix, TiledBlobPtr Sourc
 	RenderMaterial_FXPtr Transform = T_TextureHistogram::CreateMaterial_Histogram(
 		TEXT("T_Histogram"), TEXT("Result"), PermutationVector, SrcDimensions.X, SrcDimensions.Y, 1);
 
-	JobUPtr JobObj = std::make_unique<Job>(InMix, TargetId, std::static_pointer_cast<BlobTransform>(Transform));
-	JobObj->AddArg(ARG_BLOB(SourceTex, "SourceTiles"));
+	OutJob = std::make_unique<Job>(Cycle->GetMix(), TargetId, std::static_pointer_cast<BlobTransform>(Transform));
+	OutJob->AddArg(ARG_BLOB(SourceTex, "SourceTiles"));
 
 	BufferDescriptor Desc;
 	Desc.Width = NumBins;
@@ -72,34 +61,60 @@ TiledBlobPtr T_TextureHistogram::Create(UMixInterface* InMix, TiledBlobPtr Sourc
 	Desc.Name = FString::Printf(TEXT("Histogram - %s"), *SourceTex->Name());
 	Desc.AllowUAV();
 
-	JobObj->SetTiled(false);
+	OutJob->SetTiled(false);
 
-	TiledBlobPtr Result = JobObj->InitResult(Name, &Desc,1,1);
+	TiledBlobPtr Result = OutJob->InitResult(Name, &Desc, 1, 1);
 
 	Result->MakeSingleBlob();
 
-	//Add the job using histogram idle service
-	AddHistogramJobToCycle(Batch->GetCycle(),std::move(JobObj), TargetId, InMix);
 
 	if (!SourceTex->HasHistogram())
 	{
 		//setting it as the histogram of source so it is retained untill the life cycle of source blob.
 		SourceTex->SetHistogram(Result);
 	}
+	return Result;
+}
+
+TiledBlobPtr T_TextureHistogram::Create(MixUpdateCyclePtr Cycle, TiledBlobPtr SourceTex, int32 TargetId)
+{
+	if (SourceTex->HasHistogram())
+	{
+		TiledBlobPtr SourceHistogram = std::static_pointer_cast<TiledBlob>(SourceTex->GetHistogram());
+		return SourceHistogram;
+	}
+
+	JobUPtr JobObj;
+	TiledBlobPtr Result = CreateJobAndResult(JobObj, Cycle, SourceTex, TargetId);
+
+	Cycle->AddJob(TargetId, std::move(JobObj));
 
 	return Result;
 }
 
-void T_TextureHistogram::AddHistogramJobToCycle(MixUpdateCyclePtr Cycle, JobUPtr Job, int32 TargetId,UMixInterface* Mix)
+TiledBlobPtr T_TextureHistogram::CreateOnService(UMixInterface* InMix, TiledBlobPtr SourceTex, int32 TargetId)
 {
+	if (SourceTex->HasHistogram())
+	{
+		TiledBlobPtr SourceHistogram = std::static_pointer_cast<TiledBlob>(SourceTex->GetHistogram());
+		return SourceHistogram;
+	}
+
+	check(InMix);
+	check(SourceTex);
+	check(!SourceTex->IsTransient());
+
 	HistogramServicePtr Service = TextureGraphEngine::GetScheduler()->GetHistogramService().lock();
+	check(Service);
 
-	if (!Service)
-		return;
+	JobBatchPtr Batch = Service->GetOrCreateNewBatch(InMix);
 
-	//TODO get rid of mix here and use Null Mix instead
-	UMixInterface* mix = Job->GetMix();
-	check(mix);
+	JobUPtr JobObj;
+	TiledBlobPtr Result = CreateJobAndResult(JobObj, Batch->GetCycle(), SourceTex, TargetId);
 
-	Service->AddHistogramJob(Cycle,std::move(Job), TargetId,Mix);
+	//Add the job using histogram idle service
+	//TODO: get rid of mix here and use Null Mix instead
+	Service->AddHistogramJob(Batch->GetCycle(), std::move(JobObj), TargetId, InMix);
+
+	return Result;
 }
diff --git a/Engine/Plugins/Experimental/TextureGraph/Source/TextureGraphEngine/Transform/Utility/T_TextureHistogram.h b/Engine/Plugins/Experimental/TextureGraph/Source/TextureGraphEngine/Transform/Utility/T_TextureHistogram.h
--- a/Engine/Plugins/Experimental/TextureGraph/Source/TextureGraphEngine/Transform/Utility/T_TextureHistogram.h
+++ b/Engine/Plugins/Experimental/TextureGraph/Source/TextureGraphEngine/Transform/Utility/T_TextureHistogram.h
@@ -67,10 +67,11 @@ public:
 	T_TextureHistogram();
 	~T_TextureHistogram();
 
-	static TiledBlobPtr	Create(UMixInterface* InMix, TiledBlobPtr SourceTex, int32 TargetId);
-	static void AddHistogramJobToCycle(MixUpdateCyclePtr Cycle, JobUPtr Job, int32 TargetId, UMixInterface* Mix);
-	
+	static TiledBlobPtr	Create(MixUpdateCyclePtr Cycle, TiledBlobPtr SourceTex, int32 TargetId);
+
+	static TiledBlobPtr	CreateOnService(UMixInterface* InMix, TiledBlobPtr SourceTex, int32 TargetId);
+
 private:
-	static TiledBlobPtr GetMaxValueForHistogram(MixUpdateCyclePtr Cycle, TiledBlobPtr Histogram, int32 TargetId);
 	static RenderMaterial_FXPtr CreateMaterial_Histogram(FString Name, FString OutputId, const CSH_Histogram::FPermutationDomain& cmpshPermutationDomain,int NumThreadsX, int NumThreadsY, int NumThreadsZ = 1);
+	static TiledBlobPtr			CreateJobAndResult(JobUPtr& OutJob, MixUpdateCyclePtr Cycle, TiledBlobPtr Histogram, int32 TargetId);
 };
