Fixed issue: Reset option appear infront of each indices in details tab for materialID node.

#tests 65b3a4b136560a2b8059396d

diff --git a/Engine/Plugins/Experimental/TextureGraph/Source/TextureGraph/Public/Expressions/Utilities/TG_Expression_MaterialID.h b/Engine/Plugins/Experimental/TextureGraph/Source/TextureGraph/Public/Expressions/Utilities/TG_Expression_MaterialID.h
--- a/Engine/Plugins/Experimental/TextureGraph/Source/TextureGraph/Public/Expressions/Utilities/TG_Expression_MaterialID.h
+++ b/Engine/Plugins/Experimental/TextureGraph/Source/TextureGraph/Public/Expressions/Utilities/TG_Expression_MaterialID.h
@@ -21,7 +21,7 @@ public:
 	UPROPERTY(meta = (TGType = "TG_Input"))
 	FTG_Texture MaterialIDMap;
 
-	UPROPERTY(EditAnywhere, EditFixedSize, Category = NoCategory, meta = (TGType = "TG_InputParam"))
+	UPROPERTY(EditAnywhere, EditFixedSize, Category = NoCategory, meta = (TGType = "TG_InputParam", NoResetToDefault))
 	TArray<FMaterialIDMaskInfo> MaterialIDMaskInfos;
 
 	UPROPERTY()
diff --git a/Engine/Plugins/Experimental/TextureGraph/Source/TextureGraphEngine/Transform/Expressions/T_ExtractMaterialIds.h b/Engine/Plugins/Experimental/TextureGraph/Source/TextureGraphEngine/Transform/Expressions/T_ExtractMaterialIds.h
--- a/Engine/Plugins/Experimental/TextureGraph/Source/TextureGraphEngine/Transform/Expressions/T_ExtractMaterialIds.h
+++ b/Engine/Plugins/Experimental/TextureGraph/Source/TextureGraphEngine/Transform/Expressions/T_ExtractMaterialIds.h
@@ -64,7 +64,7 @@ struct FMaterialIDMaskInfo
 	UPROPERTY()
 	int32 MaterialIdReferenceId = 0;
 
-	UPROPERTY(VisibleAnywhere, Category = MaterialIDInfo)
+	UPROPERTY(VisibleAnywhere, Category = MaterialIDInfo, meta = (NoResetToDefault))
 	FColor Color = FColor::Black;
 
 	UPROPERTY(EditAnywhere, Category = MaterialIDInfo)
