TextureGraph: Texture preset and options added to output settings

diff --git a/Engine/Plugins/Experimental/TextureGraph/Source/TextureGraph/Private/Blueprint/TG_BlueprintFunctionLibrary.cpp b/Engine/Plugins/Experimental/TextureGraph/Source/TextureGraph/Private/Blueprint/TG_BlueprintFunctionLibrary.cpp
--- a/Engine/Plugins/Experimental/TextureGraph/Source/TextureGraph/Private/Blueprint/TG_BlueprintFunctionLibrary.cpp
+++ b/Engine/Plugins/Experimental/TextureGraph/Source/TextureGraph/Private/Blueprint/TG_BlueprintFunctionLibrary.cpp
@@ -233,12 +233,34 @@ FLinearColor UTG_BlueprintFunctionLibrary::GetColorParameterValue(UObject* World
 	return ParameterValue;
 }
 
-void UTG_BlueprintFunctionLibrary::SetOutputSettingsParameterValue(UObject* WorldContextObject, UTextureGraph* InTextureGraph, FName ParameterName, int Width, int Height, FString FileName /*= "None"*/, FString FolderPath /*= "None"*/, ETG_TextureFormat Format /*= ETG_TextureFormat::BGRA8*/)
+void UTG_BlueprintFunctionLibrary::SetSettingsParameterValue(UObject* WorldContextObject, UTextureGraph* InTextureGraph, FName ParameterName, int Width, int Height, 
+	FString FileName /*= "None"*/, FString FolderPath /*= "None"*/, ETG_TextureFormat Format /*= ETG_TextureFormat::BGRA8*/, ETG_TexturePresetType TextureType /*= ETG_TexturePresetType::None*/,
+	TextureCompressionSettings Compression /*= TextureCompressionSettings::TC_Default*/, TextureGroup LodGroup /*= TextureGroup::TEXTUREGROUP_World*/, bool bSRGB /*= false*/)
 {
 	if (UWorld* World = GEngine->GetWorldFromContextObject(WorldContextObject, EGetWorldErrorMode::LogAndReturnNull))
 	{
-		auto bFoundParameter = false;
-		if (InTextureGraph)
+		bool bFoundParameter = false;
+		bool bSizeError = false;
+		bool bPathError = false;
+		bool bNameError = false;
+
+		const FString FunctionName = "SetSettingsParameterValue";
+
+		if (Width <= 0 || Height <= 0)
+		{
+			bSizeError = true;
+		}
+		if (FolderPath.IsEmpty() || FolderPath == "None")
+		{
+			bPathError = true;
+		}
+		if (FileName.IsEmpty() || FolderPath == "None")
+		{
+			bNameError = true;
+		}
+
+
+		if (InTextureGraph && !bSizeError && !bPathError && !bNameError)
 		{
 			auto PinParam = InTextureGraph->Graph()->FindParamPin(ParameterName);
 			if (PinParam)
@@ -247,22 +269,37 @@ void UTG_BlueprintFunctionLibrary::SetOutputSettingsParameterValue(UObject* Worl
 				if (ExpressionPtr)
 				{
 					FTG_OutputSettings ParameterValue;
-					ParameterValue.Set(Width, Height, FileName, FolderPath, Format);
+					ParameterValue.Set(Width, Height, FileName, FolderPath, Format, TextureType, Compression, LodGroup, bSRGB);
 					ExpressionPtr->Settings = ParameterValue;
 					PinParam->EditSelfVar()->EditAs<FTG_OutputSettings>() = ParameterValue;
 					bFoundParameter = true;
 				}
 			}
+
+			if (!bFoundParameter)
+			{
+				AddParamWarning(ParameterName, InTextureGraph, FunctionName);
+			}
 		}
 
-		if (!bFoundParameter)
+		if (bSizeError)
 		{
-			AddParamWarning(ParameterName, InTextureGraph, "SetOutputSettingsParameterValue");
+			AddError(InTextureGraph, FunctionName, "Invalid size try to set non zero size");
+		}
+
+		if (bPathError)
+		{
+			AddError(InTextureGraph, FunctionName, "Invalid path try to set a valid path , path connot be empty or none");
+		}
+
+		if (bNameError)
+		{
+			AddError(InTextureGraph, FunctionName, "Invalid file name , file name cannot be empty or None");
 		}
 	}
 }
 
-FTG_OutputSettings UTG_BlueprintFunctionLibrary::GetOutputSettingsParameterValue(UObject* WorldContextObject, UTextureGraph* InTextureGraph, FName ParameterName)
+FTG_OutputSettings UTG_BlueprintFunctionLibrary::GetSettingsParameterValue(UObject* WorldContextObject, UTextureGraph* InTextureGraph, FName ParameterName)
 {
 	FTG_OutputSettings ParameterValue;
 
@@ -281,7 +318,7 @@ FTG_OutputSettings UTG_BlueprintFunctionLibrary::GetOutputSettingsParameterValue
 
 		if (!bFoundParameter)
 		{
-			AddParamWarning(ParameterName, InTextureGraph, "GetOutputSettingsParameterValue");
+			AddParamWarning(ParameterName, InTextureGraph, "GetSettingsParameterValue");
 		}
 	}
 
@@ -297,5 +334,13 @@ void UTG_BlueprintFunctionLibrary::AddParamWarning(FName ParamName, UObject* Obj
 		->AddToken(FUObjectToken::Create(ObjectPtr))
 		->AddToken(FTextToken::Create(FText::Format(LOCTEXT("WithInvalidParam", "with invalid ParameterName '{ParamName}'. This is likely due to a Blueprint error."), Arguments)));
 }
+
+void UTG_BlueprintFunctionLibrary::AddError( UObject* ObjectPtr, FString FunctionName , FString Error)
+{
+	FMessageLog("PIE").Error()
+		->AddToken(FTextToken::Create(FText::Format(LOCTEXT("{FunctionName}", "{FunctionName} called on"), FText::FromString(FunctionName))))
+		->AddToken(FUObjectToken::Create(ObjectPtr))
+		->AddToken(FTextToken::Create(FText::FromString(Error)));
+}
 #undef LOCTEXT_NAMESPACE
 
diff --git a/Engine/Plugins/Experimental/TextureGraph/Source/TextureGraph/Private/TG_HelperFunctions.cpp b/Engine/Plugins/Experimental/TextureGraph/Source/TextureGraph/Private/TG_HelperFunctions.cpp
--- a/Engine/Plugins/Experimental/TextureGraph/Source/TextureGraph/Private/TG_HelperFunctions.cpp
+++ b/Engine/Plugins/Experimental/TextureGraph/Source/TextureGraph/Private/TG_HelperFunctions.cpp
@@ -135,6 +135,9 @@ AsyncBool FTG_HelperFunctions::ExportAsync(UTextureGraph* InTextureGraph, FStrin
 				MapSettings.Path = Path;
 				MapSettings.UseOverridePath = OverrideExportPath;
 				MapSettings.OverwriteTextures = OverwriteTextures;
+				MapSettings.LODGroup = OutputSetting.LODGroup;
+				MapSettings.Compression = OutputSetting.Compression;
+				MapSettings.IsSRGB = OutputSetting.bSRGB;
 				TargetExportSettings.ExportPreset.push_back(std::pair<FName, FExportMapSettings>{ MapSettings.Name, MapSettings });
 			}
 			else
diff --git a/Engine/Plugins/Experimental/TextureGraph/Source/TextureGraph/Private/TG_OutputSettings.cpp b/Engine/Plugins/Experimental/TextureGraph/Source/TextureGraph/Private/TG_OutputSettings.cpp
--- a/Engine/Plugins/Experimental/TextureGraph/Source/TextureGraph/Private/TG_OutputSettings.cpp
+++ b/Engine/Plugins/Experimental/TextureGraph/Source/TextureGraph/Private/TG_OutputSettings.cpp
@@ -13,6 +13,72 @@ FTG_OutputSettings* FTG_OutputExpressionInfo::GetOutputSettings()
 	return &OutputPtr->OutputSettings;
 }
 
+void FTG_OutputSettings::Set(int InWidth, int InHeight, FString Name /*= "None"*/, FString Path /*= "None"*/, ETG_TextureFormat Format /*= ETG_TextureFormat::BGRA8*/, ETG_TexturePresetType InTextureType /*= ETG_TexturePresetType::None*/,
+	TextureCompressionSettings InCompression /*= TextureCompressionSettings::TC_Default*/, TextureGroup InLodGroup /*= TextureGroup::TEXTUREGROUP_World*/, bool InbSRGB /*= false*/)
+{
+	BaseName = *Name;
+	FolderPath = *Path;
+	Width = (EResolution)InWidth;
+	Height = (EResolution)InHeight;
+	TextureFormat = Format;
+	TexturePresetType = InTextureType;
+
+	if (TexturePresetType != ETG_TexturePresetType::None)
+	{
+		OnSetTexturePresetType(TexturePresetType);
+	}
+	else
+	{
+		Compression = InCompression;
+		LODGroup = InLodGroup;
+		bSRGB = InbSRGB;
+	}
+}
+
+void FTG_OutputSettings::OnSetTexturePresetType(ETG_TexturePresetType Type)
+{
+	switch (Type)
+	{
+	case ETG_TexturePresetType::None:
+		LODGroup = TextureGroup::TEXTUREGROUP_World;
+		Compression = TextureCompressionSettings::TC_Default;
+		bSRGB = true;
+		break;
+	case ETG_TexturePresetType::Diffuse:
+	case ETG_TexturePresetType::Emissive:
+		LODGroup = TextureGroup::TEXTUREGROUP_Character;
+		Compression = TextureCompressionSettings::TC_Default;
+		bSRGB = true;
+		break;
+	case ETG_TexturePresetType::FX:
+		LODGroup = TextureGroup::TEXTUREGROUP_Character;
+		Compression = TextureCompressionSettings::TC_Masks;
+		bSRGB = false;
+		break;
+	case ETG_TexturePresetType::Normal:
+		LODGroup = TextureGroup::TEXTUREGROUP_CharacterNormalMap;
+		Compression = TextureCompressionSettings::TC_Normalmap;
+		bSRGB = false;
+		break;
+	case ETG_TexturePresetType::MaskComp:
+	case ETG_TexturePresetType::Specular:
+		LODGroup = TextureGroup::TEXTUREGROUP_CharacterSpecular;
+		Compression = TextureCompressionSettings::TC_Masks;
+		bSRGB = false;
+		break;
+	case ETG_TexturePresetType::Tangent:
+		LODGroup = TextureGroup::TEXTUREGROUP_CharacterSpecular;
+		Compression = TextureCompressionSettings::TC_Default;
+		bSRGB = false;
+		break;
+	default:
+		LODGroup = TextureGroup::TEXTUREGROUP_World;
+		Compression = TextureCompressionSettings::TC_Default;
+		bSRGB = true;
+		break;
+	}
+}
+
 void UTG_OutputSettingsSet::InitOutputSettings()
 {
 	OutputExpressionInfos.Empty();
diff --git a/Engine/Plugins/Experimental/TextureGraph/Source/TextureGraph/Public/Blueprint/TG_BlueprintFunctionLibrary.h b/Engine/Plugins/Experimental/TextureGraph/Source/TextureGraph/Public/Blueprint/TG_BlueprintFunctionLibrary.h
--- a/Engine/Plugins/Experimental/TextureGraph/Source/TextureGraph/Public/Blueprint/TG_BlueprintFunctionLibrary.h
+++ b/Engine/Plugins/Experimental/TextureGraph/Source/TextureGraph/Public/Blueprint/TG_BlueprintFunctionLibrary.h
@@ -50,12 +50,15 @@ public:
 	static FLinearColor GetColorParameterValue(UObject* WorldContextObject, UTextureGraph* TextureScript, FName ParameterName);
 
 	/** Sets a FTG_OutputSettings parameter value on the TextureGraph instance. Logs if ParameterName is invalid. */
-	UFUNCTION(BlueprintCallable, Category = "TextureGraph", meta = (Keywords = "SetOutputSettingsParameterValue", WorldContext = "WorldContextObject"))
-	static void SetOutputSettingsParameterValue(UObject* WorldContextObject, UTextureGraph* TextureScript, FName ParameterName, int Width,int Height, FString FileName = "None", FString FolderPath = "None", ETG_TextureFormat Format = ETG_TextureFormat::BGRA8);
+	UFUNCTION(BlueprintCallable, Category = "TextureGraph", meta = (Keywords = "SetSettingsParameterValue", WorldContext = "WorldContextObject"))
+	static void SetSettingsParameterValue(UObject* WorldContextObject, UTextureGraph* TextureScript, FName ParameterName, int Width,int Height, FString FileName = "None",
+		FString FolderPath = "None", ETG_TextureFormat Format = ETG_TextureFormat::BGRA8 , ETG_TexturePresetType TextureType = ETG_TexturePresetType::None,
+		TextureCompressionSettings Compression = TextureCompressionSettings::TC_Default, TextureGroup LodGroup = TextureGroup::TEXTUREGROUP_World, bool bSRGB = false);
 
 	/** Gets a FOutputSettings parameter value from the TextureGraph instance. Logs if ParameterName is invalid. */
 	UFUNCTION(BlueprintCallable, Category = "TextureGraph", meta = (Keywords = "GetOutputSettingsParameterValue", WorldContext = "WorldContextObject"))
-	static FTG_OutputSettings GetOutputSettingsParameterValue(UObject* WorldContextObject, UTextureGraph* TextureScript, FName ParameterName);
+	static FTG_OutputSettings GetSettingsParameterValue(UObject* WorldContextObject, UTextureGraph* TextureScript, FName ParameterName);
 
 	static void AddParamWarning(FName ParamName, UObject* ObjectPtr, FString FunctionName);
+	static void AddError(UObject* ObjectPtr, FString FunctionName, FString Error);
 };
diff --git a/Engine/Plugins/Experimental/TextureGraph/Source/TextureGraph/Public/Expressions/Output/TG_Expression_Output.h b/Engine/Plugins/Experimental/TextureGraph/Source/TextureGraph/Public/Expressions/Output/TG_Expression_Output.h
--- a/Engine/Plugins/Experimental/TextureGraph/Source/TextureGraph/Public/Expressions/Output/TG_Expression_Output.h
+++ b/Engine/Plugins/Experimental/TextureGraph/Source/TextureGraph/Public/Expressions/Output/TG_Expression_Output.h
@@ -28,7 +28,7 @@ public:
 	FTG_Variant Output = 0;
 
 	//When we will work on Node UI for FOutputSettings we will set the category as TG_Setting
-	UPROPERTY(EditAnywhere, Category = NoCategory, meta = (TGType = "TG_Input", CollapsableChildProperties,ShowOnlyInnerProperties, FullyExpand, NoResetToDefault) )
+	UPROPERTY(EditAnywhere, Category = NoCategory, meta = (TGType = "TG_Input", CollapsableChildProperties,ShowOnlyInnerProperties, FullyExpand, NoResetToDefault, PinDisplayName = "Settings") )
 	FTG_OutputSettings OutputSettings;
 	
 	virtual FTG_Name GetDefaultName() const override { return TEXT("Output");}
diff --git a/Engine/Plugins/Experimental/TextureGraph/Source/TextureGraph/Public/TG_OutputSettings.h b/Engine/Plugins/Experimental/TextureGraph/Source/TextureGraph/Public/TG_OutputSettings.h
--- a/Engine/Plugins/Experimental/TextureGraph/Source/TextureGraph/Public/TG_OutputSettings.h
+++ b/Engine/Plugins/Experimental/TextureGraph/Source/TextureGraph/Public/TG_OutputSettings.h
@@ -34,6 +34,18 @@ public:
 	UPROPERTY(BlueprintReadWrite,EditAnywhere, Category = "Advanced", DisplayName = "Texture Format", Meta = (NoResetToDefault))
 		ETG_TextureFormat TextureFormat = ETG_TextureFormat::BGRA8;
 
+	UPROPERTY(BlueprintReadWrite, EditAnywhere, Category = "Advanced", DisplayName = "Texture Type", Meta = (NoResetToDefault))
+		ETG_TexturePresetType TexturePresetType = ETG_TexturePresetType::None;
+
+	UPROPERTY(BlueprintReadWrite, EditAnywhere, Category = "Advanced", DisplayName = "LOD Texture Group", Meta = (NoResetToDefault, EditCondition = "TexturePresetType == ETG_TexturePresetType::None"))
+		TEnumAsByte<enum TextureGroup> LODGroup;
+
+	UPROPERTY(BlueprintReadWrite, EditAnywhere, Category = "Advanced", DisplayName = "Compression", Meta = (NoResetToDefault, EditCondition = "TexturePresetType == ETG_TexturePresetType::None") )
+		TEnumAsByte <enum TextureCompressionSettings> Compression;
+
+	UPROPERTY(BlueprintReadWrite, EditAnywhere, Category = "Advanced", DisplayName = "sRGB", Meta = (NoResetToDefault, EditCondition = "TexturePresetType == ETG_TexturePresetType::None"))
+		bool bSRGB = false;
+
 	FString GetFullOutputName() { return  FString::Format(TEXT("{0}"), { BaseName.ToString()});}
 
 	bool operator==(const FTG_OutputSettings& Other) const
@@ -54,14 +66,10 @@ public:
 		return ExportString;
 	}
 
-	void Set(int InWidth, int InHeight, FString Name = "None", FString Path = "None", ETG_TextureFormat Format = ETG_TextureFormat::BGRA8)
-	{
-		BaseName = *Name;
-		FolderPath = *Path;
-		Width = (EResolution)InWidth;
-		Height = (EResolution)InHeight;
-		TextureFormat = Format;
-	}
+	void Set(int InWidth, int InHeight, FString Name = "None", FString Path = "None", ETG_TextureFormat Format = ETG_TextureFormat::BGRA8, ETG_TexturePresetType InTextureType = ETG_TexturePresetType::None,
+		TextureCompressionSettings InCompression = TextureCompressionSettings::TC_Default, TextureGroup InLodGroup = TextureGroup::TEXTUREGROUP_World, bool InbSRGB = false);
+
+	void OnSetTexturePresetType(ETG_TexturePresetType Type);
 };
 
 USTRUCT()
diff --git a/Engine/Plugins/Experimental/TextureGraph/Source/TextureGraphEditor/Private/Pins/STG_GraphPinOutputSettings.cpp b/Engine/Plugins/Experimental/TextureGraph/Source/TextureGraphEditor/Private/Pins/STG_GraphPinOutputSettings.cpp
--- a/Engine/Plugins/Experimental/TextureGraph/Source/TextureGraphEditor/Private/Pins/STG_GraphPinOutputSettings.cpp
+++ b/Engine/Plugins/Experimental/TextureGraph/Source/TextureGraphEditor/Private/Pins/STG_GraphPinOutputSettings.cpp
@@ -46,7 +46,7 @@ void STG_GraphPinOutputSettings::OnOutputSettingsChanged(const FTG_OutputSetting
 	}
 }
 
-FProperty* STG_GraphPinOutputSettings::GetPinProperty()
+FProperty* STG_GraphPinOutputSettings::GetPinProperty() const
 {
 	const UTG_EdGraphSchema* Schema = Cast<const UTG_EdGraphSchema>(GraphPinObj->GetOwningNode()->GetSchema());
 	UTG_Pin* TSPin = Schema->GetTGPinFromEdPin(GraphPinObj);
@@ -54,7 +54,7 @@ FProperty* STG_GraphPinOutputSettings::GetPinProperty()
 	return Property;
 }
 
-bool STG_GraphPinOutputSettings::ShowChildProperties()
+bool STG_GraphPinOutputSettings::ShowChildProperties() const
 {
 	FProperty* Property = GetPinProperty();
 	bool ShowChildProperties = true;
@@ -66,7 +66,7 @@ bool STG_GraphPinOutputSettings::ShowChildProperties()
 	return ShowChildProperties;
 }
 
-bool STG_GraphPinOutputSettings::CollapsibleChildProperties()
+bool STG_GraphPinOutputSettings::CollapsibleChildProperties() const
 {
 	FProperty* Property = GetPinProperty();
 	bool Collapsible = false;
@@ -78,6 +78,18 @@ bool STG_GraphPinOutputSettings::CollapsibleChildProperties()
 	return Collapsible;
 }
 
+EVisibility STG_GraphPinOutputSettings::ShowLabel() const
+{
+	bool bHide = ShowChildProperties();
+
+	if ((GraphPinObj->GetOwningNode()->AdvancedPinDisplay == ENodeAdvancedPins::Type::Hidden && GraphPinObj->LinkedTo.Num() > 0))
+	{
+		bHide = false;
+	}
+
+	return bHide ? EVisibility::Collapsed : EVisibility::Visible;
+}
+
 TSharedRef<SWidget>	STG_GraphPinOutputSettings::GetDefaultValueWidget()
 {
 	if (ShowChildProperties())
@@ -99,14 +111,11 @@ TSharedRef<SWidget>	STG_GraphPinOutputSettings::GetDefaultValueWidget()
 
 TSharedRef<SWidget> STG_GraphPinOutputSettings::GetLabelWidget(const FName& InLabelStyle)
 {
-	if(ShowChildProperties() && CollapsibleChildProperties())
-	{
-		return SNew(SHorizontalBox);
-	}
-	else
-	{
-		return SGraphPin::GetLabelWidget(InLabelStyle);
-	}
+	return SNew(STextBlock)
+			.Text(this, &STG_GraphPinOutputSettings::GetPinLabel)
+			.TextStyle(FAppStyle::Get(), InLabelStyle)
+			.Visibility(this, &STG_GraphPinOutputSettings::ShowLabel)
+			.ColorAndOpacity(this, &STG_GraphPinOutputSettings::GetPinTextColor);
 }
 
 void STG_GraphPinOutputSettings::OnAdvancedViewChanged(const ECheckBoxState NewCheckedState)
diff --git a/Engine/Plugins/Experimental/TextureGraph/Source/TextureGraphEditor/Private/Pins/STG_GraphPinOutputSettings.h b/Engine/Plugins/Experimental/TextureGraph/Source/TextureGraphEditor/Private/Pins/STG_GraphPinOutputSettings.h
--- a/Engine/Plugins/Experimental/TextureGraph/Source/TextureGraphEditor/Private/Pins/STG_GraphPinOutputSettings.h
+++ b/Engine/Plugins/Experimental/TextureGraph/Source/TextureGraphEditor/Private/Pins/STG_GraphPinOutputSettings.h
@@ -14,9 +14,10 @@ public:
 	SLATE_END_ARGS()
 
 	void Construct(const FArguments& InArgs, UEdGraphPin* InGraphPinObj);
-	FProperty* GetPinProperty();
-	bool ShowChildProperties();
-	bool CollapsibleChildProperties();
+	FProperty* GetPinProperty() const;
+	bool ShowChildProperties() const;
+	bool CollapsibleChildProperties() const;
+	EVisibility ShowLabel() const;
 	//~ Begin SGraphPin Interface
 	virtual TSharedRef<SWidget>	GetDefaultValueWidget() override;
 	virtual TSharedRef<SWidget> GetLabelWidget(const FName& InLabelStyle) override;
diff --git a/Engine/Plugins/Experimental/TextureGraph/Source/TextureGraphEditor/Private/Pins/STG_GraphPinOutputSettingsWidget.cpp b/Engine/Plugins/Experimental/TextureGraph/Source/TextureGraphEditor/Private/Pins/STG_GraphPinOutputSettingsWidget.cpp
--- a/Engine/Plugins/Experimental/TextureGraph/Source/TextureGraphEditor/Private/Pins/STG_GraphPinOutputSettingsWidget.cpp
+++ b/Engine/Plugins/Experimental/TextureGraph/Source/TextureGraphEditor/Private/Pins/STG_GraphPinOutputSettingsWidget.cpp
@@ -56,6 +56,12 @@ void STG_GraphPinOutputSettingsWidget::Construct(const FArguments& InArgs, UEdGr
 	GetHeightDelegate.BindRaw(this, &STG_GraphPinOutputSettingsWidget::HandleHeightText);
 	OnGenerateFormatMenu.BindRaw(this, &STG_GraphPinOutputSettingsWidget::OnGenerateFormatEnumMenu);
 	GetFormatDelegate.BindRaw(this, &STG_GraphPinOutputSettingsWidget::HandleFormatText);
+	OnGenerateTexturePresetTypeMenu.BindRaw(this, &STG_GraphPinOutputSettingsWidget::OnGenerateTexturePresetTypeEnumMenu);
+	GetTexturePresetTypeDelegate.BindRaw(this, &STG_GraphPinOutputSettingsWidget::HandleTexturePresetTypeText);
+	OnGenerateLodGroupMenu.BindRaw(this, &STG_GraphPinOutputSettingsWidget::OnGenerateLodGroupEnumMenu);
+	GetLodGroupDelegate.BindRaw(this, &STG_GraphPinOutputSettingsWidget::HandleLodGroupText);
+	OnGenerateCompressionMenu.BindRaw(this, &STG_GraphPinOutputSettingsWidget::OnGenerateCompressionEnumMenu);
+	GetCompressionDelegate.BindRaw(this, &STG_GraphPinOutputSettingsWidget::HandleCompressionText);
 
 	GraphPinObj = InGraphPinObj;
 
@@ -63,162 +69,229 @@ void STG_GraphPinOutputSettingsWidget::Construct(const FArguments& InArgs, UEdGr
 
 	OutputSettingsAttribute.Assign(*this, InArgs._OutputSettings);
 
+	const int UniformPadding = 2;
+
 	ChildSlot
 	[
-		SNew(SVerticalBox)
-
-		//Path Slot
-		+ SVerticalBox::Slot()
-		.Padding(2)
-		.AutoHeight()
+		//Param Box hidden when pin is connected and Advanced view is collapsed
+		SNew(SBox)
+		.MinDesiredWidth(250)
+		.MaxDesiredWidth(450)
+		.Visibility(this, &STG_GraphPinOutputSettingsWidget::ShowParameters)
 		[
-			AddEditBoxWithBrowseButton(LOCTEXT("OutputPath", "Path"), GetPathDelegate, PathCommitted)
-		]
+			SNew(SVerticalBox)
 
-		//Name Slot
-		+ SVerticalBox::Slot()
-		//.HAlign(HAlign_Fill)
-		.Padding(2)
-		.AutoHeight()
-		[
-			AddEditBox(LOCTEXT("OutputName", "File Name"), GetNameDelegate, NameCommitted)
-		]
+			//Path Slot
+			+ SVerticalBox::Slot()
+			.Padding(UniformPadding)
+			.AutoHeight()
+			[
+				AddEditBoxWithBrowseButton(LOCTEXT("OutputPath", "Path"), GetPathDelegate, PathCommitted)
+			]
 
-		//Path Width
-		+ SVerticalBox::Slot()
-		.Padding(2)
-		[
-			AddEnumComobox(LOCTEXT("OutputWidth", "Width"), GetWidthDelegate, OnGenerateWidthMenu)
-		]
+			//Name Slot
+			+ SVerticalBox::Slot()
+			//.HAlign(HAlign_Fill)
+			.Padding(UniformPadding)
+			.AutoHeight()
+			[
+				AddEditBox(LOCTEXT("OutputName", "File Name"), GetNameDelegate, NameCommitted)
+			]
 
-		//Path Height
-		+ SVerticalBox::Slot()
-		.Padding(2)
-		[
-			AddEnumComobox(LOCTEXT("OutputHeight", "Height"), GetHeightDelegate, OnGenerateHeightMenu)
-		]
+			//Path Width
+			+ SVerticalBox::Slot()
+			.Padding(UniformPadding)
+			[
+				AddEnumComobox(LOCTEXT("OutputWidth", "Width"), GetWidthDelegate, OnGenerateWidthMenu)
+			]
 
-		//Path Format
-		+ SVerticalBox::Slot()
-		.Padding(2)
-		[
-			AddEnumComobox(LOCTEXT("OutputFormat", "Format"), GetFormatDelegate, OnGenerateFormatMenu)
-		]
-	];
-}
+			//Path Height
+			+ SVerticalBox::Slot()
+			.Padding(UniformPadding)
+			[
+				AddEnumComobox(LOCTEXT("OutputHeight", "Height"), GetHeightDelegate, OnGenerateHeightMenu)
+			]
 
-TSharedRef<SWidget> STG_GraphPinOutputSettingsWidget::AddEditBoxWithBrowseButton(FText Label, FGetTextDelegate GetText, FTextCommitted OnTextCommitted)
-{
-	return SNew(SBox)
-		.MinDesiredWidth(150)
-		.MaxDesiredWidth(350)
-		[
-			SNew(SHorizontalBox)
+			//Path Format
+			+ SVerticalBox::Slot()
+			.Padding(UniformPadding)
+			[
+				AddEnumComobox(LOCTEXT("OutputFormat", "Format"), GetFormatDelegate, OnGenerateFormatMenu)
+			]
 
-			+ SHorizontalBox::Slot()
-			.FillWidth(0.4)
-			.VAlign(VAlign_Center)
+			//Texture Preset Type
+			+ SVerticalBox::Slot()
+			.Padding(UniformPadding)
 			[
-				SNew(STextBlock)
-				.Text(Label)
-				.TextStyle(FAppStyle::Get(), TEXT("Graph.Node.PinName"))
+				AddEnumComobox(LOCTEXT("TextureType", "Texture Type"), GetTexturePresetTypeDelegate, OnGenerateTexturePresetTypeMenu)
 			]
 
-			+ SHorizontalBox::Slot()
+			//Lod Group
+			+ SVerticalBox::Slot()
+			.Padding(UniformPadding)
 			[
 				SNew(SHorizontalBox)
+				.IsEnabled(this, &STG_GraphPinOutputSettingsWidget::IsDefaultPreset)
 
 				+ SHorizontalBox::Slot()
-				.VAlign(VAlign_Center)
-				.HAlign(HAlign_Fill)
-				.FillWidth(1)
 				[
-					SNew(SEditableTextBox)
-					.Text_Lambda([GetText]() { return GetText.Execute(); })
-					.SelectAllTextWhenFocused(true)
-					.SelectAllTextOnCommit(true)
-					.OnTextCommitted_Lambda([OnTextCommitted](const FText& InText, ETextCommit::Type InCommitType) {
-						OnTextCommitted.ExecuteIfBound(InText, InCommitType);
-					})
+					AddEnumComobox(LOCTEXT("LodGroup", "LOD Group"), GetLodGroupDelegate, OnGenerateLodGroupMenu)
+		
 				]
+			]
+
+			//Compression Settings
+			+ SVerticalBox::Slot()
+			.Padding(UniformPadding)
+			[
+				SNew(SHorizontalBox)
+				.IsEnabled(this, &STG_GraphPinOutputSettingsWidget::IsDefaultPreset)
 
 				+ SHorizontalBox::Slot()
-				.VAlign(VAlign_Center)
-				.HAlign(HAlign_Fill)
-				.AutoWidth()
 				[
-					SNew(SButton)
-					.ButtonStyle(FAppStyle::Get(), "SimpleButton")
-					.OnClicked(this, &STG_GraphPinOutputSettingsWidget::OnBrowseClick)
-					.HAlign(HAlign_Left)
-					.VAlign(VAlign_Center)
-					.ToolTipText(LOCTEXT("Browsepath_ToolTip", "Select the path for the output"))
-					[
-						SNew(SImage)
-						.ColorAndOpacity(FSlateColor::UseForeground())
-						.Image(FAppStyle::GetBrush("Icons.BrowseContent"))
-					]
+					AddEnumComobox(LOCTEXT("CompressionSettings", "Compression Settings"), GetCompressionDelegate, OnGenerateCompressionMenu)
 				]
 			]
-		];
+
+			//SRGB
+			+ SVerticalBox::Slot()
+			.Padding(UniformPadding)
+			[
+				SNew(SHorizontalBox)
+				.IsEnabled(this, &STG_GraphPinOutputSettingsWidget::IsDefaultPreset)
+
+				+ SHorizontalBox::Slot()
+				[
+					AddSRGBWidget().ToSharedRef()
+				]
+			]			
+		]
+	];
 }
 
-TSharedRef<SWidget> STG_GraphPinOutputSettingsWidget::AddEditBox(FText Label, FGetTextDelegate GetText, FTextCommitted OnTextCommitted)
+EVisibility STG_GraphPinOutputSettingsWidget::ShowPinLabel() const
 {
-	return SNew(SBox)
-	.MinDesiredWidth(150)
-	.MaxDesiredWidth(350)
-	[
-		SNew(SHorizontalBox)
+	return ShowParameters() == EVisibility::Collapsed ? EVisibility::Visible : EVisibility::Collapsed;
+}
 
-		+ SHorizontalBox::Slot()
-		.FillWidth(0.4)
-		.VAlign(VAlign_Center)
+EVisibility STG_GraphPinOutputSettingsWidget::ShowParameters() const
+{
+	return (GraphPinObj->GetOwningNode()->AdvancedPinDisplay == ENodeAdvancedPins::Type::Hidden && GraphPinObj->LinkedTo.Num() > 0) ? EVisibility::Collapsed : EVisibility::Visible;
+}
+
+TSharedRef<SWidget> STG_GraphPinOutputSettingsWidget::AddEditBoxWithBrowseButton(FText Label, FGetTextDelegate GetText, FTextCommitted OnTextCommitted)
+{
+	return SNew(SHorizontalBox)
+
+	+ SHorizontalBox::Slot()
+	.VAlign(VAlign_Center)
+	.AutoWidth()
+	[
+		SNew(SBox)
+		.MinDesiredWidth(LabelSize)
+		.MaxDesiredWidth(LabelSize)
 		[
 			SNew(STextBlock)
 			.Text(Label)
 			.TextStyle(FAppStyle::Get(), TEXT("Graph.Node.PinName"))
 		]
+	]
+
+	+ SHorizontalBox::Slot()
+	.FillWidth(1.0)
+	[
+		SNew(SHorizontalBox)
 
 		+ SHorizontalBox::Slot()
+		.VAlign(VAlign_Center)
+		.HAlign(HAlign_Fill)
 		[
-
 			SNew(SEditableTextBox)
 			.Text_Lambda([GetText]() { return GetText.Execute(); })
 			.SelectAllTextWhenFocused(true)
 			.SelectAllTextOnCommit(true)
 			.OnTextCommitted_Lambda([OnTextCommitted](const FText& InText, ETextCommit::Type InCommitType) {
-					OnTextCommitted.ExecuteIfBound(InText, InCommitType);
+				OnTextCommitted.ExecuteIfBound(InText, InCommitType);
 			})
 		]
+
+		+ SHorizontalBox::Slot()
+		.VAlign(VAlign_Center)
+		.HAlign(HAlign_Right)
+		.AutoWidth()
+		[
+			SNew(SButton)
+			.ButtonStyle(FAppStyle::Get(), "SimpleButton")
+			.OnClicked(this, &STG_GraphPinOutputSettingsWidget::OnBrowseClick)
+			.HAlign(HAlign_Left)
+			.VAlign(VAlign_Center)
+			.ToolTipText(LOCTEXT("Browsepath_ToolTip", "Select the path for the output"))
+			[
+				SNew(SImage)
+				.ColorAndOpacity(FSlateColor::UseForeground())
+				.Image(FAppStyle::GetBrush("Icons.BrowseContent"))
+			]
+		]
+	];
+}
+
+TSharedRef<SWidget> STG_GraphPinOutputSettingsWidget::AddEditBox(FText Label, FGetTextDelegate GetText, FTextCommitted OnTextCommitted)
+{
+	return SNew(SHorizontalBox)
+
+	+ SHorizontalBox::Slot()
+	.VAlign(VAlign_Center)
+	.AutoWidth()
+	[
+		SNew(SBox)
+		.MinDesiredWidth(LabelSize)
+		.MaxDesiredWidth(LabelSize)
+		[
+			SNew(STextBlock)
+			.Text(Label)
+			.TextStyle(FAppStyle::Get(), TEXT("Graph.Node.PinName"))
+		]
+	]
+
+	+ SHorizontalBox::Slot()
+	[
+		SNew(SEditableTextBox)
+		.Text_Lambda([GetText]() { return GetText.Execute(); })
+		.SelectAllTextWhenFocused(true)
+		.SelectAllTextOnCommit(true)
+		.OnTextCommitted_Lambda([OnTextCommitted](const FText& InText, ETextCommit::Type InCommitType) {
+				OnTextCommitted.ExecuteIfBound(InText, InCommitType);
+		})
 	];
 }
 
 TSharedRef<SWidget> STG_GraphPinOutputSettingsWidget::AddEnumComobox(FText Label, FGetTextDelegate GetText, FGenerateEnumMenu OnGenerateEnumMenu)
 {
-	return SNew(SBox)
-	.MinDesiredWidth(150)
-	.MaxDesiredWidth(350)
+	return SNew(SHorizontalBox)
+
+	+ SHorizontalBox::Slot()
+	.VAlign(VAlign_Center)
+	.HAlign(HAlign_Left)
+	.AutoWidth()
 	[
-		SNew(SHorizontalBox)
-		+ SHorizontalBox::Slot()
-		.FillWidth(0.4)
-		.VAlign(VAlign_Center)
+		SNew(SBox)
+		.MinDesiredWidth(LabelSize)
+		.MaxDesiredWidth(LabelSize)
 		[
 			SNew(STextBlock)
 			.Text(Label)
 			.TextStyle(FAppStyle::Get(), TEXT("Graph.Node.PinName"))
 		]
+	]
 
-		+ SHorizontalBox::Slot()
+	+ SHorizontalBox::Slot()
+	[
+		SNew(SComboButton)
+		.OnGetMenuContent(OnGenerateEnumMenu)
+		.ButtonContent()
 		[
-			SNew(SComboButton)
-			.OnGetMenuContent(OnGenerateEnumMenu)
-			.ButtonContent()
-			[
-				SNew(STextBlock)
-				.Text_Lambda([GetText]() { return GetText.Execute(); })
-			]
+			SNew(STextBlock)
+			.Text_Lambda([GetText]() { return GetText.Execute(); })
 		]
 	];
 }
@@ -233,6 +306,48 @@ FReply STG_GraphPinOutputSettingsWidget::OnBrowseClick()
 	return FReply::Handled();
 }
 
+TSharedPtr<SWidget> STG_GraphPinOutputSettingsWidget::AddSRGBWidget()
+{
+	TSharedPtr<SWidget> SRGB = SNew(SHorizontalBox)
+
+	+ SHorizontalBox::Slot()
+	.VAlign(VAlign_Center)
+	.HAlign(HAlign_Left)
+	.AutoWidth()
+	[
+		SNew(SBox)
+		.MinDesiredWidth(LabelSize)
+		.MaxDesiredWidth(LabelSize)
+		[
+			SNew(STextBlock)
+			.Text(FText::FromString("sRGB"))
+		]
+	]
+
+	+ SHorizontalBox::Slot()
+	[
+		SNew(SCheckBox)
+		.IsChecked(this, &STG_GraphPinOutputSettingsWidget::HandleSRGBIsChecked)
+		.OnCheckStateChanged(this, &STG_GraphPinOutputSettingsWidget::HandleSRGBExecute)
+	];
+
+	return SRGB;
+}
+
+ECheckBoxState STG_GraphPinOutputSettingsWidget::HandleSRGBIsChecked() const
+{
+	auto Settings = GetSettings();
+	return Settings.bSRGB ? ECheckBoxState::Checked : ECheckBoxState::Unchecked;
+}
+
+void STG_GraphPinOutputSettingsWidget::HandleSRGBExecute(ECheckBoxState InNewState)
+{
+	bSRGB = InNewState == ECheckBoxState::Checked;
+	auto Settings = GetSettings();
+	Settings.bSRGB = bSRGB;
+	OnOutputSettingsChanged.ExecuteIfBound(Settings);
+}
+
 FTG_OutputSettings STG_GraphPinOutputSettingsWidget::GetSettings() const
 {
 	FString OutputSettingString = GraphPinObj->GetDefaultAsString();
@@ -260,13 +375,36 @@ void STG_GraphPinOutputSettingsWidget::GenerateStringsFromEnum(TArray<FString>&
 	}
 }
 
+void STG_GraphPinOutputSettingsWidget::GenerateValuesFromEnum(TArray<uint8>& OutEnumValues, const FString& EnumPathName) const
+{
+	UEnum* EnumPtr = FindObject<UEnum>(nullptr, *EnumPathName);
+	if (EnumPtr)
+	{
+		for (int32 i = 0; i < EnumPtr->NumEnums() - 1; ++i)
+		{
+			if (!EnumPtr->HasMetaData(TEXT("Hidden"), i))
+			{
+				FString DisplayName = EnumPtr->GetDisplayNameTextByIndex(i).ToString();
+				uint8 EnumValue = EnumPtr->GetValueByIndex(i);
+				UE_LOG(LogTemp, Warning, TEXT("Enum Value: %d, Display Name: %s"), EnumValue, *DisplayName);
+				OutEnumValues.Add(EnumValue);
+			}
+		}
+	}
+}
+
 int STG_GraphPinOutputSettingsWidget::GetValueFromIndex(const FString& EnumPathName, int Index) const
 {
 	int Value = 0;
+	TArray<uint8> EnumValues;
 	UEnum* EnumPtr = FindObject<UEnum>(nullptr, *EnumPathName);
 	if (EnumPtr)
 	{
-		Value = EnumPtr->GetValueByIndex(Index);
+		GenerateValuesFromEnum(EnumValues, EnumPathName);
+		if (EnumValues.Num() > Index)
+		{
+			Value = EnumValues[Index];
+		}
 	}
 	return Value;
 }
@@ -400,11 +538,130 @@ FText STG_GraphPinOutputSettingsWidget::HandleFormatText() const
 	return FText::FromString(GetEnumValueDisplayName(StaticEnum<ETG_TextureFormat>()->GetPathName(), (int)GetSettings().TextureFormat));
 }
 
+TSharedRef<SWidget> STG_GraphPinOutputSettingsWidget::OnGenerateTexturePresetTypeEnumMenu()
+{
+	FMenuBuilder MenuBuilder(true, nullptr);
+	//Get list for Enum
+	TArray<FString> EnumItems;
+	GenerateStringsFromEnum(EnumItems, StaticEnum<ETG_TexturePresetType>()->GetPathName());
+
+	for (int i = 0; i < EnumItems.Num(); i++)
+	{
+		auto Item = EnumItems[i];
+		MenuBuilder.AddMenuEntry(
+			FText::FromString(Item),
+			FText::FromString(Item),
+			FSlateIcon(),
+			FUIAction(
+				FExecuteAction::CreateSP(this, &STG_GraphPinOutputSettingsWidget::HandleTexturePresetTypeChanged, Item, i),
+				FCanExecuteAction(),
+				FIsActionChecked::CreateLambda([this, i]() {return SelectedTextureTypeIndex == i; })
+			));
+	}
+
+	return MenuBuilder.MakeWidget();
+}
+
+void STG_GraphPinOutputSettingsWidget::HandleTexturePresetTypeChanged(FString Name, int Index)
+{
+	auto Settings = GetSettings();
+	Settings.TexturePresetType = (ETG_TexturePresetType)GetValueFromIndex(StaticEnum<ETG_TexturePresetType>()->GetPathName(), Index);
+
+	Settings.OnSetTexturePresetType(Settings.TexturePresetType);
+
+	OnOutputSettingsChanged.ExecuteIfBound(Settings);
+
+	SelectedFormatIndex = Index;
+}
+
+FText STG_GraphPinOutputSettingsWidget::HandleTexturePresetTypeText() const
+{
+	return FText::FromString(GetEnumValueDisplayName(StaticEnum<ETG_TexturePresetType>()->GetPathName(), (int)GetSettings().TexturePresetType));
+}
+
+TSharedRef<SWidget> STG_GraphPinOutputSettingsWidget::OnGenerateLodGroupEnumMenu()
+{
+	FMenuBuilder MenuBuilder(true, nullptr);
+	//Get list for Enum
+	TArray<FString> EnumItems;
+	GenerateStringsFromEnum(EnumItems, StaticEnum<TextureGroup>()->GetPathName());
+
+	for (int i = 0; i < EnumItems.Num(); i++)
+	{
+		auto Item = EnumItems[i];
+		MenuBuilder.AddMenuEntry(
+			FText::FromString(Item),
+			FText::FromString(Item),
+			FSlateIcon(),
+			FUIAction(
+				FExecuteAction::CreateSP(this, &STG_GraphPinOutputSettingsWidget::HandleLodGroupChanged, Item, i),
+				FCanExecuteAction(),
+				FIsActionChecked::CreateLambda([this, i]() {return SelectedLodGroupIndex == i; })
+			));
+	}
+
+	return MenuBuilder.MakeWidget();
+}
+
+void STG_GraphPinOutputSettingsWidget::HandleLodGroupChanged(FString Name, int Index)
+{
+	auto Settings = GetSettings();
+	Settings.LODGroup = (TextureGroup)GetValueFromIndex(StaticEnum<TextureGroup>()->GetPathName(), Index);
+
+	OnOutputSettingsChanged.ExecuteIfBound(Settings);
+
+	SelectedFormatIndex = Index;
+}
+
+FText STG_GraphPinOutputSettingsWidget::HandleLodGroupText() const
+{
+	return FText::FromString(GetEnumValueDisplayName(StaticEnum<TextureGroup>()->GetPathName(), (int)GetSettings().LODGroup));
+}
+
 FText STG_GraphPinOutputSettingsWidget::GetNameAsText() const
 {
 	return FText::FromName(GetSettings().BaseName);
 }
 
+TSharedRef<SWidget> STG_GraphPinOutputSettingsWidget::OnGenerateCompressionEnumMenu()
+{
+	FMenuBuilder MenuBuilder(true, nullptr);
+	//Get list for Enum
+	TArray<FString> EnumItems;
+	GenerateStringsFromEnum(EnumItems, StaticEnum<TextureCompressionSettings>()->GetPathName());
+
+	for (int i = 0; i < EnumItems.Num(); i++)
+	{
+		auto Item = EnumItems[i];
+		MenuBuilder.AddMenuEntry(
+			FText::FromString(Item),
+			FText::FromString(Item),
+			FSlateIcon(),
+			FUIAction(
+				FExecuteAction::CreateSP(this, &STG_GraphPinOutputSettingsWidget::HandleCompressionChanged, Item, i),
+				FCanExecuteAction(),
+				FIsActionChecked::CreateLambda([this, i]() {return SelectedCompressionIndex == i; })
+			));
+	}
+
+	return MenuBuilder.MakeWidget();
+}
+
+void STG_GraphPinOutputSettingsWidget::HandleCompressionChanged(FString Name, int Index)
+{
+	auto Settings = GetSettings();
+	Settings.Compression = (TextureCompressionSettings)GetValueFromIndex(StaticEnum<TextureCompressionSettings>()->GetPathName(), Index);
+
+	OnOutputSettingsChanged.ExecuteIfBound(Settings);
+
+	SelectedFormatIndex = Index;
+}
+
+FText STG_GraphPinOutputSettingsWidget::HandleCompressionText() const
+{
+	return FText::FromString(GetEnumValueDisplayName(StaticEnum<TextureCompressionSettings>()->GetPathName(), (int)GetSettings().Compression));
+}
+
 void STG_GraphPinOutputSettingsWidget::OnNameCommitted(const FText& NewText, ETextCommit::Type /*CommitInfo*/)
 {
 	auto Settings = GetSettings(); 
@@ -427,6 +684,12 @@ void STG_GraphPinOutputSettingsWidget::OnPathCommitted(const FText& NewText, ETe
 	OnOutputSettingsChanged.ExecuteIfBound(Settings);
 }
 
+bool STG_GraphPinOutputSettingsWidget::IsDefaultPreset() const
+{
+	auto Settings = GetSettings();
+	return Settings.TexturePresetType == ETG_TexturePresetType::None;
+}
+
 void STG_GraphPinOutputSettingsWidget::PostUndo(bool bSuccess)
 {
 	
diff --git a/Engine/Plugins/Experimental/TextureGraph/Source/TextureGraphEditor/Private/Pins/STG_GraphPinOutputSettingsWidget.h b/Engine/Plugins/Experimental/TextureGraph/Source/TextureGraphEditor/Private/Pins/STG_GraphPinOutputSettingsWidget.h
--- a/Engine/Plugins/Experimental/TextureGraph/Source/TextureGraphEditor/Private/Pins/STG_GraphPinOutputSettingsWidget.h
+++ b/Engine/Plugins/Experimental/TextureGraph/Source/TextureGraphEditor/Private/Pins/STG_GraphPinOutputSettingsWidget.h
@@ -5,6 +5,7 @@
 #include "Widgets/SWidget.h"
 #include "Widgets/SCompoundWidget.h"
 #include "Widgets/DeclarativeSyntaxSupport.h"
+#include "Widgets/Input/SCheckbox.h"
 #include "TG_OutputSettings.h"
 #include "EditorUndoClient.h"
 #include "EdGraph/EdGraphPin.h"
@@ -53,10 +54,14 @@ public:
 	TSharedRef<SWidget> AddEditBoxWithBrowseButton(FText Label, FGetTextDelegate GetText, FTextCommitted OnTextCommitted);
 	TSharedRef<SWidget> AddEditBox(FText Label, FGetTextDelegate GetText, FTextCommitted OnTextCommitted);
 	TSharedRef<SWidget> AddEnumComobox(FText Label, FGetTextDelegate GetText, FGenerateEnumMenu OnGenerateEnumMenu);
+	TSharedPtr<SWidget> AddSRGBWidget();
 	FTG_OutputSettings GetSettings() const;
 	void GenerateStringsFromEnum(TArray<FString>& OutEnumNames, const FString& EnumPathName);
+	void GenerateValuesFromEnum(TArray<uint8>& OutEnumValues, const FString& EnumPathName) const;
 	int GetValueFromIndex(const FString& EnumPathName, int Index) const;
 	FString GetEnumValueDisplayName(const FString& EnumPathName, int EnumValue) const;
+	EVisibility ShowParameters() const;
+	EVisibility ShowPinLabel() const;
 
 protected:
 	//~ Begin FEditorUndoClient Interface
@@ -75,6 +80,10 @@ private:
 	int SelectedWidthIndex;
 	int SelectedHeightIndex;
 	int SelectedFormatIndex;
+	int SelectedTextureTypeIndex;
+	int SelectedLodGroupIndex;
+	int SelectedCompressionIndex;
+	bool bSRGB = false;
 
 	FGetTextDelegate GetPathDelegate;
 	FTextCommitted PathCommitted;
@@ -86,6 +95,12 @@ private:
 	FGetTextDelegate GetHeightDelegate;
 	FGenerateEnumMenu OnGenerateFormatMenu;
 	FGetTextDelegate GetFormatDelegate;
+	FGenerateEnumMenu OnGenerateTexturePresetTypeMenu;
+	FGetTextDelegate GetTexturePresetTypeDelegate;
+	FGenerateEnumMenu OnGenerateLodGroupMenu;
+	FGetTextDelegate GetLodGroupDelegate;
+	FGenerateEnumMenu OnGenerateCompressionMenu;
+	FGetTextDelegate GetCompressionDelegate;
 
 	FText GetNameAsText() const;
 	void OnNameCommitted(const FText& NewText, ETextCommit::Type CommitInfo);
@@ -105,5 +120,24 @@ private:
 	void HandleFormatChanged(FString OuputName, int Index);
 	FText HandleFormatText() const;
 
+	TSharedRef<SWidget> OnGenerateTexturePresetTypeEnumMenu();
+	void HandleTexturePresetTypeChanged(FString OuputName, int Index);
+	FText HandleTexturePresetTypeText() const;
+
+	TSharedRef<SWidget> OnGenerateLodGroupEnumMenu();
+	void HandleLodGroupChanged(FString OuputName, int Index);
+	FText HandleLodGroupText() const;
+
+	TSharedRef<SWidget> OnGenerateCompressionEnumMenu();
+	void HandleCompressionChanged(FString OuputName, int Index);
+	FText HandleCompressionText() const;
+
+	ECheckBoxState HandleSRGBIsChecked() const;
+	void HandleSRGBExecute(ECheckBoxState InNewState);
+
 	FReply OnBrowseClick();
+
+	bool IsDefaultPreset() const;
+
+	const int LabelSize = 150;
 };
diff --git a/Engine/Plugins/Experimental/TextureGraph/Source/TextureGraphEngine/Export/TextureExporter.cpp b/Engine/Plugins/Experimental/TextureGraph/Source/TextureGraphEngine/Export/TextureExporter.cpp
--- a/Engine/Plugins/Experimental/TextureGraph/Source/TextureGraphEngine/Export/TextureExporter.cpp
+++ b/Engine/Plugins/Experimental/TextureGraph/Source/TextureGraphEngine/Export/TextureExporter.cpp
@@ -177,6 +177,10 @@ AsyncInt TextureExporter::ExportRawAsUAsset(RawBufferPtr RawObj,const FExportMap
 			return;
 		}
 		
+		NewTexture->SRGB = Setting.IsSRGB;
+		NewTexture->LODGroup = Setting.LODGroup;
+		NewTexture->CompressionSettings = Setting.Compression;
+
 		NewTexture->Source.Init(WidthToUse, HeightToUse, 1,1, SourceFormat, RawBufferData);
 		NewTexture->UpdateResource();
 		FAssetRegistryModule::AssetCreated(NewTexture);
diff --git a/Engine/Plugins/Experimental/TextureGraph/Source/TextureGraphEngine/Export/TextureExporter.h b/Engine/Plugins/Experimental/TextureGraph/Source/TextureGraphEngine/Export/TextureExporter.h
--- a/Engine/Plugins/Experimental/TextureGraph/Source/TextureGraphEngine/Export/TextureExporter.h
+++ b/Engine/Plugins/Experimental/TextureGraph/Source/TextureGraphEngine/Export/TextureExporter.h
@@ -31,6 +31,7 @@ struct TEXTUREGRAPHENGINE_API FExportMapSettings
 	int32 Width;
 	int32 Height;
 	TextureCompressionSettings Compression;
+	TextureGroup LODGroup;
 
 	DECLARE_DELEGATE_OneParam(OnExportMapDone, FExportMapSettings MapSettings);
 	OnExportMapDone OnDone;
diff --git a/Engine/Plugins/Experimental/TextureGraph/Source/TextureGraphEngine/Model/Mix/MixSettings.h b/Engine/Plugins/Experimental/TextureGraph/Source/TextureGraphEngine/Model/Mix/MixSettings.h
--- a/Engine/Plugins/Experimental/TextureGraph/Source/TextureGraphEngine/Model/Mix/MixSettings.h
+++ b/Engine/Plugins/Experimental/TextureGraph/Source/TextureGraphEngine/Model/Mix/MixSettings.h
@@ -69,6 +69,19 @@ enum class ETG_TextureFormat : uint8
 	RGBA32F		UMETA(DisplayName = "32-bit (float) per-pixel RGBA")
 };
 
+UENUM(BlueprintType)
+enum class ETG_TexturePresetType : uint8
+{
+	None		UMETA(DisplayName = "None"), // None exposes the other settings like srgb, compression and Lod Texture group */
+	Diffuse		UMETA(DisplayName = "Diffuse"),
+	Emissive	UMETA(DisplayName = "Emissive"),
+	FX			UMETA(DisplayName = "FX"),
+	Normal		UMETA(DisplayName = "Normal"),
+	MaskComp	UMETA(DisplayName = "Mask Comp"),
+	Specular	UMETA(DisplayName = "Specular"),
+	Tangent		UMETA(DisplayName = "Tangent")
+};
+
 class UMixInterface;
 
 class RenderMesh;
