UI tweaks and updates

#tests 65b3c3fc1b5a55f3aa01fdbe

diff --git a/Engine/Plugins/Experimental/TextureGraph/Source/TextureGraphEditor/Private/EdGraph/TG_EdGraphNode.cpp b/Engine/Plugins/Experimental/TextureGraph/Source/TextureGraphEditor/Private/EdGraph/TG_EdGraphNode.cpp
--- a/Engine/Plugins/Experimental/TextureGraph/Source/TextureGraphEditor/Private/EdGraph/TG_EdGraphNode.cpp
+++ b/Engine/Plugins/Experimental/TextureGraph/Source/TextureGraphEditor/Private/EdGraph/TG_EdGraphNode.cpp
@@ -124,6 +124,10 @@ void UTG_EdGraphNode::GetNodeContextMenuActions(UToolMenu* Menu, class UGraphNod
 
 	{
 		FToolMenuSection& Section = Menu->AddSection("EdGraphSchemaGeneral", LOCTEXT("GeneralHeader", "General"));
+		if (this->GetCanRenameNode())
+		{
+			Section.AddMenuEntry(FGenericCommands::Get().Rename);
+		}
 		Section.AddMenuEntry(FGenericCommands::Get().Delete);
 		// TODO: Add also Copy, Paste, Duplicate...
 		Section.AddMenuEntry(FGenericCommands::Get().Cut);
@@ -395,13 +399,13 @@ FString UTG_EdGraphNode::GetTitleDetail()
 			if (Texture.RasterBlob)
 			{
 				const auto Desc = Texture.RasterBlob->GetDescriptor();
-				const FString FormatString = Desc.FormatToString(Desc.Format);
-				const FString Channel = TextureHelper::GetChannelsTextFromItemsPerPoint(Desc.ItemsPerPoint);
-				FString BufferString = Desc.bIsSRGB ? "sRGB" : "Linear";
-				BufferString += " - " + Channel + "_" + FormatString + "\n";
-				BufferString += FString::FromInt(Desc.Width) + "x" + FString::FromInt(Desc.Height);
-				Details = BufferString;
-				return Details;
+				FStringFormatNamedArguments Args;
+				Args.Add(TEXT("Channels"), TextureHelper::GetChannelsTextFromItemsPerPoint(Desc.ItemsPerPoint));
+				Args.Add(TEXT("Format"), Desc.FormatToString(Desc.Format));
+				Args.Add(TEXT("IsSRGB"), Desc.bIsSRGB ? "sRGB" : "Linear");
+				Args.Add(TEXT("Width"), FString::FromInt(Desc.Width));
+				Args.Add(TEXT("Height"), FString::FromInt(Desc.Height));
+				Details = FString::Format(TEXT("{Channels}_{Format}, {IsSRGB}\r\n{Width}x{Height}"), Args);
 			}
 		}
 	}
diff --git a/Engine/Plugins/Experimental/TextureGraph/Source/TextureGraphEditor/Private/EdGraph/TG_EdGraphSchema.cpp b/Engine/Plugins/Experimental/TextureGraph/Source/TextureGraphEditor/Private/EdGraph/TG_EdGraphSchema.cpp
--- a/Engine/Plugins/Experimental/TextureGraph/Source/TextureGraphEditor/Private/EdGraph/TG_EdGraphSchema.cpp
+++ b/Engine/Plugins/Experimental/TextureGraph/Source/TextureGraphEditor/Private/EdGraph/TG_EdGraphSchema.cpp
@@ -91,7 +91,6 @@ void UTG_EdGraphSchema::GetTG_ExpressionsActions(FGraphActionMenuBuilder& Action
 		if (const UTG_Expression* TG_Expression = TG_ExpressionClass->GetDefaultObject<UTG_Expression>())
 		{
 			static const FName DefaultCategory = TG_Category::Default;
-			FString testName("Test");
 
 			const FText MenuDesc = FText::FromString(TG_Expression->GetDefaultName().ToString());
 			FName Category = TG_Expression->GetCategory();
diff --git a/Engine/Plugins/Experimental/TextureGraph/Source/TextureGraphEditor/Private/STG_EditorViewport.cpp b/Engine/Plugins/Experimental/TextureGraph/Source/TextureGraphEditor/Private/STG_EditorViewport.cpp
--- a/Engine/Plugins/Experimental/TextureGraph/Source/TextureGraphEditor/Private/STG_EditorViewport.cpp
+++ b/Engine/Plugins/Experimental/TextureGraph/Source/TextureGraphEditor/Private/STG_EditorViewport.cpp
@@ -277,6 +277,8 @@ void STG_EditorViewport::Construct(const FArguments& InArgs)
 	FCoreUObjectDelegates::OnObjectPropertyChanged.AddRaw(this, &STG_EditorViewport::OnObjectPropertyChanged);
 	
 	GenerateRendermodeToolbar();
+	
+	TogglePreviewBackground();
 }
 
 STG_EditorViewport::STG_EditorViewport()
@@ -329,6 +331,7 @@ TSharedRef<FEditorViewportClient> STG_EditorViewport::MakeEditorViewportClient()
 	EditorViewportClient->SetViewLocationForOrbiting(FVector::ZeroVector, 500);
 	EditorViewportClient->bSetListenerPosition = false;
 	EditorViewportClient->EngineShowFlags.EnableAdvancedFeatures();
+	EditorViewportClient->EngineShowFlags.SetGrid(false);
 	EditorViewportClient->EngineShowFlags.SetLighting(true);
 	EditorViewportClient->EngineShowFlags.SetIndirectLightingCache(true);
 	EditorViewportClient->EngineShowFlags.SetPostProcessing(true);
diff --git a/Engine/Plugins/Experimental/TextureGraph/Source/TextureGraphEditor/Private/STG_Palette.cpp b/Engine/Plugins/Experimental/TextureGraph/Source/TextureGraphEditor/Private/STG_Palette.cpp
--- a/Engine/Plugins/Experimental/TextureGraph/Source/TextureGraphEditor/Private/STG_Palette.cpp
+++ b/Engine/Plugins/Experimental/TextureGraph/Source/TextureGraphEditor/Private/STG_Palette.cpp
@@ -152,7 +152,7 @@ void STG_PaletteTileItem::Construct(const FArguments& InArgs, FCreateWidgetForAc
 
 	FLinearColor Color = UTG_EdGraphSchema::GetCategoryColor(FName(GraphAction->GetCategory().ToString()));
 
-	BackgroundBrush = *(FTG_Style::Get().GetBrush("TG.Palette.Backgorund"));
+	BackgroundBrush = *(FTG_Style::Get().GetBrush("TG.Palette.Background"));
 	BackgroundBrush.OutlineSettings.Color = FSlateColor(Color);
 	
 	const int IconWidth = 16;
@@ -172,8 +172,7 @@ void STG_PaletteTileItem::Construct(const FArguments& InArgs, FCreateWidgetForAc
 			.HeightOverride(ItemHeight)
 			[
 				SNew(SBorder)
-				.BorderImage(FAppStyle::GetBrush("UniformShadow"))
-				.Padding(2)
+				.BorderImage(FAppStyle::GetBrush("ProjectBrowser.ProjectTile.DropShadow"))
 				[
 					SNew(SImage)
 					.Image(&BackgroundBrush)
diff --git a/Engine/Plugins/Experimental/TextureGraph/Source/TextureGraphEditor/Private/STG_SelectionPreview.cpp b/Engine/Plugins/Experimental/TextureGraph/Source/TextureGraphEditor/Private/STG_SelectionPreview.cpp
--- a/Engine/Plugins/Experimental/TextureGraph/Source/TextureGraphEditor/Private/STG_SelectionPreview.cpp
+++ b/Engine/Plugins/Experimental/TextureGraph/Source/TextureGraphEditor/Private/STG_SelectionPreview.cpp
@@ -8,6 +8,7 @@
 #include "Widgets/Layout/SScaleBox.h"
 #include "Widgets/Input/SComboButton.h"
 #include "Widgets/Input/SButton.h"
+#include <Widgets/Input/SCheckBox.h>
 #include "Engine/Texture2D.h"
 #include "EdGraph/TG_EdGraph.h"
 #include "EdGraph/TG_PinSelectionManager.h"
@@ -21,6 +22,8 @@
 #include "TG_Texture.h"
 #include "TG_Editor.h"
 #include "TG_HelperFunctions.h"
+#include "Styling/StyleColors.h"
+#include "Brushes/SlateNoResource.h"
 #include "Transform/Utility/T_TextureHistogram.h"
 #include "STextureHistogram.h"
 #include "SPrimaryButton.h"
@@ -29,6 +32,8 @@
 void STG_SelectionPreview::Construct(const FArguments& InArgs)
 {
 	const float TEXT_PADDING = 2.0;
+	UncheckedBrush = new FSlateRoundedBoxBrush(FStyleColors::Input, CoreStyleConstants::InputFocusRadius);
+	CheckedBrush = new FSlateRoundedBoxBrush(FStyleColors::InputOutline, CoreStyleConstants::InputFocusRadius);
 	OnBlobSelectionChanged = InArgs._OnBlobSelectionChanged;
 	ChildSlot
 		[
@@ -51,17 +56,7 @@ void STG_SelectionPreview::Construct(const FArguments& InArgs)
 			[
 				SNew(SHorizontalBox)
 
-				//RGBA buttons
-				+ SHorizontalBox::Slot()
-				.Padding(4.0f, 0.0f, 2.0f, 0.0f)
-				.VAlign(VAlign_Center)
-				.HAlign(HAlign_Left)
-				.FillWidth(0.1)
-				[
-					SAssignNew(RGBAButtons, STG_RGBAButtons)
-				]
-
-					/// DEBUG ONLY
+				/// DEBUG ONLY
 #if UE_BUILD_DEBUG
 					+ SHorizontalBox::Slot()
 					.VAlign(VAlign_Center)
@@ -79,10 +74,30 @@ void STG_SelectionPreview::Construct(const FArguments& InArgs)
 				+ SHorizontalBox::Slot()
 				.VAlign(VAlign_Center)
 				.HAlign(HAlign_Right)
-				.Padding(1.0f)
-				.FillWidth(0.9)
+				.AutoWidth()
 				[
-					SNew(SHorizontalBox)
+					 SNew(SBox)
+						.HAlign(HAlign_Right)
+						.VAlign(VAlign_Center)
+						[
+							SNew(SCheckBox)
+								.Style(&FAppStyle::Get().GetWidgetStyle<FCheckBoxStyle>("ToggleButtonCheckBoxAlt"))
+								.Type(ESlateCheckBoxType::ToggleButton)
+								.UncheckedImage(UncheckedBrush)
+								.UncheckedHoveredImage(UncheckedBrush)
+								.CheckedHoveredImage(CheckedBrush)
+								.CheckedImage(CheckedBrush)
+								.Padding(FMargin(4))
+								.ToolTipText(NSLOCTEXT("STG_SelectionPreview", "LockSelectionButton_ToolTip", "Locks the output of current selected node"))
+								.OnCheckStateChanged(this, &STG_SelectionPreview::OnLockChanged)
+								.IsChecked(this, &STG_SelectionPreview::IsPreviewLocked)
+								[
+									SNew(SImage)
+										.ColorAndOpacity(FSlateColor::UseForeground())
+										.Image(this, &STG_SelectionPreview::OnGetLockButtonImageResource)
+								]
+						]
+				]
 
 					//Zoom Dropdown Menu
 					+ SHorizontalBox::Slot()
@@ -94,40 +109,46 @@ void STG_SelectionPreview::Construct(const FArguments& InArgs)
 						MakeZoomControlWidget()
 					]
 
-					//Zoom Plus button
-					+ SHorizontalBox::Slot()
-					.AutoWidth()
-					.VAlign(VAlign_Center)
-					.HAlign(HAlign_Right)
-					.Padding(1.0f)
+				//Zoom Plus button
+				+ SHorizontalBox::Slot()
+				.AutoWidth()
+				.VAlign(VAlign_Center)
+				.HAlign(HAlign_Right)
+				.Padding(1.0f)
+				[
+					SNew(SButton)
+					.ButtonStyle(FAppStyle::Get(), "TextureEditor.MipmapButtonStyle")
+					.OnClicked(this, &STG_SelectionPreview::HandleZoomPlusButtonClicked)
 					[
-						SNew(SButton)
-						.ButtonStyle(FAppStyle::Get(), "TextureEditor.MipmapButtonStyle")
-						.OnClicked(this, &STG_SelectionPreview::HandleZoomPlusButtonClicked)
-						[
-							SNew(SImage)
-							.Image(FAppStyle::Get().GetBrush("Icons.Plus"))
-							.ColorAndOpacity(FSlateColor::UseForeground())
-						]
+						SNew(SImage)
+						.Image(FAppStyle::Get().GetBrush("Icons.Plus"))
+						.ColorAndOpacity(FSlateColor::UseForeground())
 					]
+				]
 
-					//Zoom Minus button
-					+ SHorizontalBox::Slot()
-					.AutoWidth()
-					.VAlign(VAlign_Center)
-					.HAlign(HAlign_Right)
-					.Padding(1.0f)
+				//Zoom Minus button
+				+ SHorizontalBox::Slot()
+				.AutoWidth()
+				.VAlign(VAlign_Center)
+				.HAlign(HAlign_Right)
+				.Padding(1.0f)
+				[
+					SNew(SButton)
+					.ButtonStyle(FAppStyle::Get(), "TextureEditor.MipmapButtonStyle")
+					.OnClicked(this, &STG_SelectionPreview::HandleZoomMinusButtonClicked)
 					[
-						SNew(SButton)
-						.ButtonStyle(FAppStyle::Get(), "TextureEditor.MipmapButtonStyle")
-						.OnClicked(this, &STG_SelectionPreview::HandleZoomMinusButtonClicked)
-						[
-							SNew(SImage)
-							.Image(FAppStyle::Get().GetBrush("Icons.Minus"))
-							.ColorAndOpacity(FSlateColor::UseForeground())
-						]
+						SNew(SImage)
+						.Image(FAppStyle::Get().GetBrush("Icons.Minus"))
+						.ColorAndOpacity(FSlateColor::UseForeground())
 					]
-
+				]
+				+ SHorizontalBox::Slot()
+				.VAlign(VAlign_Center)
+				.HAlign(HAlign_Right)
+				.Padding(1.0f)
+				.FillWidth(0.9)
+				[
+					SNew(SHorizontalBox)
 
 					//Output selection menu
 					//+ SHorizontalBox::Slot()
@@ -150,39 +171,36 @@ void STG_SelectionPreview::Construct(const FArguments& InArgs)
 					//	]
 					//]
 
-					//Lock button
+					//RGBA buttons
 					+ SHorizontalBox::Slot()
-
+					.Padding(4.0f, 0.0f, 2.0f, 0.0f)
 					.VAlign(VAlign_Center)
-					.HAlign(HAlign_Right)
-					.Padding(2.0f, 0.0f, 4.0f, 0.0f)
-					.AutoWidth()
+					.HAlign(HAlign_Left)
+					.FillWidth(0.1)
 					[
-						SAssignNew(LockButton,SButton)
-						.ButtonStyle(FAppStyle::Get(), "SimpleButton")
-						.OnClicked(this, &STG_SelectionPreview::OnLockClick)
-						.ContentPadding(FMargin(4, 2))
-						.HAlign(HAlign_Center)
-						.VAlign(VAlign_Center)
-						.ToolTipText(NSLOCTEXT("STG_SelectionPreview", "LockSelectionButton_ToolTip", "Locks the output of current selected node"))
-						.IsEnabled_Lambda([this]() { return IsLockButtonEnable(); })
-						[
-							SNew(SImage)
-							.ColorAndOpacity(FSlateColor::UseForeground())
-							.Image(this, &STG_SelectionPreview::OnGetLockButtonImageResource)
-						]
+						SAssignNew(RGBAButtons, STG_RGBAButtons)
 					]
 				]
-				
 			]
 		];
 
-
 	//Update the ViewPort Texture
 	ConstructBlobView();
 	ZoomMode = ETSZoomMode::Fill;
 }
 
+ECheckBoxState STG_SelectionPreview::IsPreviewLocked() const
+{
+	if (IsLocked)
+	{
+		return ECheckBoxState::Checked;
+	}
+	else
+	{
+		return ECheckBoxState::Unchecked;
+	}
+}
+
 FReply STG_SelectionPreview::HandleZoomPlusButtonClicked()
 {
 	ZoomIn();
@@ -447,6 +465,20 @@ FReply STG_SelectionPreview::OnLockClick()
 
 	return FReply::Handled();
 }
+void STG_SelectionPreview::OnLockChanged(const ECheckBoxState NewCheckState)
+{
+	IsLocked = NewCheckState == ECheckBoxState::Checked;
+	if (IsLocked)
+	{
+		LockedNode = SelectedNode;
+	}
+	else
+	{
+		LockedNode = nullptr;
+		//Need to update the preview because to show the output of current selected node
+		UpdatePreview();
+	}
+}
 
 const FSlateBrush* STG_SelectionPreview::OnGetLockButtonImageResource() const
 {
@@ -583,8 +615,9 @@ void STG_SelectionPreview::UpdatePixelInfo()
 				Buffer->Raw()
 					.then([=, this](RawBufferPtr raw)
 					{
-						auto Color = TextureHelper::GetPixelValueFromRaw(raw, Width, Height, X, Y);
-						PixelInfo = "UV : (" + FString::Printf(TEXT("%f"), TextureCoordinates.X) + " , " + FString::Printf(TEXT("%f"), TextureCoordinates.Y) + ")" + " Color :" + Color.ToString();
+						FLinearColor LinearColor = TextureHelper::GetPixelValueFromRaw(raw, Width, Height, X, Y);
+						FString ColorString = raw->GetDescriptor().Format == BufferFormat::Byte ? LinearColor.ToFColorSRGB().ToString() : LinearColor.ToString();   
+						PixelInfo = FString::Printf(TEXT("XY: (%d, %d), RGBA: %s"), X, Y, *ColorString);
 					});
 			}
 		}
diff --git a/Engine/Plugins/Experimental/TextureGraph/Source/TextureGraphEditor/Private/STG_SelectionPreview.h b/Engine/Plugins/Experimental/TextureGraph/Source/TextureGraphEditor/Private/STG_SelectionPreview.h
--- a/Engine/Plugins/Experimental/TextureGraph/Source/TextureGraphEditor/Private/STG_SelectionPreview.h
+++ b/Engine/Plugins/Experimental/TextureGraph/Source/TextureGraphEditor/Private/STG_SelectionPreview.h
@@ -39,6 +39,8 @@ class STG_SelectionPreview : public SCompoundWidget
 public:
 	
 	void Construct(const FArguments& InArgs);
+	ECheckBoxState IsPreviewLocked() const;
+	FText GetLockPreviewToolTip() const;
 	void OnCheckBoxStateChanged(ECheckBoxState NewState);
 	ECheckBoxState GetCheckBoxState() const;
 	void ConstructBlobView(BlobPtr InBlob = nullptr);
@@ -88,6 +90,7 @@ private:
 	TSharedRef<SWidget> OnGenerateOutputMenu();
 	FText HandleOutputText() const;
 	void HandleOutputChanged(FString OuputName);
+	void OnLockChanged(ECheckBoxState NewCheckState);
 	bool GetIsLockedNode() { return IsLocked && LockedNode != nullptr; }
 	bool IsLockButtonEnable() { return SelectedNode || LockedNode; }
 	
@@ -116,7 +119,7 @@ private:
 	FString BufferDescriptionString = "";
 	FString PixelInfo = "";
 	int32 MaxOutputs;
-	bool IsLocked;
+	bool IsLocked = false;
 	bool IsSingleChannel = false;
 	bool bSRGB = false;
 
@@ -125,6 +128,7 @@ private:
 	/** Viewport */
 	TSharedPtr<STG_TexturePreviewViewport> TextureViewport;
 	FSlateRoundedBoxBrush* CheckedBrush;
+	FSlateRoundedBoxBrush* UncheckedBrush;
 
 	/** The maximum width/height at which the texture will render in the preview window */
 	int32 PreviewEffectiveTextureWidth;
diff --git a/Engine/Plugins/Experimental/TextureGraph/Source/TextureGraphEditor/Private/STG_TextureDetails.cpp b/Engine/Plugins/Experimental/TextureGraph/Source/TextureGraphEditor/Private/STG_TextureDetails.cpp
--- a/Engine/Plugins/Experimental/TextureGraph/Source/TextureGraphEditor/Private/STG_TextureDetails.cpp
+++ b/Engine/Plugins/Experimental/TextureGraph/Source/TextureGraphEditor/Private/STG_TextureDetails.cpp
@@ -11,7 +11,6 @@
 
 void STG_TextureDetails::Construct(const FArguments& InArgs)
 {
-	const float TEXT_PADDING = 2.0;
 	CheckedBrush = new FSlateRoundedBoxBrush(FLinearColor(0.039, 0.039, 0.039, 1), CoreStyleConstants::InputFocusRadius);
 
 	ChildSlot
@@ -24,12 +23,22 @@ void STG_TextureDetails::Construct(const FArguments& InArgs)
 
 void STG_TextureDetails::MakeControls()
 {
+	const float TEXT_PADDING = 2.0;
 	VerticalBox->AddSlot()
+	.Padding(TEXT_PADDING)
 	.AutoHeight()
-	.Padding(6, 3, 3, 0)
+	.HAlign(HAlign_Fill)
 	[
-		SAssignNew(RGBAButtons,STG_RGBAButtons)
+		SNew(SHorizontalBox)
+		+ SHorizontalBox::Slot()
+		.VAlign(VAlign_Center)
+		.HAlign(HAlign_Right)
+		.Padding(5.0f, 1.0f, 3.0f, 1.0f)
+		[
+			SAssignNew(RGBAButtons,STG_RGBAButtons)
+		]
 	];
+			
 
 	AddHistogramWidget();
 
@@ -50,23 +59,10 @@ void STG_TextureDetails::AddHistogramWidget()
 		SNew(SBox)
 		.MinDesiredHeight(300) // Set your desired minimum height here
 		[
-			SNew(SBorder)
-			.BorderBackgroundColor(FLinearColor::White)  // Set the border color to white
-			.BorderImage(FCoreStyle::Get().GetBrush("WhiteBrush"))
-			.Padding(FMargin(1))
-			[
 				SNew(SOverlay)
 
 				+
 
-				SOverlay::Slot()
-				[
-					SNew(SImage)
-					.Image(FCoreStyle::Get().GetBrush("WhiteBrush"))
-					.ColorAndOpacity(FLinearColor::Black)
-				]
-
-				+
 
 				SOverlay::Slot()
 				.Padding(2, 10, 2, 2)
@@ -105,7 +101,7 @@ void STG_TextureDetails::AddHistogramWidget()
 					.Curves(ETG_HistogramCurves::Luma)
 					.Visibility(this, &STG_TextureDetails::ShowLuma)
 				]
-			]
+			
 		]
 	];
 }
diff --git a/Engine/Plugins/Experimental/TextureGraph/Source/TextureGraphEditor/Private/STG_TexturePreviewViewport.cpp b/Engine/Plugins/Experimental/TextureGraph/Source/TextureGraphEditor/Private/STG_TexturePreviewViewport.cpp
--- a/Engine/Plugins/Experimental/TextureGraph/Source/TextureGraphEditor/Private/STG_TexturePreviewViewport.cpp
+++ b/Engine/Plugins/Experimental/TextureGraph/Source/TextureGraphEditor/Private/STG_TexturePreviewViewport.cpp
@@ -41,7 +41,6 @@ void STG_TexturePreviewViewport::Construct( const FArguments& InArgs, const TSha
 				.IsEnabled(FSlateApplication::Get().GetNormalExecutionAttribute())
 				.ShowEffectWhenDisabled(false)
 				.EnableBlending(true)
-				.ToolTip(SNew(SToolTip).Text(this, &STG_TexturePreviewViewport::GetDisplayedResolution))
 			]
 			+ SOverlay::Slot()
 			.HAlign(HAlign_Right)
diff --git a/Engine/Plugins/Experimental/TextureGraph/Source/TextureGraphEditor/Private/TG_Editor.cpp b/Engine/Plugins/Experimental/TextureGraph/Source/TextureGraphEditor/Private/TG_Editor.cpp
--- a/Engine/Plugins/Experimental/TextureGraph/Source/TextureGraphEditor/Private/TG_Editor.cpp
+++ b/Engine/Plugins/Experimental/TextureGraph/Source/TextureGraphEditor/Private/TG_Editor.cpp
@@ -54,6 +54,7 @@
 #include "WorkflowOrientedApp/WorkflowTabManager.h"
 #include <Toolkits/AssetEditorToolkit.h>
 
+#include "EngineAnalytics.h"
 #include "GraphEditorActions.h"
 #include "IMessageLogListing.h"
 #include "MessageLogInitializationOptions.h"
@@ -88,7 +89,7 @@ void FTG_Editor::RegisterTabSpawners(const TSharedRef<class FTabManager>& InTabM
 		.SetGroup(WorkspaceMenuCategoryRef);
 
 	InTabManager->RegisterTabSpawner(FTG_EditorTabs::ViewportTabId, FOnSpawnTab::CreateSP(this, &FTG_Editor::SpawnTab_Viewport))
-		.SetDisplayName(LOCTEXT("ViewportTab", "Viewport"))
+		.SetDisplayName(LOCTEXT("ViewportTab", "3D Preview"))
 		.SetGroup(WorkspaceMenuCategoryRef)
 		.SetIcon(FSlateIcon(FAppStyle::GetAppStyleSetName(), "LevelEditor.Tabs.Viewports"));
 
@@ -118,7 +119,7 @@ void FTG_Editor::RegisterTabSpawners(const TSharedRef<class FTabManager>& InTabM
 		.SetIcon(FSlateIcon(FAppStyle::GetAppStyleSetName(), "LevelEditor.Tabs.Details"));
 
 	InTabManager->RegisterTabSpawner(FTG_EditorTabs::SelectionPreviewTabId, FOnSpawnTab::CreateSP(this, &FTG_Editor::SpawnTab_SelectionPreview))
-		.SetDisplayName(LOCTEXT("SelectionPreviewTab", "Selection Preview"))
+		.SetDisplayName(LOCTEXT("SelectionPreviewTab", "Node Preview"))
 		.SetGroup(WorkspaceMenuCategoryRef)
 		.SetIcon(FSlateIcon(FAppStyle::GetAppStyleSetName(), "LevelEditor.Tabs.Viewports"));
 
@@ -128,7 +129,7 @@ void FTG_Editor::RegisterTabSpawners(const TSharedRef<class FTabManager>& InTabM
 		.SetIcon(FSlateIcon(FAppStyle::GetAppStyleSetName(), "LevelEditor.Tabs.Viewports"));*/
 
 	InTabManager->RegisterTabSpawner(FTG_EditorTabs::PreviewSettingsTabId, FOnSpawnTab::CreateSP(this, &FTG_Editor::SpawnTab_Settings))
-		.SetDisplayName(LOCTEXT("PreviewSettingsTab", "Preview Settings"))
+		.SetDisplayName(LOCTEXT("PreviewSettingsTab", "3D Preview Settings"))
 		.SetGroup(WorkspaceMenuCategoryRef)
 		.SetIcon(FSlateIcon(FAppStyle::GetAppStyleSetName(), "LevelEditor.Tabs.Details"));
 
@@ -139,7 +140,7 @@ void FTG_Editor::RegisterTabSpawners(const TSharedRef<class FTabManager>& InTabM
 		.SetIcon(FSlateIcon(FAppStyle::Get().GetStyleSetName(), "MaterialEditor.TogglePlatformStats.Tab"));
 
 	InTabManager->RegisterTabSpawner(FTG_EditorTabs::TextureDetailsTabId, FOnSpawnTab::CreateSP(this, &FTG_Editor::SpawnTab_TextureDetails))
-		.SetDisplayName(LOCTEXT("TextureDetailsTab", "Texture Details"))
+		.SetDisplayName(LOCTEXT("TextureDetailsTab", "Node Histogram"))
 		.SetGroup(WorkspaceMenuCategoryRef)
 		.SetIcon(FSlateIcon(FAppStyle::GetAppStyleSetName(), "LevelEditor.Tabs.Viewports"));
 }
@@ -321,6 +322,16 @@ void FTG_Editor::InitEditor(const EToolkitMode::Type Mode, const TSharedPtr< cla
 
 	FAssetEditorToolkit::InitAssetEditor(Mode, InitToolkitHost, TG_EditorAppIdentifier, StandaloneDefaultLayout, /*bCreateDefaultToolbar*/ true, /*bCreateDefaultStandaloneMenu*/ true, InTextureGraph);
 
+	RegenerateMenusAndToolbars();
+
+	bool bViewportIsOff = !ViewportTabContent.IsValid();
+	TSharedPtr<SDockTab> ViewportTab; 
+	// check here if 3d viewport is turned off, we need to turn it on temporarily to initialize our systems correctly
+	if (bViewportIsOff)
+	{
+		ViewportTab = GetTabManager()->TryInvokeTab(FTG_EditorTabs::ViewportTabId);	
+	}
+	
 	// Set the preview mesh for the material.  This call must occur after the toolbar is initialized.
 	if (!SetPreviewAssetByName(*EditedTextureGraph->PreviewMesh.ToString()))
 	{
@@ -328,9 +339,23 @@ void FTG_Editor::InitEditor(const EToolkitMode::Type Mode, const TSharedPtr< cla
 		GetEditorViewport()->InitPreviewMesh();
 	}
 
-	RegenerateMenusAndToolbars();
+	if (bViewportIsOff)
+	{
+		// turn the viewporttab back off
+		ViewportTab->RequestCloseTab();
+	}
 	
+	bShowPaletteView = PaletteTab.IsValid();
+	bShowNodeHistogram = NodeHistogramTab.IsValid();
+
 	EditedTextureGraph->GetPackage()->ClearDirtyFlag();
+
+	// Add analytics tag
+	if (FEngineAnalytics::IsAvailable())
+	{
+		SessionStartTime = FDateTime::Now();
+		FEngineAnalytics::GetProvider().RecordEvent(TEXT("Editor.Usage.TextureGraph"), TEXT("SessionStarted"), OriginalTextureGraph.GetName());
+	}
 }
 
 FTG_Editor::FTG_Editor()
@@ -430,74 +455,62 @@ void FTG_Editor::RegisterToolbar()
 
 	const FTG_EditorCommands& TG_EditorCommands = FTG_EditorCommands::Get();
 
-	auto AutoRun = MakeAutoRunWidget();
-
 	const FToolMenuInsert InsertAfterAssetSection("Asset", EToolMenuInsertType::After);
 	{
 		FToolMenuSection& Section = ToolBar->AddSection("TG_Toolbar", TAttribute<FText>(), InsertAfterAssetSection);
-
-		Section.AddEntry(FToolMenuEntry::InitWidget("Auto Run", AutoRun.ToSharedRef(), FText::FromString("Auto Run")));
-
+		
 		Section.AddEntry(FToolMenuEntry::InitToolBarButton(
-			FTG_EditorCommands::Get().RunGraph,
+			FTG_EditorCommands::Get().ExportAsUAsset,
 			TAttribute<FText>(),
 			TAttribute<FText>(),
-			FSlateIcon(FAppStyle::GetAppStyleSetName(), "SourceControl.Actions.Refresh")));
-
-#if !UE_BUILD_SHIPPING
+			FSlateIcon(FAppStyle::GetAppStyleSetName(), "Icons.Toolbar.Export")));
+		
 		Section.AddEntry(FToolMenuEntry::InitToolBarButton(
-			FTG_EditorCommands::Get().LogGraph,
+			FTG_EditorCommands::Get().RunGraph,
 			TAttribute<FText>(),
-			FText::FromString("Log current state of Texture Graph (only available in non-shipping builds)"),
-			FSlateIcon(FAppStyle::GetAppStyleSetName(), "SourceControl.Actions.Refresh")));
-#endif
+			TAttribute<FText>(),
+			FSlateIcon(FAppStyle::GetAppStyleSetName(), "FontEditor.Update")));
 
 		Section.AddEntry(FToolMenuEntry::InitToolBarButton(
-			FTG_EditorCommands::Get().ExportAsUAsset,
+			FTG_EditorCommands::Get().AutoUpdateGraph,
 			TAttribute<FText>(),
 			TAttribute<FText>(),
-			FSlateIcon()));
-	}
-}
-
-TSharedPtr<class SWidget> FTG_Editor::MakeAutoRunWidget()
-{
-	TSharedPtr<class SWidget> AutoRun = SNew(SBox)
-		.WidthOverride(100.0f)
-		[
-			SNew(SHorizontalBox)
-				//.IsEnabled(this, &FTG_Editor::HandleMipLevelCheckBoxIsEnabled)
-				+ SHorizontalBox::Slot()
-				.Padding(4.0f, 0.0f, 2.0f, 0.0f)
-				.VAlign(VAlign_Center)
-				.AutoWidth()
-				[
-					SNew(SCheckBox)
-						.IsChecked(this, &FTG_Editor::HandleAutoRunActionIsChecked)
-						.OnCheckStateChanged(this, &FTG_Editor::HandleAutoRunActionExecute)
-				]
-
-				+ SHorizontalBox::Slot()
+			FSlateIcon(FAppStyle::GetAppStyleSetName(), "FontEditor.Update")));
+		// Section.AddEntry(FToolMenuEntry::InitWidget("Auto Update", AutoRun.ToSharedRef(), FText::FromString("Auto Update")));
 
-				.VAlign(VAlign_Center)
-				.Padding(2.0f, 0.0f, 4.0f, 0.0f)
-				[
-					SNew(STextBlock)
-						.Text(FText::FromString("Auto Run"))
-				]
-		];
+#if UE_BUILD_DEBUG
+		Section.AddEntry(FToolMenuEntry::InitToolBarButton(
+			FTG_EditorCommands::Get().LogGraph,
+			TAttribute<FText>(),
+			FText::FromString("Log current state of Texture Graph (only available in debug builds)"),
+			FSlateIcon(FAppStyle::GetAppStyleSetName(), "FontEditor.Update")));
+#endif
 
-	return AutoRun;
+		{
+			FToolMenuSection& TabsSection = ToolBar->AddSection("TabWindows", TAttribute<FText>(), InsertAfterAssetSection);
+			TabsSection.AddEntry(FToolMenuEntry::InitToolBarButton(
+				FTG_EditorCommands::Get().TogglePaletteTab,
+				TAttribute<FText>(),
+				TAttribute<FText>(),
+				FSlateIcon(FAppStyle::GetAppStyleSetName(), "LevelEditor.Tabs.Details")));
+
+			TabsSection.AddEntry(FToolMenuEntry::InitToolBarButton(
+				FTG_EditorCommands::Get().ToggleNodeHistogramTab,
+				TAttribute<FText>(),
+				TAttribute<FText>(),
+				FSlateIcon(FAppStyle::GetAppStyleSetName(), "Kismet.Tabs.Palette")));
+		}
+	}
 }
 
-ECheckBoxState FTG_Editor::HandleAutoRunActionIsChecked() const
+bool FTG_Editor::IsShowingAutoUpdate() const
 {
-	return bAutoRunGraph ? ECheckBoxState::Checked : ECheckBoxState::Unchecked;
+	return bAutoRunGraph;
 }
 
-void FTG_Editor::HandleAutoRunActionExecute(ECheckBoxState InNewState)
+void FTG_Editor::ToggleAutoUpdate()
 {
-	bAutoRunGraph = InNewState == ECheckBoxState::Checked;
+	bAutoRunGraph = !bAutoRunGraph;
 	if (bAutoRunGraph)
 	{
 		OnRunGraph_Clicked();
@@ -520,6 +533,7 @@ void FTG_Editor::PostInitAssetEditor()
 	}
 }
 
+
 void FTG_Editor::BindCommands()
 {
 	const FTG_EditorCommands& TGEditorCommands = FTG_EditorCommands::Get();
@@ -528,6 +542,13 @@ void FTG_Editor::BindCommands()
 		TGEditorCommands.RunGraph,
 		FExecuteAction::CreateSP(this, &FTG_Editor::OnRunGraph_Clicked),
 		FCanExecuteAction::CreateSP(this, &FTG_Editor::CanEnableOnRun));
+
+	ToolkitCommands->MapAction(
+		TGEditorCommands.AutoUpdateGraph,
+		FExecuteAction::CreateSP(this, &FTG_Editor::ToggleAutoUpdate),
+		FCanExecuteAction(),
+		FIsActionChecked::CreateSP(this, &FTG_Editor::IsShowingAutoUpdate));
+	
 	ToolkitCommands->MapAction(
 		TGEditorCommands.LogGraph,
 		FExecuteAction::CreateSP(this, &FTG_Editor::OnLogGraph_Clicked));
@@ -543,7 +564,19 @@ void FTG_Editor::BindCommands()
 	ToolkitCommands->MapAction(
 		FGenericCommands::Get().Redo,
 		FExecuteAction::CreateSP(this, &FTG_Editor::RedoGraphAction));
+	
+	ToolkitCommands->MapAction(
+		TGEditorCommands.TogglePaletteTab,
+		FExecuteAction::CreateSP(this, &FTG_Editor::TogglePaletteView),
+		FCanExecuteAction(),
+		FIsActionChecked::CreateSP(this, &FTG_Editor::IsShowingPaletteView));
 
+	ToolkitCommands->MapAction(
+		TGEditorCommands.ToggleNodeHistogramTab,
+		FExecuteAction::CreateSP(this, &FTG_Editor::ToggleNodeHistogramView),
+		FCanExecuteAction(),
+		FIsActionChecked::CreateSP(this, &FTG_Editor::IsShowingNodeHistogramView));
+	
 	GraphEditorCommands->MapAction(
 		TGEditorCommands.ConvertInputParameterToConstant,
 		FExecuteAction::CreateSP(this, &FTG_Editor::OnConvertInputParameterToFromConstant));
@@ -552,6 +585,63 @@ void FTG_Editor::BindCommands()
 		TGEditorCommands.ConvertInputParameterFromConstant,
 		FExecuteAction::CreateSP(this, &FTG_Editor::OnConvertInputParameterToFromConstant));
 
+	GraphEditorCommands->MapAction(
+		FGenericCommands::Get().Rename,
+		FExecuteAction::CreateSP(this, &FTG_Editor::OnRenameNodeClicked),
+		FCanExecuteAction::CreateSP(this, &FTG_Editor::CanRenameNode));
+
+}
+
+void FTG_Editor::SetShowNodeHistogramView(const bool bValue)
+{
+	bShowNodeHistogram = bValue;
+
+	// open/close 3DPreview tab
+	DisplayNodeHistogramView(bShowNodeHistogram);
+}
+
+void FTG_Editor::SetShowPaletteView(const bool bValue)
+{
+	bShowPaletteView = bValue;
+
+	// open/close 3DPreview tab
+	DisplayPaletteView(bShowPaletteView);
+}
+
+void FTG_Editor::ToggleNodeHistogramView()
+{
+	// Toggle the showing of material Texture Details each time the user presses the show 3DPreview button
+	SetShowNodeHistogramView(!bShowNodeHistogram);
+}
+
+void FTG_Editor::TogglePaletteView()
+{
+	// Toggle the showing of material 3DPreview each time the user presses the show 3DPreview button
+	SetShowPaletteView(!bShowPaletteView);
+}
+
+void FTG_Editor::DisplayPaletteView(const bool bShow)
+{
+	if (bShow)
+	{
+		GetTabManager()->TryInvokeTab(FTG_EditorTabs::PaletteTabId);
+	}
+	else if (!bShowPaletteView && PaletteTab.IsValid())
+	{
+		PaletteTab.Pin()->RequestCloseTab();
+	}
+}
+
+void FTG_Editor::DisplayNodeHistogramView(const bool bShow)
+{
+	if (bShow)
+	{
+		GetTabManager()->TryInvokeTab(FTG_EditorTabs::TextureDetailsTabId);
+	}
+	else if (!bShowNodeHistogram && NodeHistogramTab.IsValid())
+	{
+		NodeHistogramTab.Pin()->RequestCloseTab();
+	}
 }
 
 void FTG_Editor::OnRunGraph_Clicked()
@@ -562,7 +652,6 @@ void FTG_Editor::OnRunGraph_Clicked()
 		EditedTextureGraph->TriggerUpdate(false);
 	}
 }
-
 void FTG_Editor::OnRenderingDone(UMixInterface* TextureGraph, const FInvalidationDetails* Details)
 {
 	if (TextureGraph != nullptr)
@@ -704,6 +793,31 @@ TArray<UTG_EdGraphNode*>	FTG_Editor::GetCurrentSelectedTG_EdGraphNodes() const
 	return EdGraphNodes;
 }
 
+void FTG_Editor::OnRenameNodeClicked()
+{
+	const FGraphPanelSelectionSet SelectedNodes = GraphEditorWidget->GetSelectedNodes();
+	for (FGraphPanelSelectionSet::TConstIterator NodeIt(SelectedNodes); NodeIt; ++NodeIt)
+	{
+		UEdGraphNode* SelectedNode = Cast<UEdGraphNode>(*NodeIt);
+		if (SelectedNode != nullptr && SelectedNode->GetCanRenameNode())
+		{
+			bool ToRename = true;
+			GraphEditorWidget->IsNodeTitleVisible(SelectedNode, ToRename);
+			break;
+		}
+	}
+}
+
+bool FTG_Editor::CanRenameNode() const
+{
+	TArray<UTG_EdGraphNode*> SelectedNodes = GetCurrentSelectedTG_EdGraphNodes();
+	if ( SelectedNodes.Num() == 1)
+	{
+		return SelectedNodes[0]->bCanRenameNode;
+	}
+
+	return false;
+}
 void FTG_Editor::OnConvertInputParameterToFromConstant()
 {
 	TArray<UTG_EdGraphNode*> EdGraphNodes = GetCurrentSelectedTG_EdGraphNodes();
@@ -791,6 +905,7 @@ TSharedRef<class SGraphEditor> FTG_Editor::CreateGraphEditorWidget()
 		.ShowGraphStateOverlay(false)
 		.AutoExpandActionMenu(true)
 		.AssetEditorToolkit(this->AsShared());
+	
 }
 
 FActionMenuContent FTG_Editor::OnCreateGraphActionMenu(UEdGraph* InGraph, const FVector2D& InNodePosition, const TArray<UEdGraphPin*>& InDraggedPins, bool bAutoExpand, SGraphEditor::FActionMenuClosed InOnMenuClosed)
@@ -837,7 +952,6 @@ TSharedRef<SDockTab> FTG_Editor::SpawnTab_GraphEditor(const FSpawnTabArgs& Args)
 	check(Args.GetTabId() == FTG_EditorTabs::GraphEditorId);
 
 	return SNew(SDockTab)
-		.Label(LOCTEXT("TextureGraphTitle", "Graph"))
 		.TabColorScale(GetTabColorScale())
 		[
 			GraphEditorWidget.ToSharedRef()
@@ -864,7 +978,14 @@ TSharedRef<SDockTab> FTG_Editor::SpawnTab_Viewport(const FSpawnTabArgs& Args)
 
 	const FString LayoutId = FString("TG_EditorViewport");
 	ViewportTabContent->Initialize(MakeViewportFunc, DockableTab, LayoutId);
-
+	
+	// Set the preview mesh for the material.  This call must occur after the toolbar is initialized.
+	if (!SetPreviewAssetByName(*EditedTextureGraph->PreviewMesh.ToString()))
+	{
+		// The material preview mesh couldn't be found or isn't loaded.  Default to the one of the primitive types.
+		GetEditorViewport()->InitPreviewMesh();
+	
+	}
 	return DockableTab;
 }
 
@@ -873,7 +994,6 @@ TSharedRef<SDockTab> FTG_Editor::SpawnTab_TG_Properties(const FSpawnTabArgs& Arg
 	check(Args.GetTabId() == FTG_EditorTabs::PropertiesTabId);
 
 	TSharedPtr<SDockTab> DetailsTab = SNew(SDockTab)
-		.Label(LOCTEXT("TSDetailsTitle", "Details"))
 		[
 			DetailsView.ToSharedRef()
 		];
@@ -887,19 +1007,19 @@ TSharedRef<SDockTab> FTG_Editor::SpawnTab_Palette(const FSpawnTabArgs& Args)
 {
 	check(Args.GetTabId() == FTG_EditorTabs::PaletteTabId);
 
-	return SNew(SDockTab)
-		.Label(LOCTEXT("TSPaletteTitle", "Palette"))
+	TSharedRef<SDockTab> DockableTab = SNew(SDockTab)
 		[
 			Palette.ToSharedRef()
 		];
+	PaletteTab = DockableTab;
+	return DockableTab;
 }
 
 TSharedRef<SDockTab> FTG_Editor::SpawnTab_Find(const FSpawnTabArgs& Args)
 {
 	check(Args.GetTabId() == FTG_EditorTabs::FindTabId);
 
-	return SNew(SDockTab)
-		.Label(LOCTEXT("TSFindTitle", "Find Results"));
+	return SNew(SDockTab);
 }
 
 TSharedRef<SDockTab> FTG_Editor::SpawnTab_PreviewSettings(const FSpawnTabArgs& Args)
@@ -915,14 +1035,13 @@ TSharedRef<SDockTab> FTG_Editor::SpawnTab_PreviewSettings(const FSpawnTabArgs& A
 	}
 
 	TSharedRef<SDockTab> SpawnedTab = SNew(SDockTab)
-		.Label(LOCTEXT("TG_PreviewSceneSettingsTab", "Preview Scene Settings"))
 		[
 			SNew(SBox)
 				[
 					InWidget
 				]
 		];
-
+	
 	return SpawnedTab;
 }
 
@@ -931,7 +1050,6 @@ TSharedRef<SDockTab> FTG_Editor::SpawnTab_ParameterDefaults(const FSpawnTabArgs&
 	check(Args.GetTabId() == FTG_EditorTabs::ParameterDefaultsTabId);
 
 	return SNew(SDockTab)
-		.Label(LOCTEXT("Parameters", "Parameters"))
 		[
 			SNew(SBox)
 				[
@@ -944,12 +1062,14 @@ TSharedRef<SDockTab> FTG_Editor::SpawnTab_SelectionPreview(const FSpawnTabArgs&
 {
 	check(Args.GetTabId() == FTG_EditorTabs::SelectionPreviewTabId);
 
-	return SNew(SDockTab)
-		.Label(LOCTEXT("SelectionView", "Selection Preview"))
+	TSharedPtr<SDockTab> Tab = SNew(SDockTab)
 		.TabColorScale(GetTabColorScale())
 		[
 			SelectionPreview.ToSharedRef()
 		];
+	NodeHistogramTab = Tab;
+
+	return Tab.ToSharedRef();
 
 }
 
@@ -1035,7 +1155,6 @@ TSharedRef<SDockTab> FTG_Editor::SpawnTab_Settings(const FSpawnTabArgs& Args)
 	check(Args.GetTabId() == FTG_EditorTabs::PreviewSettingsTabId);
 
 	TSharedPtr<SDockTab> SettingsTab = SNew(SDockTab)
-		.Label(LOCTEXT("Preview Settings", "Preview Settings"))
 		[
 			SettingsView.ToSharedRef()
 		];
@@ -1050,7 +1169,6 @@ TSharedRef<SDockTab> FTG_Editor::SpawnTab_Errors(const FSpawnTabArgs& Args)
 	check(Args.GetTabId() == FTG_EditorTabs::ErrorsTabId);
 
 	TSharedPtr<SDockTab> ErrorTab = SNew(SDockTab)
-		.Label(LOCTEXT("Errors", "Errors"))
 		[
 			ErrorsWidget.ToSharedRef()
 		];
@@ -1062,13 +1180,14 @@ TSharedRef<SDockTab> FTG_Editor::SpawnTab_TextureDetails(const FSpawnTabArgs& Ar
 {
 	check(Args.GetTabId() == FTG_EditorTabs::TextureDetailsTabId);
 
-	return SNew(SDockTab)
-		.Label(LOCTEXT("TextureDetails", "Texture Details"))
+	TSharedRef<SDockTab> SpawnedTab = SNew(SDockTab)
 		.TabColorScale(GetTabColorScale())
 		[
 			TextureDetails.ToSharedRef()
 		];
 
+	NodeHistogramTab = SpawnedTab;
+	return SpawnedTab;
 }
 
 TSharedPtr<STG_EditorViewport> FTG_Editor::GetEditorViewport() const
@@ -1466,6 +1585,23 @@ bool FTG_Editor::OnRequestClose(EAssetEditorCloseReason InCloseReason)
 
 	return true;
 }
+
+void FTG_Editor::OnClose()
+{
+	// Add analytics tag
+	if (FEngineAnalytics::IsAvailable())
+	{
+		TArray<FAnalyticsEventAttribute> Attributes;
+		Attributes.Add(FAnalyticsEventAttribute(TEXT("GraphName"), OriginalTextureGraph.GetName()));
+		Attributes.Add(FAnalyticsEventAttribute(TEXT("TimeActive"),  (FDateTime::Now() - SessionStartTime).ToString()));
+				
+		// Send Analytics event 
+		FEngineAnalytics::GetProvider().RecordEvent(TEXT("Editor.Usage.TextureGraph.SessionEnded"), Attributes);
+				
+	}
+	ITG_Editor::OnClose();
+}
+
 bool FTG_Editor::UpdateOriginalTextureGraph()
 {
 	// TODO : We should cancel saving when TextureGraph has errors.
diff --git a/Engine/Plugins/Experimental/TextureGraph/Source/TextureGraphEditor/Private/TG_Editor.h b/Engine/Plugins/Experimental/TextureGraph/Source/TextureGraphEditor/Private/TG_Editor.h
--- a/Engine/Plugins/Experimental/TextureGraph/Source/TextureGraphEditor/Private/TG_Editor.h
+++ b/Engine/Plugins/Experimental/TextureGraph/Source/TextureGraphEditor/Private/TG_Editor.h
@@ -131,6 +131,7 @@ protected:
 	virtual void									SaveAsset_Execute() override;
 	virtual void									SaveAssetAs_Execute() override;
 	virtual bool									OnRequestClose(EAssetEditorCloseReason InCloseReason) override;
+	virtual void									OnClose() override;
 	// ~End FAssetEditorToolkit interface
 
 	bool											UpdateOriginalTextureGraph();
@@ -141,8 +142,12 @@ private:
 
 	virtual void									PostInitAssetEditor() override;
 
+	void 											OnRenameNodeClicked();
+	bool 											CanRenameNode() const;
 	/** Bind commands to delegates */
 	void											BindCommands();
+	void 											SetShowNodeHistogramView(bool bValue);
+	void 											SetShowPaletteView(bool bValue);
 
 	void											OnGraphChanged(UTG_Graph* InGraph, UTG_Node* InNode, bool Tweaking);
 
@@ -155,10 +160,10 @@ private:
 	void											OnPinSelectionUpdated(UEdGraphPin* Pin);
 	/** Run Graph by invoking the graph notifications  */
 	void											OnRunGraph_Clicked();
-	TSharedPtr<class SWidget>						MakeAutoRunWidget();
 
-	ECheckBoxState									HandleAutoRunActionIsChecked() const;
-	void											HandleAutoRunActionExecute(ECheckBoxState InNewState);
+	
+	void											ToggleAutoUpdate();
+	bool											IsShowingAutoUpdate() const;
 
 	void											OnRenderingDone(UMixInterface* TextureGraph, const FInvalidationDetails* Details);
 	void											OnViewportSettingsChanged();
@@ -190,8 +195,6 @@ private:
 	/** Create Texture details widget */
 	TSharedRef<class STG_TextureDetails>			CreateTextureDetailsWidget();
 
-	TSharedRef<class STG_GraphParametersOverviewPanel>	CreateGraphParametersOverviewWidget();
-
 	/** Gets the current TG_ Graph's appearance */
 	FGraphAppearanceInfo							GetGraphAppearance() const;
 
@@ -254,19 +257,28 @@ private:
 	TSharedRef<SDockTab>							SpawnTab_Settings(const FSpawnTabArgs& Args);
 	TSharedRef<SDockTab>							SpawnTab_Errors(const FSpawnTabArgs& Args);
 	TSharedRef<SDockTab>							SpawnTab_TextureDetails(const FSpawnTabArgs& Args);
-
+	void 											DisplayNodeHistogramView(bool bShow3DPreview);
+	
+	void 											ToggleNodeHistogramView();
+	void 											TogglePaletteView();
+	void 											DisplayPaletteView(bool bShowSelectionPreview);
 	void											ExportAsUAsset();
 	FReply											OnExportClick();
 
 	TSharedPtr<class STG_EditorViewport>			GetEditorViewport() const;	
-	void											GeneratorRowsRefreshed();
+	
 	bool											CanEnableOnRun();
 
 	void											UpdateGenerator();
 	void											RunGraph(bool Tweaking);
 	void											JumpToNode(const UEdGraphNode* Node);
 	void											OnMessageLogLinkActivated(const class TSharedRef<IMessageToken>& Token);
+	// FEditorUndoClient Interface
+	virtual void 									PostUndo(bool bSuccess) override;
+	virtual void 									PostRedo(bool bSuccess) override { PostUndo(bSuccess); }
 	
+	FORCEINLINE bool 								IsShowingNodeHistogramView() const { return bShowNodeHistogram; }
+	FORCEINLINE bool 								IsShowingPaletteView() const { return bShowPaletteView; };
 public:
 	/** Original Texture Graph */
 	TObjectPtr<UTextureGraph>						OriginalTextureGraph;
@@ -274,13 +286,13 @@ public:
 	/** Duplicated Texture Graph used in the editor */
 	TObjectPtr<UTextureGraph>						EditedTextureGraph;
 	TObjectPtr<UTG_EdGraph>							TG_EdGraph;
+	FDateTime										SessionStartTime;
+	bool 											bShowNodeHistogram = false;
+	bool 											bShowPaletteView = false;
 
 private:
 	TObjectPtr<UTG_Parameters>						TG_Parameters;
 	FExportSettings                                 TargetExportSettings;
-	// FEditorUndoClient Interface
-	virtual void PostUndo(bool bSuccess) override;
-	virtual void PostRedo(bool bSuccess) override { PostUndo(bSuccess); }
 
 	/** Property View */
 	TSharedPtr<class IDetailsView>					DetailsView;
@@ -314,6 +326,8 @@ private:
 
 	/** Tab that holds the details panel */
 	TWeakPtr<SDockTab>								SpawnedDetailsTab;	
+	TWeakPtr<SDockTab>								NodeHistogramTab;	
+	TWeakPtr<SDockTab>								PaletteTab;	
 
 	/** Stats log, with the log listing that it reflects */
 	TSharedPtr<class SWidget>						ErrorsWidget;
diff --git a/Engine/Plugins/Experimental/TextureGraph/Source/TextureGraphEditor/Private/TG_EditorCommands.cpp b/Engine/Plugins/Experimental/TextureGraph/Source/TextureGraphEditor/Private/TG_EditorCommands.cpp
--- a/Engine/Plugins/Experimental/TextureGraph/Source/TextureGraphEditor/Private/TG_EditorCommands.cpp
+++ b/Engine/Plugins/Experimental/TextureGraph/Source/TextureGraphEditor/Private/TG_EditorCommands.cpp
@@ -20,9 +20,10 @@ FTG_EditorCommands::FTG_EditorCommands()
 
 void FTG_EditorCommands::RegisterCommands()
 {
-	UI_COMMAND(RunGraph, "Run", "Run the current graph.", EUserInterfaceActionType::Button, FInputChord());
+	UI_COMMAND(RunGraph, "Update", "Update the current graph.", EUserInterfaceActionType::Button, FInputChord());
+	UI_COMMAND(AutoUpdateGraph, "Auto Update", "Toggles to update graph automatically or not.", EUserInterfaceActionType::ToggleButton, FInputChord());
 	UI_COMMAND(LogGraph, "Log", "Log the current graph.", EUserInterfaceActionType::Button, FInputChord());
-	UI_COMMAND(ExportAsUAsset, "Export", "Export all the mix textures as UAssets", EUserInterfaceActionType::Button, FInputChord());
+	UI_COMMAND(ExportAsUAsset, "Export", "Export all the output nodes as Texture UAssets", EUserInterfaceActionType::Button, FInputChord());
 	UI_COMMAND(ConvertInputParameterToConstant, "Convert input to Constant", "Convert input parameter to or from a constant.", EUserInterfaceActionType::Button, FInputChord());
 	UI_COMMAND(ConvertInputParameterFromConstant, "Convert input to Parameter", "Convert input parameter to or from a constant.", EUserInterfaceActionType::Button, FInputChord());
 	UI_COMMAND(SetCylinderPreview, "Cylinder", "Sets the preview mesh to a cylinder primitive.", EUserInterfaceActionType::ToggleButton, FInputChord());
@@ -32,6 +33,9 @@ void FTG_EditorCommands::RegisterCommands()
 	UI_COMMAND(SetPreviewMeshFromSelection, "Mesh", "Sets the preview mesh based on the current content browser selection.", EUserInterfaceActionType::ToggleButton, FInputChord());
 	UI_COMMAND(TogglePreviewGrid, "Grid", "Toggles the preview pane's grid.", EUserInterfaceActionType::ToggleButton, FInputChord());
 	UI_COMMAND(TogglePreviewBackground, "Background", "Toggles the preview pane's background.", EUserInterfaceActionType::ToggleButton, FInputChord());
+	UI_COMMAND(TogglePaletteTab, "Palette", "Toggles the Palette window.", EUserInterfaceActionType::ToggleButton, FInputChord());
+	UI_COMMAND(ToggleNodeHistogramTab, "Node Histogram", "Toggles the Node Histogram window.", EUserInterfaceActionType::ToggleButton, FInputChord() );
+	
 }
 
 #undef LOCTEXT_NAMESPACE
diff --git a/Engine/Plugins/Experimental/TextureGraph/Source/TextureGraphEditor/Private/TG_EditorCommands.h b/Engine/Plugins/Experimental/TextureGraph/Source/TextureGraphEditor/Private/TG_EditorCommands.h
--- a/Engine/Plugins/Experimental/TextureGraph/Source/TextureGraphEditor/Private/TG_EditorCommands.h
+++ b/Engine/Plugins/Experimental/TextureGraph/Source/TextureGraphEditor/Private/TG_EditorCommands.h
@@ -15,8 +15,11 @@ public:
 	// ~End TCommands<> interface
 
 	TSharedPtr<FUICommandInfo>									RunGraph;
+	TSharedPtr<FUICommandInfo>									AutoUpdateGraph;
 	TSharedPtr<FUICommandInfo>									LogGraph;
 	TSharedPtr<FUICommandInfo>									ExportAsUAsset;
+	TSharedPtr<FUICommandInfo>									TogglePaletteTab;
+	TSharedPtr<FUICommandInfo>									ToggleNodeHistogramTab;
 
 	/**
 	 * Node Actions
diff --git a/Engine/Plugins/Experimental/TextureGraph/Source/TextureGraphEditor/Private/TG_Style.cpp b/Engine/Plugins/Experimental/TextureGraph/Source/TextureGraphEditor/Private/TG_Style.cpp
--- a/Engine/Plugins/Experimental/TextureGraph/Source/TextureGraphEditor/Private/TG_Style.cpp
+++ b/Engine/Plugins/Experimental/TextureGraph/Source/TextureGraphEditor/Private/TG_Style.cpp
@@ -139,7 +139,7 @@ FTG_Style::FTG_Style() : FSlateStyleSet("TG_Style")
 	Set("TG.Graph.Node.Header", new FSlateRoundedBoxBrush(NoSpillColor, NodeHeaderRadius, NoSpillColor, 2.0));
 	Set("TG.Graph.Node.ShadowSelected", new BOX_BRUSH("Style/TG_shadow_selected", FMargin(18.0/64.0)));
 
-	Set("TG.Palette.Backgorund", new FSlateRoundedBoxBrush(NoSpillColor, PalleteRadius, NoSpillColor, 2.0));
+	Set("TG.Palette.Background", new FSlateRoundedBoxBrush(NoSpillColor, PalleteRadius, NoSpillColor, 2.0));
 	
 	SetParentStyleName("EditorStyle");
 	FTextBlockStyle NormalText = GetParentStyle()->GetWidgetStyle<FTextBlockStyle>("NormalText");
diff --git a/Engine/Plugins/Experimental/TextureGraph/Source/TextureGraphEditor/Private/Widgets/ActionMenuWidgets/STG_GraphActionMenu.cpp b/Engine/Plugins/Experimental/TextureGraph/Source/TextureGraphEditor/Private/Widgets/ActionMenuWidgets/STG_GraphActionMenu.cpp
--- a/Engine/Plugins/Experimental/TextureGraph/Source/TextureGraphEditor/Private/Widgets/ActionMenuWidgets/STG_GraphActionMenu.cpp
+++ b/Engine/Plugins/Experimental/TextureGraph/Source/TextureGraphEditor/Private/Widgets/ActionMenuWidgets/STG_GraphActionMenu.cpp
@@ -37,7 +37,7 @@ void STG_GraphActionMenu::Construct(const FArguments& InArgs)
 
 	if (bSpwanOnSelect)
 	{
-		Brush = *FTG_Style::Get().GetBrush("TG.Palette.Backgorund");
+		Brush = *FTG_Style::Get().GetBrush("TG.Palette.Background");
 		Brush.OutlineSettings.Color = FSlateColor(FillColor * UTG_EdGraphSchema::NodeOutlineColorMultiplier);
 	}
 	else
diff --git a/Engine/Plugins/Experimental/TextureGraph/Source/TextureGraphEditor/Private/Widgets/STG_RGBAButtons.cpp b/Engine/Plugins/Experimental/TextureGraph/Source/TextureGraphEditor/Private/Widgets/STG_RGBAButtons.cpp
--- a/Engine/Plugins/Experimental/TextureGraph/Source/TextureGraphEditor/Private/Widgets/STG_RGBAButtons.cpp
+++ b/Engine/Plugins/Experimental/TextureGraph/Source/TextureGraphEditor/Private/Widgets/STG_RGBAButtons.cpp
@@ -3,7 +3,7 @@
 #include "STG_RGBAButtons.h"
 #include "Widgets/Input/SComboButton.h"
 #include "Brushes/SlateImageBrush.h"
-
+#include "Styling/StyleColors.h"
 #define LOCTEXT_NAMESPACE "TextureGraphEditor"
 /** Constructs this widget with InArgs */
 void STG_RGBAButtons::Construct(const FArguments& InArgs)
@@ -16,7 +16,7 @@ void STG_RGBAButtons::Construct(const FArguments& InArgs)
 
 TSharedRef<SWidget> STG_RGBAButtons::MakeChannelControlWidget()
 {
-	CheckedBrush = new FSlateRoundedBoxBrush(FLinearColor(0.039, 0.039, 0.039, 1), CoreStyleConstants::InputFocusRadius);
+	CheckedBrush = new FSlateRoundedBoxBrush(FStyleColors::Input, CoreStyleConstants::InputFocusRadius);
 
 	TSharedRef<SWidget> ChannelControl =
 		SNew(SHorizontalBox)
@@ -70,20 +70,20 @@ TSharedRef<SWidget> STG_RGBAButtons::CreateChannelWidget(ETSChannelButton Type,
 		[
 			SNew(SCheckBox)
 			.Style(&FAppStyle::Get().GetWidgetStyle<FCheckBoxStyle>("DetailsView.ChannelToggleButton"))
-		.Type(ESlateCheckBoxType::ToggleButton)
-		.CheckedImage(CheckedBrush)
-		.CheckedHoveredImage(CheckedBrush)
-		.Padding(FMargin(0, 4))
-		.OnCheckStateChanged_Lambda(OnChannelCheckStateChanged, Type)
-		.IsChecked(this, &STG_RGBAButtons::OnGetChannelButtonCheckState, Type)
-		//.IsEnabled(this, &STG_SelectionPreview::IsChannelButtonEnabled, ETSChannelButton::Red)
-		[
-			SNew(STextBlock)
-			.Font(FAppStyle::Get().GetFontStyle("TG_Editor.ChannelButtonFont"))
-		.Justification(ETextJustify::Center)
-		.OverflowPolicy(ETextOverflowPolicy::Ellipsis)
-		.Text(FText::FromString(Name))
-		]
+			.Type(ESlateCheckBoxType::ToggleButton)
+			.CheckedImage(CheckedBrush)
+			.CheckedHoveredImage(CheckedBrush)
+			.Padding(FMargin(0, 4))
+			.OnCheckStateChanged_Lambda(OnChannelCheckStateChanged, Type)
+			.IsChecked(this, &STG_RGBAButtons::OnGetChannelButtonCheckState, Type)
+			//.IsEnabled(this, &STG_SelectionPreview::IsChannelButtonEnabled, ETSChannelButton::Red)
+			[
+				SNew(STextBlock)
+				.Font(FAppStyle::Get().GetFontStyle("TG_Editor.ChannelButtonFont"))
+			.Justification(ETextJustify::Center)
+			.OverflowPolicy(ETextOverflowPolicy::Ellipsis)
+			.Text(FText::FromString(Name))
+			]
 		];
 }
 
