Minor UI Fixes

diff --git a/Engine/Plugins/Experimental/TextureGraph/Source/TextureGraphEditor/Private/EdGraph/TG_EdGraphNode.cpp b/Engine/Plugins/Experimental/TextureGraph/Source/TextureGraphEditor/Private/EdGraph/TG_EdGraphNode.cpp
--- a/Engine/Plugins/Experimental/TextureGraph/Source/TextureGraphEditor/Private/EdGraph/TG_EdGraphNode.cpp
+++ b/Engine/Plugins/Experimental/TextureGraph/Source/TextureGraphEditor/Private/EdGraph/TG_EdGraphNode.cpp
@@ -405,7 +405,7 @@ FString UTG_EdGraphNode::GetTitleDetail()
 				Args.Add(TEXT("IsSRGB"), Desc.bIsSRGB ? "sRGB" : "Linear");
 				Args.Add(TEXT("Width"), FString::FromInt(Desc.Width));
 				Args.Add(TEXT("Height"), FString::FromInt(Desc.Height));
-				Details = FString::Format(TEXT("{Channels}_{Format}, {IsSRGB}\r\n{Width}x{Height}"), Args);
+				Details = FString::Format(TEXT("{Channels}_{Format}, {IsSRGB}\n{Width}x{Height}"), Args);
 			}
 		}
 	}
diff --git a/Engine/Plugins/Experimental/TextureGraph/Source/TextureGraphEditor/Private/Widgets/STG_RGBAButtons.cpp b/Engine/Plugins/Experimental/TextureGraph/Source/TextureGraphEditor/Private/Widgets/STG_RGBAButtons.cpp
--- a/Engine/Plugins/Experimental/TextureGraph/Source/TextureGraphEditor/Private/Widgets/STG_RGBAButtons.cpp
+++ b/Engine/Plugins/Experimental/TextureGraph/Source/TextureGraphEditor/Private/Widgets/STG_RGBAButtons.cpp
@@ -16,7 +16,7 @@ void STG_RGBAButtons::Construct(const FArguments& InArgs)
 
 TSharedRef<SWidget> STG_RGBAButtons::MakeChannelControlWidget()
 {
-	CheckedBrush = new FSlateRoundedBoxBrush(FStyleColors::Input, CoreStyleConstants::InputFocusRadius);
+	CheckedBrush = new FSlateRoundedBoxBrush(FLinearColor(0.039, 0.039, 0.039, 1), CoreStyleConstants::InputFocusRadius);
 
 	TSharedRef<SWidget> ChannelControl =
 		SNew(SHorizontalBox)
