Renamed MixerInsight to TextureGraphInsight

#tests 65b386838118420ba8f122e4

diff --git a/Engine/Plugins/Experimental/TextureGraph/Source/MixerInsight/MixerInsight.Build.cs b/Engine/Plugins/Experimental/TextureGraph/Source/MixerInsight/MixerInsight.Build.cs
deleted file mode 100644
diff --git a/Engine/Plugins/Experimental/TextureGraph/Source/MixerInsight/Private/MixerInsight.cpp b/Engine/Plugins/Experimental/TextureGraph/Source/MixerInsight/Private/MixerInsight.cpp
deleted file mode 100644
diff --git a/Engine/Plugins/Experimental/TextureGraph/Source/MixerInsight/Private/Model/MixerInsightObserver.cpp b/Engine/Plugins/Experimental/TextureGraph/Source/MixerInsight/Private/Model/MixerInsightObserver.cpp
deleted file mode 100644
diff --git a/Engine/Plugins/Experimental/TextureGraph/Source/MixerInsight/Private/Model/MixerInsightRecord.cpp b/Engine/Plugins/Experimental/TextureGraph/Source/MixerInsight/Private/Model/MixerInsightRecord.cpp
deleted file mode 100644
diff --git a/Engine/Plugins/Experimental/TextureGraph/Source/MixerInsight/Private/Model/MixerInsightSession.cpp b/Engine/Plugins/Experimental/TextureGraph/Source/MixerInsight/Private/Model/MixerInsightSession.cpp
deleted file mode 100644
diff --git a/Engine/Plugins/Experimental/TextureGraph/Source/MixerInsight/Private/View/SMixerInsightActionView.cpp b/Engine/Plugins/Experimental/TextureGraph/Source/MixerInsight/Private/View/SMixerInsightActionView.cpp
deleted file mode 100644
diff --git a/Engine/Plugins/Experimental/TextureGraph/Source/MixerInsight/Private/View/SMixerInsightBlobView.cpp b/Engine/Plugins/Experimental/TextureGraph/Source/MixerInsight/Private/View/SMixerInsightBlobView.cpp
deleted file mode 100644
diff --git a/Engine/Plugins/Experimental/TextureGraph/Source/MixerInsight/Private/View/SMixerInsightDeviceBufferView.cpp b/Engine/Plugins/Experimental/TextureGraph/Source/MixerInsight/Private/View/SMixerInsightDeviceBufferView.cpp
deleted file mode 100644
diff --git a/Engine/Plugins/Experimental/TextureGraph/Source/MixerInsight/Private/View/SMixerInsightDeviceView.cpp b/Engine/Plugins/Experimental/TextureGraph/Source/MixerInsight/Private/View/SMixerInsightDeviceView.cpp
deleted file mode 100644
diff --git a/Engine/Plugins/Experimental/TextureGraph/Source/MixerInsight/Private/View/SMixerInsightInspectorView.cpp b/Engine/Plugins/Experimental/TextureGraph/Source/MixerInsight/Private/View/SMixerInsightInspectorView.cpp
deleted file mode 100644
diff --git a/Engine/Plugins/Experimental/TextureGraph/Source/MixerInsight/Private/View/SMixerInsightMixView.cpp b/Engine/Plugins/Experimental/TextureGraph/Source/MixerInsight/Private/View/SMixerInsightMixView.cpp
deleted file mode 100644
diff --git a/Engine/Plugins/Experimental/TextureGraph/Source/MixerInsight/Private/View/SMixerInsightRecordTrailView.cpp b/Engine/Plugins/Experimental/TextureGraph/Source/MixerInsight/Private/View/SMixerInsightRecordTrailView.cpp
deleted file mode 100644
diff --git a/Engine/Plugins/Experimental/TextureGraph/Source/MixerInsight/Private/View/SMixerInsightResourceView.cpp b/Engine/Plugins/Experimental/TextureGraph/Source/MixerInsight/Private/View/SMixerInsightResourceView.cpp
deleted file mode 100644
diff --git a/Engine/Plugins/Experimental/TextureGraph/Source/MixerInsight/Private/View/SMixerInsightSessionView.cpp b/Engine/Plugins/Experimental/TextureGraph/Source/MixerInsight/Private/View/SMixerInsightSessionView.cpp
deleted file mode 100644
diff --git a/Engine/Plugins/Experimental/TextureGraph/Source/MixerInsight/Public/MixerInsight.h b/Engine/Plugins/Experimental/TextureGraph/Source/MixerInsight/Public/MixerInsight.h
deleted file mode 100644
diff --git a/Engine/Plugins/Experimental/TextureGraph/Source/MixerInsight/Public/Model/MixerInsightObserver.h b/Engine/Plugins/Experimental/TextureGraph/Source/MixerInsight/Public/Model/MixerInsightObserver.h
deleted file mode 100644
diff --git a/Engine/Plugins/Experimental/TextureGraph/Source/MixerInsight/Public/Model/MixerInsightRecord.h b/Engine/Plugins/Experimental/TextureGraph/Source/MixerInsight/Public/Model/MixerInsightRecord.h
deleted file mode 100644
diff --git a/Engine/Plugins/Experimental/TextureGraph/Source/MixerInsight/Public/Model/MixerInsightSession.h b/Engine/Plugins/Experimental/TextureGraph/Source/MixerInsight/Public/Model/MixerInsightSession.h
deleted file mode 100644
diff --git a/Engine/Plugins/Experimental/TextureGraph/Source/MixerInsight/Public/View/SMixerInsightActionView.h b/Engine/Plugins/Experimental/TextureGraph/Source/MixerInsight/Public/View/SMixerInsightActionView.h
deleted file mode 100644
diff --git a/Engine/Plugins/Experimental/TextureGraph/Source/MixerInsight/Public/View/SMixerInsightBlobView.h b/Engine/Plugins/Experimental/TextureGraph/Source/MixerInsight/Public/View/SMixerInsightBlobView.h
deleted file mode 100644
diff --git a/Engine/Plugins/Experimental/TextureGraph/Source/MixerInsight/Public/View/SMixerInsightDeviceBufferView.h b/Engine/Plugins/Experimental/TextureGraph/Source/MixerInsight/Public/View/SMixerInsightDeviceBufferView.h
deleted file mode 100644
diff --git a/Engine/Plugins/Experimental/TextureGraph/Source/MixerInsight/Public/View/SMixerInsightDeviceView.h b/Engine/Plugins/Experimental/TextureGraph/Source/MixerInsight/Public/View/SMixerInsightDeviceView.h
deleted file mode 100644
diff --git a/Engine/Plugins/Experimental/TextureGraph/Source/MixerInsight/Public/View/SMixerInsightInspectorView.h b/Engine/Plugins/Experimental/TextureGraph/Source/MixerInsight/Public/View/SMixerInsightInspectorView.h
deleted file mode 100644
diff --git a/Engine/Plugins/Experimental/TextureGraph/Source/MixerInsight/Public/View/SMixerInsightMixView.h b/Engine/Plugins/Experimental/TextureGraph/Source/MixerInsight/Public/View/SMixerInsightMixView.h
deleted file mode 100644
diff --git a/Engine/Plugins/Experimental/TextureGraph/Source/MixerInsight/Public/View/SMixerInsightRecordTrailView.h b/Engine/Plugins/Experimental/TextureGraph/Source/MixerInsight/Public/View/SMixerInsightRecordTrailView.h
deleted file mode 100644
diff --git a/Engine/Plugins/Experimental/TextureGraph/Source/MixerInsight/Public/View/SMixerInsightResourceView.h b/Engine/Plugins/Experimental/TextureGraph/Source/MixerInsight/Public/View/SMixerInsightResourceView.h
deleted file mode 100644
diff --git a/Engine/Plugins/Experimental/TextureGraph/Source/MixerInsight/Public/View/SMixerInsightSessionView.h b/Engine/Plugins/Experimental/TextureGraph/Source/MixerInsight/Public/View/SMixerInsightSessionView.h
deleted file mode 100644
diff --git a/Engine/Plugins/Experimental/TextureGraph/Source/MixerInsightEditor/MixerInsightEditor.Build.cs b/Engine/Plugins/Experimental/TextureGraph/Source/MixerInsightEditor/MixerInsightEditor.Build.cs
deleted file mode 100644
diff --git a/Engine/Plugins/Experimental/TextureGraph/Source/MixerInsightEditor/Private/MixerInsightEditor.cpp b/Engine/Plugins/Experimental/TextureGraph/Source/MixerInsightEditor/Private/MixerInsightEditor.cpp
deleted file mode 100644
diff --git a/Engine/Plugins/Experimental/TextureGraph/Source/MixerInsightEditor/Private/MixerInsightEditorCommands.cpp b/Engine/Plugins/Experimental/TextureGraph/Source/MixerInsightEditor/Private/MixerInsightEditorCommands.cpp
deleted file mode 100644
diff --git a/Engine/Plugins/Experimental/TextureGraph/Source/MixerInsightEditor/Private/MixerInsightEditorStyle.cpp b/Engine/Plugins/Experimental/TextureGraph/Source/MixerInsightEditor/Private/MixerInsightEditorStyle.cpp
deleted file mode 100644
diff --git a/Engine/Plugins/Experimental/TextureGraph/Source/MixerInsightEditor/Public/MixerInsightEditor.h b/Engine/Plugins/Experimental/TextureGraph/Source/MixerInsightEditor/Public/MixerInsightEditor.h
deleted file mode 100644
diff --git a/Engine/Plugins/Experimental/TextureGraph/Source/MixerInsightEditor/Public/MixerInsightEditorCommands.h b/Engine/Plugins/Experimental/TextureGraph/Source/MixerInsightEditor/Public/MixerInsightEditorCommands.h
deleted file mode 100644
diff --git a/Engine/Plugins/Experimental/TextureGraph/Source/MixerInsightEditor/Public/MixerInsightEditorStyle.h b/Engine/Plugins/Experimental/TextureGraph/Source/MixerInsightEditor/Public/MixerInsightEditorStyle.h
deleted file mode 100644
diff --git a/Engine/Plugins/Experimental/TextureGraph/Source/TextureGraphInsight/Private/Model/TextureGraphInsightObserver.cpp b/Engine/Plugins/Experimental/TextureGraph/Source/TextureGraphInsight/Private/Model/TextureGraphInsightObserver.cpp
new file mode 100644
--- /dev/null
+++ b/Engine/Plugins/Experimental/TextureGraph/Source/TextureGraphInsight/Private/Model/TextureGraphInsightObserver.cpp
@@ -0,0 +1,151 @@
+// Copyright Epic Games, Inc. All Rights Reserved.
+#include "Model/TextureGraphInsightObserver.h"
+
+#include "TextureGraphInsight.h"
+#include "Model/TextureGraphInsightSession.h"
+
+#include <Device/DeviceManager.h>
+
+//DEFINE_LOG_CATEGORY(LogTextureGraphInsight);
+
+
+TextureGraphInsightDeviceObserver::TextureGraphInsightDeviceObserver()
+{
+}
+TextureGraphInsightDeviceObserver::TextureGraphInsightDeviceObserver(DeviceType InDevType) : DevType(InDevType)
+{
+}
+TextureGraphInsightDeviceObserver::~TextureGraphInsightDeviceObserver()
+{
+}
+void TextureGraphInsightDeviceObserver::DeviceBuffersUpdated(HashNDescArray&& AddedBuffers, HashArray&& RemovedBuffers)
+{
+	TextureGraphInsight::Instance()->GetSession()->DeviceBuffersUpdated(DevType, std::move(AddedBuffers), std::move(RemovedBuffers));
+
+}
+
+TextureGraphInsightBlobberObserver::TextureGraphInsightBlobberObserver()
+{
+}
+TextureGraphInsightBlobberObserver::~TextureGraphInsightBlobberObserver()
+{
+}
+void TextureGraphInsightBlobberObserver::BlobberUpdated(HashArray&& AddedHashes, HashArray&& RemappedHashes)
+{
+	TextureGraphInsight::Instance()->GetSession()->BlobberHashesUpdated(std::move(AddedHashes), std::move(RemappedHashes));
+}
+
+
+TextureGraphInsightSchedulerObserver::TextureGraphInsightSchedulerObserver()
+{
+}
+TextureGraphInsightSchedulerObserver::~TextureGraphInsightSchedulerObserver()
+{
+}
+
+void TextureGraphInsightSchedulerObserver::Start()
+{
+	UE_LOG(LogTextureGraphInsight, Log, TEXT("TextureGraphInsightSchedulerObserver::Start"));
+}
+
+void TextureGraphInsightSchedulerObserver::UpdateIdle()
+{
+	if (TextureGraphEngine::IsTestMode())
+		return;
+
+	TextureGraphInsight::Instance()->GetSession()->UpdateIdle();
+}
+
+void TextureGraphInsightSchedulerObserver::Stop()
+{
+	UE_LOG(LogTextureGraphInsight, Log, TEXT("TextureGraphInsightSchedulerObserver::Stop"));
+}
+
+void TextureGraphInsightSchedulerObserver::BatchAdded(JobBatchPtr Batch)
+{
+	if (TextureGraphEngine::IsTestMode())
+		return;
+
+	TextureGraphInsight::Instance()->GetSession()->BatchAdded(Batch);
+}
+
+void TextureGraphInsightSchedulerObserver::BatchDone(JobBatchPtr Batch)
+{
+	if (TextureGraphEngine::IsTestMode())
+		return;
+
+	TextureGraphInsight::Instance()->GetSession()->BatchDone(Batch);
+}
+
+void TextureGraphInsightSchedulerObserver::BatchJobsDone(JobBatchPtr Batch)
+{
+	if (TextureGraphEngine::IsTestMode())
+		return;
+
+	TextureGraphInsight::Instance()->GetSession()->BatchJobsDone(Batch);
+}
+
+TextureGraphInsightEngineObserver::TextureGraphInsightEngineObserver()
+{
+	UE_LOG(LogTextureGraphInsight, Log, TEXT("TextureGraphInsightEngineObserver::Constructor"));
+
+	for (int i = 0; i < (uint32)DeviceType::Count; ++i)
+	{
+		_deviceObservers[i] = std::make_shared<TextureGraphInsightDeviceObserver>((DeviceType)i);
+	}
+	BlobberObserver = std::make_shared<TextureGraphInsightBlobberObserver>();
+	SchedulerObserver = std::make_shared<TextureGraphInsightSchedulerObserver>();
+}
+
+TextureGraphInsightEngineObserver::~TextureGraphInsightEngineObserver()
+{
+	UE_LOG(LogTextureGraphInsight, Log, TEXT("TextureGraphInsightEngineObserver::Destructor"));
+}
+
+void TextureGraphInsightEngineObserver::Created()
+{
+	if (TextureGraphEngine::IsTestMode())
+		return;
+
+	UE_LOG(LogTextureGraphInsight, Log, TEXT("TextureGraphInsightEngineObserver::Created"));
+
+	/// Engine is created when notified, this should be true
+	if (TextureGraphEngine::GetInstance())
+	{
+		for (int i = 0; i < (uint32)DeviceType::Count; ++i)
+		{
+			auto Dev = TextureGraphEngine::GetInstance()->GetDeviceManager()->GetDevice(i);
+			if (Dev)
+				Dev->RegisterObserverSource(_deviceObservers[i]);
+		}
+		TextureGraphEngine::GetInstance()->GetBlobber()->RegisterObserverSource(BlobberObserver);
+		TextureGraphEngine::GetInstance()->GetScheduler()->RegisterObserverSource(SchedulerObserver);
+	}
+
+	/// This is a brand new session engine
+	TextureGraphInsight::Instance()->GetSession()->EngineCreated();
+}
+
+
+void TextureGraphInsightEngineObserver::Destroyed()
+{
+	if (TextureGraphEngine::IsTestMode())
+		return;
+
+	UE_LOG(LogTextureGraphInsight, Log, TEXT("TextureGraphInsightEngineObserver::Destroyed"));
+
+	/// Engine is already destroyed when notified
+	if (TextureGraphEngine::GetInstance())
+	{
+		for (int i = 0; i < (uint32)DeviceType::Count; ++i)
+		{
+			auto Dev = TextureGraphEngine::GetInstance()->GetDeviceManager()->GetDevice(i);
+			if (Dev)
+				Dev->RegisterObserverSource(nullptr);
+		}
+		TextureGraphEngine::GetInstance()->GetBlobber()->RegisterObserverSource(nullptr); /// remove blobber observer from the engine
+		TextureGraphEngine::GetInstance()->GetScheduler()->RegisterObserverSource(nullptr); /// remove scheduler observer from the engine
+	}
+
+	TextureGraphInsight::Instance()->GetSession()->EngineDestroyed();
+}
diff --git a/Engine/Plugins/Experimental/TextureGraph/Source/TextureGraphInsight/Private/Model/TextureGraphInsightRecord.cpp b/Engine/Plugins/Experimental/TextureGraph/Source/TextureGraphInsight/Private/Model/TextureGraphInsightRecord.cpp
new file mode 100644
--- /dev/null
+++ b/Engine/Plugins/Experimental/TextureGraph/Source/TextureGraphInsight/Private/Model/TextureGraphInsightRecord.cpp
@@ -0,0 +1,340 @@
+// Copyright Epic Games, Inc. All Rights Reserved.
+#include "Model/TextureGraphInsightRecord.h"
+
+const DeviceBufferRecord DeviceBufferRecord::null;
+const BlobRecord BlobRecord::null;
+const JobRecord JobRecord::GNull;
+const BatchRecord BatchRecord::null;
+const ActionRecord ActionRecord::null;
+const MixRecord MixRecord::null;
+
+///
+/// Batch
+/// 
+
+RecordID SessionRecord::FindBatchRecord(BatchID InBatchID) const
+{
+	auto found = RecordFromBatchID.find(InBatchID);
+	if (found == RecordFromBatchID.end())
+	{
+		return RecordID();
+	} else
+	{
+		return found->second;
+	}
+}
+
+RecordID SessionRecord::NewBatchRecord(BatchID InBatchID, BatchRecordConstructor Constructor)
+{
+	// Unless batchID already exists, we will push the batch as the next element in the array
+	auto found = RecordFromBatchID.find(InBatchID);
+	if (found == RecordFromBatchID.end())
+	{
+		RecordID recordID = RecordID::fromBatch( Batches.size() );
+		RecordFromBatchID.insert({ InBatchID, recordID });
+		Batches.emplace_back(Constructor(recordID));
+
+		auto& newBatchRecord = Batches.back();
+		newBatchRecord.SelfID = recordID;
+		int j = 0;
+		for (auto& jr : newBatchRecord.Jobs)
+		{
+			jr.SelfID = RecordID::fromBatchJob(recordID.Batch(), j);
+			j++;
+		}
+
+		return recordID;
+	}
+	else
+	{
+		return RecordID();
+	}
+}
+
+
+bool SessionRecord::EditBatchRecord(RecordID RecId, BatchRecordEditor Editor)
+{
+	if (RecId.IsBatch())
+	{
+		return Editor(Batches[RecId.Batch()]);
+	}
+	return false;
+}
+
+
+std::tuple<RecordID,bool> SessionRecord::NewDeviceBufferRecord(DeviceType DevType, DeviceBufferID BufferID, const DeviceBufferRecord& Record)
+{
+	// Unless blobID already exists, we will push the blob as the next element in the array
+	auto found = BufferFromDeviceBufferIDActive[(int32) DevType].find(BufferID);
+	if (found == BufferFromDeviceBufferIDActive[(int32) DevType].end())
+	{
+		// Define the RecordID for this new buffer
+		RecordID recordID = RecordID::fromBuffer(Record.DevType, DeviceBuffers[(uint32) Record.DevType].size());
+		// Add it in the device active buffer map
+		BufferFromDeviceBufferIDActive[(int32) DevType].insert({ BufferID, recordID });
+		// Add it at the expected location in the per defice array
+		DeviceBuffers[(uint32)Record.DevType].emplace_back(std::move(Record));
+		// And add it in the global buffer from hash map
+		BufferFromHash.insert({ Record.HashValue, recordID });
+
+		return { recordID, true };
+	}
+	else
+	{
+		auto& r = DeviceBuffers[(uint32)Record.DevType][found->second.Buffer()];
+		if (r.HashValue != Record.HashValue)
+		{
+			// Take note of the new hash assigned to this device buffer
+			r.PrevHashValue = r.HashValue;
+			r.HashValue = Record.HashValue;
+			r.bLeaked = false;
+
+			// And add it in the global buffer from hash map
+			BufferFromHash.insert({ r.HashValue, found->second });
+
+			return  { found->second, false } ;
+		}
+		else
+		{
+			return { RecordID(), false};
+		}
+	}
+}
+
+RecordID SessionRecord::EraseDeviceBufferRecord(DeviceType DevType, DeviceBufferID BufferID)
+{
+	// Expect to find it in the current map
+	auto found = BufferFromDeviceBufferIDActive[(int32)DevType].find(BufferID);
+	if (found == BufferFromDeviceBufferIDActive[(int32)DevType].end())
+	{
+		// Should not happen, the removed entry should have been seen before
+		return RecordID();
+	}
+	else
+	{
+		// good we found the record in the active map
+		auto RecId = found->second;
+		// let's remove it and add it in the removed map
+		BufferFromDeviceBufferIDActive[(int32)DevType].erase(found);
+		// Add it to the removed map
+		BufferFromDeviceBufferIDErased[BufferID] = RecId;
+
+		// tag the record
+		DeviceBuffers[RecId.Buffer_DeviceType()][RecId.Buffer()].bErased = true;
+		
+		// And return its recordID
+		return RecId;
+	}
+}
+
+RecordIDArray SessionRecord::FindActiveDeviceBufferRecords(DeviceType DevType, const DeviceBufferIDArray& BufferIDs, RecordIDArray& LeakedBuffers)
+{
+	RecordIDArray results;
+	results.reserve(BufferIDs.size());
+
+	auto map = BufferFromDeviceBufferIDActive[(int32) DevType]; // copy the current active map
+
+	for (auto bid : BufferIDs)
+	{
+		auto found = map.find(bid);
+		if (found == map.end())
+		{
+			results.emplace_back(RecordID());
+		}
+		else
+		{
+			results.emplace_back(found->second);
+			map.erase(found);
+		}
+	}
+
+	for (auto lid : map)
+	{
+		LeakedBuffers.emplace_back(lid.second);
+	}
+
+	return results;
+}
+
+
+RecordID SessionRecord::FindDeviceBufferRecord(DeviceBufferID BufferID) const
+{
+	RecordFromDeviceBufferIDMap::const_iterator found;
+
+	// Look in each device's map
+	for (int32 i = 0; i < (int32) DeviceType::Count; ++i)
+	{
+		found = BufferFromDeviceBufferIDActive[i].find(BufferID);
+		if (found != BufferFromDeviceBufferIDActive[i].end())
+		{
+			return found->second;
+		}
+	}
+
+	// try again with erased map
+	found = BufferFromDeviceBufferIDErased.find(BufferID);
+	if (found != BufferFromDeviceBufferIDErased.end())
+	{
+		return found->second;
+		
+	}
+	
+	// nope
+	return RecordID();
+}
+
+RecordIDArray SessionRecord::FindDeviceBufferRecords(const DeviceBufferIDArray& BufferIDs) const
+{
+	RecordIDArray results;
+	results.reserve(BufferIDs.size());
+	for (auto b : BufferIDs)
+	{
+		results.emplace_back(FindDeviceBufferRecord(b));
+	}
+	return results;
+}
+
+RecordIDArray SessionRecord::FetchActiveDeviceBufferIDs(DeviceType DevType) const
+{
+	RecordIDArray results;
+	results.reserve(BufferFromDeviceBufferIDActive[(int32) DevType].size());
+	for (auto b : BufferFromDeviceBufferIDActive[(int32)DevType])
+	{
+		results.emplace_back(b.second);
+	}
+	return results;
+}
+
+RecordID SessionRecord::FindDeviceBufferRecordFromHash(HashType InHash) const
+{
+	// try to find the hash in the buffer from hash map
+	auto found = BufferFromHash.find(InHash);
+	if (found != BufferFromHash.end())
+	{
+		return found->second;
+	}
+	// nope
+	return RecordID();
+}
+
+
+RecordID SessionRecord::FindBlobRecord(BlobID InBlobID) const
+{
+	auto found = BlobFromBlobID.find(InBlobID);
+	if (found == BlobFromBlobID.end())
+	{
+		return RecordID();
+	}
+	else
+	{
+		return found->second;
+	}
+}
+
+RecordID SessionRecord::NewBlobRecord(BlobID InBlobID, BlobRecordConstructor Constructor)
+{
+	// Unless blobID already exists, we will push the blob as the next element in the array
+	auto found = BlobFromBlobID.find(InBlobID);
+	if (found == BlobFromBlobID.end())
+	{
+		RecordID recordID = RecordID::fromBlob(Blobs.size());
+		BlobFromBlobID.insert({ InBlobID, recordID });
+		Blobs.emplace_back(Constructor(recordID));
+
+		return recordID;
+	} else
+	{
+		return RecordID();
+	}
+}
+
+bool SessionRecord::EditBlobRecord(RecordID RecId, BlobRecordEditor Editor)
+{
+	if (RecId.IsBlob())
+	{
+		return Editor(Blobs[RecId.Blob()]);
+	}
+	return false;
+}
+
+
+
+RecordID SessionRecord::NewActionRecord(ActionID InActionID, const ActionRecord& action)
+{
+	// Unless actionID already exists, we will push the blob as the next element in the array
+	auto found = ActionFromActionID.find(InActionID);
+	if (found == ActionFromActionID.end())
+	{
+		RecordID recordID = RecordID::fromAction(Actions.size());
+		ActionFromActionID.insert({ InActionID, recordID });
+		Actions.emplace_back(std::move(action));
+
+		return recordID;
+	}
+	else
+	{
+		return RecordID();
+	}
+}
+
+RecordID SessionRecord::FindActionRecord(ActionID InActionID) const
+{
+	auto found = ActionFromActionID.find(InActionID);
+	if (found == ActionFromActionID.end())
+	{
+		return RecordID();
+	}
+	else
+	{
+		return found->second;
+	}
+}
+
+
+///
+/// Mix
+/// 
+
+RecordID SessionRecord::FindMixRecord(MixID MixID) const
+{
+	auto found = RecordFromMixID.find(MixID);
+	if (found == RecordFromMixID.end())
+	{
+		return RecordID();
+	}
+	else
+	{
+		return found->second;
+	}
+}
+
+RecordID SessionRecord::NewMixRecord(MixID mixID, MixRecordConstructor constructor)
+{
+	// Unless mixID already exists, we will push the blob as the next element in the array
+	auto found = RecordFromMixID.find(mixID);
+	if (found == RecordFromMixID.end())
+	{
+		RecordID recordID = RecordID::fromMix(Mixes.size());
+		RecordFromMixID.insert({ mixID, recordID });
+		Mixes.emplace_back(constructor(recordID));
+
+		auto& newRecord = Mixes.back();
+		newRecord.SelfID = recordID;
+
+		return recordID;
+	}
+	else
+	{
+		return RecordID();
+	}
+}
+
+bool SessionRecord::EditMixRecord(RecordID RecId, MixRecordEditor editor)
+{
+	if (RecId.IsMix())
+	{
+		return editor(Mixes[RecId.Mix()]);
+	}
+	return false;
+}
+
diff --git a/Engine/Plugins/Experimental/TextureGraph/Source/TextureGraphInsight/Private/Model/TextureGraphInsightSession.cpp b/Engine/Plugins/Experimental/TextureGraph/Source/TextureGraphInsight/Private/Model/TextureGraphInsightSession.cpp
new file mode 100644
--- /dev/null
+++ b/Engine/Plugins/Experimental/TextureGraph/Source/TextureGraphInsight/Private/Model/TextureGraphInsightSession.cpp
@@ -0,0 +1,924 @@
+// Copyright Epic Games, Inc. All Rights Reserved.
+#include "Model/TextureGraphInsightSession.h"
+
+#include "Job/Job.h"
+#include "Data/Blob.h"
+
+#include "TextureGraphEngine.h"
+#include "Data/Blobber.h"
+#include "Device/DeviceManager.h"
+#include "Device/Device.h"
+
+#include "Model/Mix/Mix.h"
+#include "TextureGraphInsight.h"
+
+#include "Transform/Mix/T_UpdateTargets.h"
+
+
+TextureGraphInsightSession::TextureGraphInsightSession() :
+	Record(new SessionRecord()),
+	Cache(new LiveCache(this))
+{
+}
+TextureGraphInsightSession::~TextureGraphInsightSession()
+{
+}
+
+void TextureGraphInsightSession::LiveCache::AddBlob(RecordID RecId, BlobPtr BlobObj)
+{
+	if (RecId.IsBlob())
+	{
+		check(RecId.Blob() == Blobs.size());
+		Blobs.emplace_back(BlobEntry(BlobObj->Hash()->Value()));
+	}
+}
+
+BlobPtr TextureGraphInsightSession::LiveCache::GetBlob(RecordID RecId) const
+{
+	if (RecId.IsBlob() && RecId.Blob() < Blobs.size())
+	{
+		const auto& blobEntry = Blobs[RecId.Blob()];
+		return TextureGraphEngine::GetInstance()->GetBlobber()->Find<Blob>(blobEntry).get();
+	}
+	return nullptr;
+}
+
+void TextureGraphInsightSession::LiveCache::AddBatch(RecordID RecId, JobBatchPtr Batch)
+{
+	if (RecId.IsBatch())
+	{
+		check(RecId.Batch() == Batches.size());
+		Batches.emplace_back(Batch);
+	}
+}
+
+JobBatchPtr TextureGraphInsightSession::LiveCache::GetBatch(RecordID RecId) const
+{
+	if ((RecId.IsBatch() || RecId.IsJob()) && RecId.Batch() < Batches.size())
+	{
+		return Batches[RecId.Batch()].lock();
+	}
+	return nullptr;
+}
+
+void TextureGraphInsightSession::LiveCache::AddMix(RecordID RecId, MixPtr MixObj)
+{
+	if (RecId.IsMix())
+	{
+		check(RecId.Mix() == Mixes.size());
+		Mixes.emplace_back(MixObj);
+	}
+}
+
+MixPtr TextureGraphInsightSession::LiveCache::GetMix(RecordID RecId) const
+{
+	if (RecId.IsMix() && RecId.Mix() < Mixes.size())
+	{
+		//return _mixes[RecId.Mix()].lock();
+		return Mixes[RecId.Mix()];
+	}
+	return nullptr;
+}
+
+
+DeviceBufferPtr TextureGraphInsightSession::LiveCache::GetDeviceBuffer(RecordID RecId) const
+{
+	if (RecId.IsBuffer())
+	{
+		auto BufferRecord = Session->GetRecord().GetBuffer(RecId);
+		auto SourceBlob = BufferRecord.SourceID;
+		auto SourceBlobPtr = GetBlob(SourceBlob);
+		auto Hash = BufferRecord.HashValue;
+		if (SourceBlobPtr)
+			Hash = SourceBlobPtr->Hash()->Value();
+
+		// Try to find the buffer in the device under the  first hash
+		auto Dev = TextureGraphEngine::GetInstance()->GetDeviceManager()->GetDevice(RecId.Buffer_DeviceType());
+		auto Buffer = Dev->Find(Hash, false);
+
+		// Couldn't find with what we think is final hash, try with prev hash if exists ?
+		if (!Buffer) {
+			if (BufferRecord.PrevHashValue != 0)
+			{
+				Buffer = Dev->Find(BufferRecord.PrevHashValue, false);
+			}
+		}
+
+		for (int32 i = RecId.Buffer_DeviceType() + 1; i < (int32) DeviceType::Count; i++)
+		{
+			Dev = TextureGraphEngine::GetInstance()->GetDeviceManager()->GetDevice(i);
+			if (Dev && !Buffer) {
+				Buffer = Dev->Find(Hash, false);
+				if (Buffer)
+					break;
+			}
+		}
+
+		// Couldn't find in the device cache... then look up in the blobber
+		if (!Buffer) 
+		{
+			auto BlobPtr = TextureGraphEngine::GetInstance()->GetBlobber()->Find(BufferRecord.HashValue);
+			if (BlobPtr) {
+				Buffer = BlobPtr->GetBufferRef();
+			}
+			else
+			{
+				UE_LOG(LogTextureGraphInsight, Log, TEXT("TextureGraphInsightActionManagerObserver::GetDeviceBuffer find the buffer from hash but no buffer"));
+			}
+		}
+		// Couldn't find in the blobber under final hash ? try the prev hash
+		if (!Buffer) {
+			if (BufferRecord.PrevHashValue != 0)
+			{
+				auto blobPtr = TextureGraphEngine::GetInstance()->GetBlobber()->Find(BufferRecord.PrevHashValue);
+				if (blobPtr) {
+					Buffer = blobPtr->GetBufferRef();
+				}
+				else
+				{
+					UE_LOG(LogTextureGraphInsight, Log, TEXT("TextureGraphInsightActionManagerObserver::GetDeviceBuffer find the buffer from prevHash but no buffer"));
+				}
+			}
+			else
+			{
+				UE_LOG(LogTextureGraphInsight, Log, TEXT("TextureGraphInsightActionManagerObserver::GetDeviceBuffer cannot find the buffer in the cache, missing second hash"));
+			}
+		}
+
+		return Buffer;
+	}
+	return nullptr;
+}
+
+void TextureGraphInsightSession::EngineCreated()
+{
+	if (TextureGraphEngine::IsTestMode())
+		return;
+
+	UE_LOG(LogTextureGraphInsight, Log, TEXT("TextureGraphInsightActionManagerObserver::EngineCreated"));
+
+	// Record and Cache should be defined, empty and ready to be used
+	RecordIDArray addedRecordIDs;
+	for (int32 i = 0; i < (int32) DeviceType::Count; ++i)
+		ParseDeviceBuffersCache((DeviceType)i, addedRecordIDs);
+
+	EmitOnEngineReset(true);
+}
+
+void TextureGraphInsightSession::EngineDestroyed()
+{
+	if (TextureGraphEngine::IsTestMode())
+		return;
+
+	UE_LOG(LogTextureGraphInsight, Log, TEXT("TextureGraphInsightActionManagerObserver::EngineDestroyed"));
+
+	// Need to reset record and cache
+	{
+		std::lock_guard<std::mutex> Lock(RecordMutex);
+		Record = std::make_shared<SessionRecord>();
+		Cache = std::make_shared<LiveCache>(this);
+	}
+
+	EmitOnEngineReset(false);
+}
+
+void TextureGraphInsightSession::EmitOnEngineReset(bool createOrDestroy)
+{
+	int32 PrevEngineID = EngineID;
+	if (createOrDestroy)
+	{
+		EngineID = abs(PrevEngineID) + 1;
+	}
+	else
+	{
+		EngineID = -abs(PrevEngineID);
+	}
+
+	OnEngineResetEvent.Broadcast(EngineID);
+}
+
+uint32 TextureGraphInsightSession::ParseDeviceBuffersCache(DeviceType DevType, RecordIDArray& AddedBufferIDs)
+{
+	uint32 NumNewBuffers = 0;
+
+	auto DevIndex = (int32) DevType;
+	auto Dev = TextureGraphEngine::GetInstance()->GetDeviceManager()->GetDevice(DevType);
+	if (Dev)
+	{
+		auto Version = Dev->GetObserverSource()->GetVersion();
+		if (Version > GetCache().DeviceCacheVersions[DevIndex])
+		{
+			GetCache().DeviceCacheVersions[DevIndex] = Version;
+
+			auto NumBuffers = Dev->GetNumDeviceBuffersUsed();
+			DeviceObserverSource::HashArray buffers(NumBuffers, 0);
+			DeviceObserverSource::HashNDescArray hashNDescs(NumBuffers);
+			Dev->TraverseBufferCache([this,&buffers, &hashNDescs](const DeviceBufferPtr& buffer, uint32 index)
+			{
+				auto& Entry = hashNDescs[index];
+				Entry.Raw = (uint64)buffer.get();
+				Entry.First = buffer->Hash()->Value();
+				Entry.Second = buffer->Descriptor();
+
+				buffers[index] = Entry.Raw;
+			});
+
+			if (buffers.size())
+			{
+				// now compare against what we know
+				RecordIDArray LeakedRecIds;
+
+				auto BufferRecordIDs = GetRecord().FindActiveDeviceBufferRecords(DevType, buffers, LeakedRecIds);
+
+
+				for (auto i = 0; i < BufferRecordIDs.size(); ++i)
+				{
+					auto& RecId = BufferRecordIDs[i];
+					if (RecId.IsBuffer() && (RecId.Buffer_DeviceType() == DevIndex))
+					{
+						// everything happened as expected, all good, update the rank		
+						auto& BuffRecord = GetRecord().GetBuffer(RecId);
+						BuffRecord.Rank = i;
+						BuffRecord.bLeaked = false;
+					}
+					else {
+						// We just discovered a buffer, let's create it in the records
+						DeviceBufferRecord DevBufferRec;
+						DevBufferRec.DevType = DevType;
+						DevBufferRec.ID = hashNDescs[i].Raw;
+						DevBufferRec.HashValue = hashNDescs[i].First;
+						DevBufferRec.Descriptor = hashNDescs[i].Second;
+
+						DevBufferRec.bLeaked = true;
+						DevBufferRec.Rank = i;
+						bool bAdded = false;
+						std::tie(RecId, bAdded) = GetRecord().NewDeviceBufferRecord(DevType, DevBufferRec.ID, DevBufferRec);
+
+						if (bAdded)
+							AddedBufferIDs.emplace_back(RecId);
+
+						++NumNewBuffers;
+					}
+				}
+
+				for (auto lid : LeakedRecIds)
+				{
+					auto& BufferRecord = GetRecord().GetBuffer(lid);
+					BufferRecord.bLeaked = true;
+					BufferRecord.Rank = -1;
+				}
+			}
+		}
+	}
+
+	if (NumNewBuffers)
+	{
+		auto t = Device::DeviceType_Names(DevType);
+		UE_LOG(LogTextureGraphInsight, VeryVerbose, TEXT("TextureGraphInsightSession::ParseDeviceBuffersCache[%s]: %d new unknown buffers"), *t, NumNewBuffers);
+	}
+	
+	return NumNewBuffers;
+}
+
+void TextureGraphInsightSession::DeviceBuffersUpdated(DeviceType DevType, TextureGraphInsightDeviceObserver::HashNDescArray&& AddedBuffers, TextureGraphInsightDeviceObserver::HashArray&& RemovedBuffers)
+{
+	RecordIDArray AddedRecIds;
+	for (auto& AddedBuffer : AddedBuffers)
+	{
+		DeviceBufferRecord DevBuffRec;
+		DevBuffRec.DevType = DevType;
+		DevBuffRec.ID = AddedBuffer.Raw;
+		DevBuffRec.HashValue = AddedBuffer.First;
+		DevBuffRec.Descriptor = AddedBuffer.Second;
+	
+		auto key = AddedBuffer.Raw;
+
+		UE_LOG(LogTextureGraphInsight, VeryVerbose, TEXT("TextureGraphInsightSession::DeviceBuffersAdded : ptr 0x%llx - hash %llu"), key, AddedBuffer.First);
+
+		RecordID RecId;
+		bool bAdded;
+		std::tie(RecId, bAdded) = GetRecord().NewDeviceBufferRecord(DevType, key, DevBuffRec);
+		if (RecId.IsValid())
+		{
+			if (bAdded)
+			{
+				UE_LOG(LogTextureGraphInsight, VeryVerbose, TEXT("TextureGraphInsightSession::DeviceBuffersAdded : %d - hash %llu"), RecId.Buffer(), DevBuffRec.HashValue);
+				AddedRecIds.emplace_back(RecId);
+			}
+			else
+			{
+				UE_LOG(LogTextureGraphInsight, VeryVerbose, TEXT("TextureGraphInsightSession::DeviceBuffers Updated hash : %d - hash %llu"), RecId.Buffer(), DevBuffRec.HashValue);
+			}
+		}
+		else
+		{
+			UE_LOG(LogTextureGraphInsight, VeryVerbose, TEXT("TextureGraphInsightSession::DeviceBuffersAdded ERROR invalid buffer: ptr 0x%llx - hash %llu"), DevBuffRec.ID, DevBuffRec.HashValue);
+		}
+	}
+
+	RecordIDArray RemovedRecIds;
+	for (int i = 0; i < RemovedBuffers.size(); ++i)
+	{
+		auto bp = RemovedBuffers[i];
+		++i;
+		auto h = RemovedBuffers[i];
+		++i;
+		auto ph = RemovedBuffers[i];
+
+		UE_LOG(LogTextureGraphInsight, VeryVerbose, TEXT("TextureGraphInsightSession::DeviceBuffersRemoved : ptr 0x%llx - hash %llu - prevhash %llu"), bp, h, ph);
+		auto RecId = GetRecord().EraseDeviceBufferRecord(DevType, bp);
+		if (RecId.IsValid())
+		{
+			UE_LOG(LogTextureGraphInsight, VeryVerbose, TEXT("TextureGraphInsightSession::DeviceBuffersRemoved : %d"), RecId.Buffer());
+			RemovedRecIds.emplace_back(RecId);
+		}
+		else
+		{
+			UE_LOG(LogTextureGraphInsight, VeryVerbose, TEXT("TextureGraphInsightSession::DeviceBuffersAdded ERROR invalid buffer ?"));
+		}
+	}
+
+	// There was some changes, let's parse and eventually collect some new buffers that we would have missed ?
+	ParseDeviceBuffersCache(DevType, AddedRecIds);
+
+	OnDeviceBufferAddedEvent.Broadcast(AddedRecIds);
+	OnDeviceBufferRemovedEvent.Broadcast(RemovedRecIds);
+}
+
+void TextureGraphInsightSession::BlobberHashesUpdated(TextureGraphInsightBlobberObserver::HashArray&& AddedHashes, TextureGraphInsightBlobberObserver::HashArray&& RemappedHashes)
+{
+	//for (auto a : addedHashes)
+	for (auto it = AddedHashes.rbegin(); it != AddedHashes.rend(); ++it)
+	{
+		auto a = (*it);
+
+//		UE_LOG(LogTextureGraphInsight, Log, TEXT("TextureGraphInsightSession::BlobberHashesUpdated::Added hash %llu"), a);
+		// Try to find the hash BlobObj
+
+//		auto n = FindOrCreateBlobRecordFromHash(a);
+	}
+	
+	int32 numRemapped = RemappedHashes.size() >> 1;
+	for (int i = 0; i < numRemapped; ++i)
+	{
+		UE_LOG(LogTextureGraphInsight, VeryVerbose, TEXT("TextureGraphInsightSession::BlobberHashesUpdated::Remapped hash old %llu ==> new %llu"), RemappedHashes[2 * i], RemappedHashes[2 * i + 1]);
+		auto o = RemappedHashes[2 * i];
+		auto n = RemappedHashes[2 * i + 1];
+		if (o != n)
+		{
+			auto RecId = RemapBlobHash(o, n);
+		}
+	}
+
+	if (numRemapped)
+		EmitOnMappedBlob();
+}
+
+
+
+void TextureGraphInsightSession::UpdateIdle()
+{
+	OnUpdateIdleEvent.Broadcast();
+}
+
+void TextureGraphInsightSession::BatchAdded(JobBatchPtr Batch)
+{
+	UE_LOG(LogTextureGraphInsight, Log, TEXT("TextureGraphInsightSchedulerObserver::BatchAdded"));
+}
+
+void TextureGraphInsightSession::BatchDone(JobBatchPtr Batch)
+{
+	UE_LOG(LogTextureGraphInsight, Log, TEXT("TextureGraphInsightSchedulerObserver::BatchDone"));
+
+	auto recordID = Record->NewBatchRecord(Batch->GetBatchId(), 
+		[this, Batch](RecordID RecId) -> BatchRecord
+		{
+			return MakeBatchRecord(RecId, Batch);
+		});
+
+	if (recordID.IsValid())
+	{
+		// Cache the live job bath ptr
+		GetCache().AddBatch(recordID, Batch);
+
+		OnBatchAddedEvent.Broadcast(recordID);
+	}
+
+	if (recordID.IsValid())
+	{
+		OnBatchDoneEvent.Broadcast(recordID);
+	}
+}
+
+void TextureGraphInsightSession::BatchJobsDone(JobBatchPtr Batch)
+{
+	UE_LOG(LogTextureGraphInsight, Log, TEXT("TextureGraphInsightSchedulerObserver::BatchDone"));
+
+	auto recordID = GetRecord().FindBatchRecord(Batch->GetBatchId());
+
+	GetRecord().EditBatchRecord(recordID,
+		[this, Batch](BatchRecord& BufferRecord)
+		{
+			UpdateBatchRecord(Batch, BufferRecord);
+			BufferRecord.bIsJobsDone = true;
+			return true;
+		});
+
+	if (recordID.IsValid())
+	{
+		OnBatchJobsDoneEvent.Broadcast(recordID);
+	}
+}
+
+
+void TextureGraphInsightSession::SendToInspector(RecordID recordID) const
+{
+	if (recordID.IsValid())
+	{
+		if (recordID.IsBatch())
+		{
+			OnBatchInspectedEvent.Broadcast(recordID);
+		}
+		else if (recordID.IsJob())
+		{
+			OnJobInspectedEvent.Broadcast(recordID);
+		}
+		else if (recordID.IsBlob())
+		{
+			OnBlobInspectedEvent.Broadcast(recordID);
+		}
+		else if (recordID.IsBuffer())
+		{
+			OnBufferInspectedEvent.Broadcast(recordID);
+		}
+	}
+}
+
+
+BatchRecord TextureGraphInsightSession::MakeBatchRecord(RecordID RecId, JobBatchPtr Batch, bool isDone)
+{
+	auto MixObj = Batch->GetCycle()->GetMix();
+	auto mixRid = AssociateBatchToMix(RecId, MixObj);
+
+	BatchRecord record;
+	record.SelfID = RecId;
+	record.BatchID = Batch->GetBatchId();
+	record.FrameID = Batch->GetFrameId();
+	record.BeginTimeMS = Batch->GetStartTime();
+	record.ReplayCount = Batch->GetReplayCount();
+	record.MixID = mixRid;
+	record.bIsDone = isDone;
+
+	record.bIsNoCache = Batch->GetCycle()->NoCache();
+	record.bIsFromIdle = Batch->IsIdle();
+
+
+	size_t numJobs = Batch->NumJobs();
+	record.Jobs.resize(numJobs);
+	std::vector<Job*> jobPtrs(numJobs, nullptr); // an array of the job pointers to retreive dependencies
+
+	size_t jobIndex = 0;
+
+	for (size_t j = 0; j < numJobs; ++j)
+	{
+		JobPtrW job = Batch->GetJobAt(j);
+		jobPtrs[j] = job.lock().get();
+
+		Job* generatorJobPtr = jobPtrs[j]->GeneratorJob().lock().get();
+
+		if (generatorJobPtr && generatorJobPtr != jobPtrs[j])
+		{
+			// find the generator job in the list of knwon jobs
+			auto genJobIdx = std::find(jobPtrs.begin(), jobPtrs.begin() + j, generatorJobPtr);
+			if (genJobIdx == jobPtrs.end())
+				check(false); /// Should never happen
+
+			record.Jobs[jobIndex++] = MakeJobRecord(job, (genJobIdx - jobPtrs.begin()), -1);
+		}
+		else
+		{
+			record.Jobs[jobIndex++] = MakeJobRecord(job, j);
+		}
+	}
+
+	/// RVO should handle this
+	return record;
+}
+
+void TextureGraphInsightSession::UpdateBatchRecord(JobBatchPtr Batch, BatchRecord& record) {
+	record.BeginTimeMS = Batch->GetStartTime();
+	record.EndTimeMS = Batch->GetEndTime();
+	record.ReplayCount = Batch->GetReplayCount();
+
+	auto cycle = Batch->GetCycle();
+
+	auto numJobs = Batch->NumJobs();
+	for (uint64_t j = 0; j < numJobs; ++j)
+	{
+		JobPtrW job = Batch->GetJobAt(j);
+
+		auto& rj = record.Jobs[j];
+		UpdateJobRecord(job, rj, record.BeginTimeMS, record.EndTimeMS);
+		rj.bIsDone = record.bIsDone;
+
+		record.NumInvalidatedTiles += rj.GetNumInvalidated();
+		record.NumTiles += rj.GetNumTiles();
+	}
+
+
+	// Capture the result target blobs
+	record.ResultBlobs.resize((int32)TextureType::Count);
+	if (cycle->GetNumTargets() && cycle->GetTarget(0))
+	{
+		TextureSet& lastRender = cycle->GetTarget(0)->GetLastRender();
+		for (int32 tti = 0; tti < (int32)TextureType::Count; tti++)
+		{
+			auto BlobObj = lastRender.GetTexture(tti);
+			if (BlobObj)
+			{
+				//auto bid = FindOrCreateBlobRecord(BlobObj, record._selfID);
+				//record._resultBlobs[tti] = bid;
+			}
+		}
+	}
+
+	EmitOnNewBlob();
+}
+
+JobRecord TextureGraphInsightSession::MakeJobRecord(JobPtrW job, uint32 idx, int32 phaseIdx)
+{
+	JobRecord record;
+	JobPtr sjob = job.lock();
+
+	// record._jobID = job->ID();
+	record.JobIdx = idx;
+	record.PhaseIdx = phaseIdx;
+
+	record.Priority = sjob->GetPriority();
+	record.ReplayCount = sjob->GetReplayCount();
+
+	if (sjob->GetTiled()) {
+		record.Tiles.Resize(1, 1);
+	}
+
+	auto r = sjob->GetResult();
+	if (r) 
+	{
+		record.TexWidth = r->GetWidth();
+		record.TexHeight = r->GetHeight();
+
+		record.Tiles.Resize(r->Rows(), r->Cols());
+	}
+
+	auto stats = sjob->GetStats();
+	record.BeginTimeMS = stats.BeginNativeTime;
+	record.BeginRunTimeMS = stats.BeginRunTime;
+	record.EndTimeMS = stats.EndNativeTime;
+
+	record.TransformName = sjob->GetTransform()->GetName();
+
+	return record;
+}
+
+void TextureGraphInsightSession::UpdateJobRecord(JobPtrW job, JobRecord& record, double batchBeginTime_ms, double batchEndTime_ms)
+{
+	JobPtr sjob = job.lock();
+
+	if (sjob->IsCulled())
+		return;
+
+	record.JobHash = sjob->Hash()->Value();
+
+	auto stats = sjob->GetStats();
+
+	if (stats.BeginNativeTime < batchBeginTime_ms)
+		record.BeginTimeMS = batchBeginTime_ms;
+	else
+		record.BeginTimeMS = stats.BeginNativeTime;
+
+	if (stats.BeginRunTime == 0)
+		record.BeginRunTimeMS = 0; // record._beginTime_ms; // If the "runtime" stat is invalid then use the begin time
+	else
+		record.BeginRunTimeMS = stats.BeginRunTime;
+
+	if (stats.EndNativeTime < record.BeginRunTimeMS)
+		record.EndTimeMS = record.BeginRunTimeMS;
+	else if (stats.EndNativeTime > batchEndTime_ms)
+		record.EndTimeMS = batchEndTime_ms;
+	else
+		record.EndTimeMS = stats.EndNativeTime;
+
+	record.ReplayCount = sjob->GetReplayCount();
+
+	// Only on first replay capture all the inputs
+	if (record.ReplayCount == 0)
+	{
+		record.InputBlobIds.clear();
+		for (uint32_t i = 0; i < sjob->NumArgs(); ++i)
+		{
+			auto arg = sjob->GetArg(i);
+			auto blobArg = std::dynamic_pointer_cast<JobArg_Blob>(arg);
+			if (blobArg)
+			{
+				auto inputBlob = blobArg->GetBlob();
+
+				if (inputBlob) 
+				{
+					record.InputBlobIds.emplace_back(FindOrCreateBlobRecord(inputBlob.get()));
+					record.InputArgNames.emplace_back(blobArg->Target());
+				}
+			}
+		}
+	}
+
+	// Deal with output
+	{
+		auto r = sjob->GetResult();
+		// In case of a replay update, all the inputs should be the same, 
+		// the generated blobs are replayed too, let's imply call for an update
+		if (record.ReplayCount && record.ResultBlobId.IsValid())
+		{
+			UpdateBlobRecord(record.ResultBlobId, std::static_pointer_cast<Blob>(r.get()));
+		}
+		// Not a replay but the first time the job is over, generate the resulting BlobObj(s)
+		else if (r)
+		{
+			record.ResultBlobId = FindOrCreateBlobRecord(std::static_pointer_cast<Blob>(r.get()), record.SelfID);
+			record.TexWidth = r->GetWidth();
+			record.TexHeight = r->GetHeight();
+
+			record.Tiles.Resize(r->Rows(), r->Cols());
+
+			record.Tiles.ForEach([&](uint8_t tx, uint8_t ty) 
+			{
+				record.Tiles(tx, ty) = sjob->GetTileInvalidation(tx, ty);
+			});
+		}
+	}
+}
+
+RecordID TextureGraphInsightSession::FindOrCreateBlobRecordFromHash(HashType hash)
+{
+	//auto RecId = Record().FindBlobRecord((uint64_t)BlobObj.get());
+	auto RecId = GetRecord().FindBlobRecord(hash);
+	if (!RecId.IsValid())
+	{
+		auto BlobObj = TextureGraphEngine::GetInstance()->GetBlobber()->Find(hash);
+		if (BlobObj)
+		{
+			return FindOrCreateBlobRecord(BlobObj.get());
+		}
+		return RecordID();
+	}
+	else
+		return RecId;
+}
+
+RecordID TextureGraphInsightSession::FindOrCreateBlobRecord(BlobPtr BlobObj, RecordID sourceID)
+{
+	if (BlobObj)
+	{
+		auto hash = BlobObj->Hash()->Value();
+		auto hashIsFinal = BlobObj->Hash()->IsFinal();
+
+		//auto RecId = Record().FindBlobRecord((uint64_t)BlobObj.get());
+		auto RecId = GetRecord().FindBlobRecord(hash);
+		if (!RecId.IsValid())
+		{
+			// Record the BlobObj right away (without tiling information)
+			//RecId = Record().NewBlobRecord((uint64_t)BlobObj.get(), MakeBlobRecord(BlobObj, sourceID));
+			RecId = GetRecord().NewBlobRecord(hash, 
+				[this, BlobObj, sourceID](RecordID srid)
+				{
+					return MakeBlobRecord(srid, BlobObj, sourceID);
+				});
+
+			// Cache the live pointer to the BlobObj 
+			GetCache().AddBlob(RecId, BlobObj);
+
+			// Remember the recordID of the new BlobObj for notification (from the batchRecord handler)
+			NewBlobsToNotify.emplace_back(RecId);
+
+			// Parse the sub tiles next (probably creating the BlobObj records)
+			if (BlobObj->IsTiled())
+			{
+				TiledBlobPtr tiledBlob = std::static_pointer_cast<TiledBlob>(BlobObj);
+
+				RecordIDTiles  subTiles(tiledBlob->Rows(), tiledBlob->Cols());
+				subTiles.ForEach([&](uint8_t tx, uint8_t ty) {
+					subTiles(tx, ty) = FindOrCreateBlobRecord(tiledBlob->GetTile(tx, ty).get(), RecId);
+				});
+
+				// And then update the tiled BlobObj record of the parent
+				GetRecord().EditBlobRecord(RecId, [&](BlobRecord& rd)
+				{
+					EditTiledBlobRecord(rd, std::move(subTiles));
+					return true;
+				});
+			}
+		}
+		return RecId;
+	}
+
+	return RecordID();
+}
+
+RecordID TextureGraphInsightSession::RemapBlobHash(HashType oldHash, HashType newHash)
+{
+	auto old_rid = GetRecord().FindBlobRecord(oldHash);
+	auto new_rid = GetRecord().FindBlobRecord(newHash);
+	if (new_rid.IsValid() && old_rid.IsValid())
+	{
+		GetRecord().EditBlobRecord(old_rid, [=](BlobRecord& BufferRecord) { BufferRecord.MappedID = new_rid; return true; });
+		GetRecord().EditBlobRecord(new_rid, [=](BlobRecord& BufferRecord) { BufferRecord.NumMapped++; return true; });
+
+		MappedBlobsToNotify.emplace_back(old_rid);
+	}
+
+	return old_rid;
+}
+
+
+BlobRecord TextureGraphInsightSession::MakeBlobRecord(RecordID RecId, BlobPtr BlobObj, RecordID sourceID)
+{
+	BlobRecord record;
+	record.SourceID = sourceID;
+	record.Name = BlobObj->Name();
+	record.HashValue = BlobObj->Hash()->Value();
+	record.TexWidth = BlobObj->GetWidth();
+	record.TexHeight = BlobObj->GetHeight();
+
+	return record;
+}
+
+
+void TextureGraphInsightSession::EditTiledBlobRecord(BlobRecord& rd, const RecordIDTiles& subTiles)
+{
+	rd.TileIdxs.Resize(subTiles.Grid());
+	for (int i = 0; i < subTiles.Grid().NumTiles(); ++i) {
+		auto f = std::find(rd.UniqueBlobs.begin(), rd.UniqueBlobs.end(), subTiles._tiles[i].Blob());
+		if (f == rd.UniqueBlobs.end())
+		{
+			rd.TileIdxs._tiles[i] = (uint8_t)rd.UniqueBlobs.size();
+			rd.UniqueBlobs.emplace_back(subTiles._tiles[i].Blob());
+		}
+		else
+		{
+			rd.TileIdxs._tiles[i] = (uint8_t)(f - rd.UniqueBlobs.begin());
+		}
+	}
+}
+
+void  TextureGraphInsightSession::UpdateBlobRecord(RecordID RecId, BlobPtr BlobObj)
+{
+	if (BlobObj && RecId.IsBlob())
+	{
+		GetRecord().EditBlobRecord(RecId, [this, BlobObj](BlobRecord& BufferRecord)
+		{
+			if (!BufferRecord.IsTiled())
+				BufferRecord.RawBufferMemSize = BlobObj->GetBufferRef()->MemSize();
+
+			// Simply update the replay count, should be the only thing changed
+			BufferRecord.ReplayCount = BlobObj->GetReplayCount();
+
+			// Blob is really tiled, then dispatch updates to sub tiles
+			if (BlobObj->IsTiled() && BufferRecord.IsTiled())
+			{
+				TiledBlobPtr tiledBlob = std::static_pointer_cast<TiledBlob>(BlobObj);
+
+				BufferRecord.TileIdxs._grid.ForEach([&](uint32 tx, uint32 ty) {
+					UpdateBlobRecord(BufferRecord.GetTileBlob(tx, ty), tiledBlob->GetTile(tx, ty).get());
+				});
+			}
+
+			return true;
+		});
+	}
+}
+
+void TextureGraphInsightSession::EmitOnNewBlob()
+{
+	auto addedBlobs = std::move(NewBlobsToNotify);
+	NewBlobsToNotify.clear();
+
+	OnBlobAddedEvent.Broadcast(addedBlobs);
+}
+
+void TextureGraphInsightSession::EmitOnMappedBlob()
+{
+	auto mappedBlobs = std::move(MappedBlobsToNotify);
+	MappedBlobsToNotify.clear();
+
+	OnBlobMappedEvent.Broadcast(mappedBlobs);
+}
+
+
+bool TextureGraphInsightSession::ReplayBatch(RecordID batchRecordID, bool captureRenderDoc)
+{
+	// Get to the Batch ptr if still alive:
+	auto batchPtr = GetCache().GetBatch(batchRecordID);
+	if (batchPtr)
+	{
+		// configure the replay and spawn it
+		batchPtr->ResetForReplay();
+		batchPtr->SetCaptureRenderDoc(false);
+		if (captureRenderDoc)
+			TextureGraphEngine::GetInstance()->GetScheduler()->SetCaptureRenderDocNextBatch(true);
+		TextureGraphEngine::GetInstance()->GetScheduler()->AddBatch(batchPtr);
+		return true;
+	}
+	return false;
+}
+
+
+bool TextureGraphInsightSession::ReplayJob(RecordID jobRecordID, bool captureRenderDoc)
+{
+	// Get to the Batch ptr if still alive:
+	auto batchPtr = GetCache().GetBatch(jobRecordID);
+	if (batchPtr)
+	{
+		// configure the replay and spawn it
+		batchPtr->ResetForReplay(jobRecordID.Job());
+		batchPtr->SetCaptureRenderDoc(false);
+		if(captureRenderDoc)
+			TextureGraphEngine::GetInstance()->GetScheduler()->SetCaptureRenderDocNextBatch(true);
+		TextureGraphEngine::GetInstance()->GetScheduler()->AddBatch(batchPtr);
+
+		return true;
+	}
+	return false;
+}
+
+
+RecordID TextureGraphInsightSession::AssociateBatchToMix(RecordID Batch, MixPtr MixObj)
+{
+	auto mixID = FindOrCreateMixRecord(MixObj);
+
+	if (mixID.IsValid())
+	{
+		if (GetRecord().EditMixRecord(mixID,
+			[Batch](MixRecord& mr)
+			{
+				mr.Batches.emplace_back(Batch);
+				return true;
+			}))
+		{
+			OnMixUpdatedEvent.Broadcast(mixID);
+		}
+	}
+
+	return mixID;
+}
+
+RecordID TextureGraphInsightSession::FindOrCreateMixRecord(MixPtr MixObj)
+{
+	if (MixObj)
+	{
+		//auto hash = MixObj->Hash()->Value();
+		auto hash = (uint64)MixObj;
+		
+		auto RecId = GetRecord().FindMixRecord(hash);
+		if (!RecId.IsValid())
+		{
+			UMix* mAsMix = dynamic_cast<UMix*>(MixObj);
+
+			// if an Instance, find or allocate the parent MixObj
+			RecordID		parentID;
+
+			// Record the MixObj
+			RecId = GetRecord().NewMixRecord(hash,
+				[=](RecordID srid)
+				{
+					MixRecord record;
+					record.Name = MixObj->GetName();
+
+					record.ParentMixID = parentID;
+
+					return record;
+				});
+
+			// Cache the live pointer to the MixObj
+			GetCache().AddMix(RecId, MixObj);
+
+			OnMixAddedEvent.Broadcast(RecId);
+
+			// If parent ID is valid, add this knowledge to parent record and notify for update
+			if (parentID.IsValid())
+			{
+				GetRecord().EditMixRecord(parentID,
+					[=](MixRecord& mr)
+					{
+						mr.InstanceMixIDs.emplace_back(RecId);
+						return true;
+					});
+
+				OnMixUpdatedEvent.Broadcast(parentID);
+			}
+		}
+
+		// Return the RecId already existing or the brand new one
+		return RecId;
+	}
+
+	return RecordID();
+}
diff --git a/Engine/Plugins/Experimental/TextureGraph/Source/TextureGraphInsight/Private/TextureGraphInsight.cpp b/Engine/Plugins/Experimental/TextureGraph/Source/TextureGraphInsight/Private/TextureGraphInsight.cpp
new file mode 100644
--- /dev/null
+++ b/Engine/Plugins/Experimental/TextureGraph/Source/TextureGraphInsight/Private/TextureGraphInsight.cpp
@@ -0,0 +1,70 @@
+// Copyright Epic Games, Inc. All Rights Reserved.
+#include "TextureGraphInsight.h"
+#include "Modules/ModuleManager.h"
+#include "TextureGraphEngine.h"
+
+#include "Model/TextureGraphInsightObserver.h"
+#include "Model/TextureGraphInsightSession.h"
+#include "Shader.h"
+
+IMPLEMENT_MODULE(FTextureGraphInsightModule, TextureGraphInsight);
+DEFINE_LOG_CATEGORY(LogTextureGraphInsight);
+
+void FTextureGraphInsightModule::StartupModule()
+{
+	FDefaultGameModuleImpl::StartupModule();
+}
+
+
+void FTextureGraphInsightModule::ShutdownModule()
+{
+	FShaderType::Uninitialize();
+}
+//////////////////////////////////////////////////////////////////////////
+
+TextureGraphInsight* TextureGraphInsight::GInstance = nullptr;
+
+TextureGraphInsight::TextureGraphInsight()
+{
+	UE_LOG(LogTextureGraphInsight, Log, TEXT("Initialising the TextureGraphInsight!"));
+
+	Session = std::make_shared<TextureGraphInsightSession>();
+}
+
+TextureGraphInsight::~TextureGraphInsight()
+{
+	UE_LOG(LogTextureGraphInsight, Log, TEXT("Destroying the TextureGraphInsight!"));
+}
+
+bool TextureGraphInsight::Create()
+{
+	/// Cannot create a new instance before destroying the old one
+	if (!GInstance)
+	{
+		GInstance = new TextureGraphInsight();
+
+		/// Need the engine observer to watch what is happening
+		auto EngineObserver = std::make_shared<TextureGraphInsightEngineObserver>();
+		TextureGraphEngine::RegisterObserverSource(EngineObserver); // will also install other observers
+
+		return true;
+	}
+	return false;
+}
+
+bool TextureGraphInsight::Destroy()
+{
+	/// Destroy an existing instance, no op otherwise
+	if (GInstance)
+	{
+		/// Remove Engine observer
+		TextureGraphEngine::RegisterObserverSource(nullptr);
+
+		delete GInstance;
+		GInstance = nullptr;
+
+		return true;
+	}
+
+	return false;
+}
diff --git a/Engine/Plugins/Experimental/TextureGraph/Source/TextureGraphInsight/Private/View/STextureGraphInsightActionView.cpp b/Engine/Plugins/Experimental/TextureGraph/Source/TextureGraphInsight/Private/View/STextureGraphInsightActionView.cpp
new file mode 100644
--- /dev/null
+++ b/Engine/Plugins/Experimental/TextureGraph/Source/TextureGraphInsight/Private/View/STextureGraphInsightActionView.cpp
@@ -0,0 +1,202 @@
+// Copyright Epic Games, Inc. All Rights Reserved.
+#include "View/STextureGraphInsightActionView.h"
+
+#include "TextureGraphInsight.h"
+#include "Model/TextureGraphInsightSession.h"
+
+
+class STextureGraphInsightActionViewRow : public SMultiColumnTableRow<STextureGraphInsightActionView::FItem>
+{
+public:
+	using FItem = STextureGraphInsightActionView::FItem;
+
+	SLATE_BEGIN_ARGS(STextureGraphInsightActionViewRow) {}
+		SLATE_ARGUMENT(FItem, Item)
+	SLATE_END_ARGS()
+
+public:
+
+	enum Column {
+		Main = 0,
+		Name,
+
+		NUM_COLUMNS,
+	};
+	using ColumnInfo = std::pair<FName, float>;
+	static ColumnInfo s_columnNames[NUM_COLUMNS];
+	static Column NameToColumn(const FName& name) {
+		for (int i = 0; i < NUM_COLUMNS; ++i)
+		{
+			if (name == s_columnNames[i].first)
+				return Column(i);
+		}
+		return NUM_COLUMNS;
+	}
+
+	void Construct(const FArguments& InArgs, const TSharedRef<STableViewBase>& InOwnerTableView)
+	{
+		_recordID = InArgs._Item->_recordID;
+		SMultiColumnTableRow<FItem >::Construct(FSuperRowType::FArguments(), InOwnerTableView);	
+	}
+
+	virtual TSharedRef<SWidget> GenerateWidgetForColumn(const FName& columnName) override
+	{
+		if (_recordID.IsValid()) {
+			Column column = NameToColumn(columnName);
+			if (column == Main) {
+				// Rows in a TreeView need an expander button and some indentation
+				return SNew(SHorizontalBox)
+					+ SHorizontalBox::Slot()
+					.AutoWidth()
+					.HAlign(HAlign_Right)
+					.VAlign(VAlign_Fill)
+					[
+						SNew(SExpanderArrow, SharedThis(this))
+						.StyleSet(ExpanderStyleSet)
+					]
+					+ SHorizontalBox::Slot()
+					.FillWidth(1)
+					.Padding(FMargin(4.0f, 0.0f))
+					.VAlign(VAlign_Center)
+					[
+						SAssignNew(_textBoxes[column], STextBlock)
+						.Text(GetTextForColumn(column))
+					];
+			} else {
+				return SNew(SBox)
+					.Padding(FMargin(4.0f, 0.0f))
+					.VAlign(VAlign_Center)
+					[
+						SAssignNew(_textBoxes[column], STextBlock)
+						.Text(GetTextForColumn(column))
+						.OnDoubleClicked(this, &STextureGraphInsightActionViewRow::OnDoubleClickedResultColumn)
+					];
+			}
+		}
+		// default to null widget if property cannot be found
+		return SNullWidget::NullWidget;
+	}
+
+	FText GetTextForColumn(Column column) const
+	{
+		const auto& ar = TextureGraphInsight::Instance()->GetSession()->GetRecord().GetAction(_recordID);
+
+		FString s;
+		switch (column)
+		{
+		case Main:
+				s = FString::FromInt(_recordID.Action()) + " " + ar.Name;
+			break;
+		case Name:
+				if (ar.SubActions.empty())
+					s = ar.Meta;
+			break;
+		}
+
+		return FText::FromString(s);
+	}
+
+	void Refresh()
+	{
+		_textBoxes[Main]->SetText(GetTextForColumn(Main));
+		_textBoxes[Name]->SetText(GetTextForColumn(Name));
+	}
+
+	FReply OnDoubleClickedResultColumn(/** The geometry of the widget*/
+		const FGeometry&,
+		/** The Mouse Event that we are processing */
+		const FPointerEvent&)
+	{
+		return FReply::Handled();
+	}
+
+protected:
+	RecordID _recordID;
+	TSharedPtr<STextBlock> _textBoxes[NUM_COLUMNS];
+};
+
+STextureGraphInsightActionViewRow::ColumnInfo STextureGraphInsightActionViewRow::s_columnNames[] = {
+	{ FName(TEXT("Action")), 0.1},    // main
+	{ FName(TEXT("Name")), 0.3},
+};
+
+void STextureGraphInsightActionView::Construct(const FArguments& Args)
+{
+
+	TSharedPtr<SHeaderRow> headerRow = SNew(SHeaderRow);
+
+	for (int i = 0; i < STextureGraphInsightActionViewRow::NUM_COLUMNS; i++)
+	{
+		headerRow->AddColumn(
+			SHeaderRow::Column(STextureGraphInsightActionViewRow::s_columnNames[i].first)
+			.DefaultLabel(FText::FromString(STextureGraphInsightActionViewRow::s_columnNames[i].first.ToString()))
+			.FillWidth(STextureGraphInsightActionViewRow::s_columnNames[i].second)
+		);
+	}
+
+	ChildSlot
+	[
+		SAssignNew(_tableView, SItemTableView)
+		.ItemHeight(24)
+		.TreeItemsSource(&_rootItems)
+		.OnGenerateRow(this, &STextureGraphInsightActionView::OnGenerateRowForView)
+		.OnGetChildren(this, &STextureGraphInsightActionView::OnGetChildrenForView)
+		.OnMouseButtonDoubleClick(this, &STextureGraphInsightActionView::OnDoubleClickItemForView)
+		.OnMouseButtonClick(this, &STextureGraphInsightActionView::OnClickItemForView)
+		.HeaderRow(headerRow)
+	];
+
+	// install the observer notifications
+	auto sr = StaticCastSharedRef<STextureGraphInsightActionView>(this->AsShared());
+	TextureGraphInsight::Instance()->GetSession()->OnEngineReset().AddSP(sr, &STextureGraphInsightActionView::OnEngineReset);
+}
+
+TSharedRef<ITableRow> STextureGraphInsightActionView::OnGenerateRowForView(FItem item, const TSharedRef<STableViewBase>& OwnerTable)
+{
+	return SAssignNew(item->_widget, STextureGraphInsightActionViewRow, OwnerTable).Item(item);
+}
+
+void STextureGraphInsightActionView::OnClickItemForView(FItem item)
+{
+	// Item clicked, let's send it to the inspector
+	auto record = item->_recordID;
+	//TextureGraphInsight::Instance()->Session()->SendToInspector(record);
+}
+void STextureGraphInsightActionView::OnDoubleClickItemForView(FItem item)
+{
+	// Item double clicked, let's send it to the inspector
+	auto record = item->_recordID;
+	//TextureGraphInsight::Instance()->Session()->SendToInspector(record);
+}
+
+void STextureGraphInsightActionView::OnActionNew(RecordID rid)
+{
+	// Fetch the actual data from the record and add job's items
+	const auto& ar = TextureGraphInsight::Instance()->GetSession()->GetRecord().GetAction(rid);
+
+	// Do not worry about actions with parent, we will get notified for the actual parent and
+	// will populate the sub ations consequently
+	if (ar.ParentAction.IsValid())
+		return;
+
+	FItem item = MakeShareable(new FItemData(rid));
+
+	for (const auto& sa : ar.SubActions)
+	{
+		item->_children.Add(MakeShareable(new FItemData(sa)));
+	}
+
+	_rootItems.Add(item);
+
+	if (item->_children.Num())
+		_tableView->SetItemExpansion(item, true);
+
+	_tableView->RequestTreeRefresh();
+}
+
+void STextureGraphInsightActionView::OnEngineReset(int id)
+{
+	_rootItems.Empty();
+	_tableView->RequestTreeRefresh();
+}
+
diff --git a/Engine/Plugins/Experimental/TextureGraph/Source/TextureGraphInsight/Private/View/STextureGraphInsightBlobView.cpp b/Engine/Plugins/Experimental/TextureGraph/Source/TextureGraphInsight/Private/View/STextureGraphInsightBlobView.cpp
new file mode 100644
--- /dev/null
+++ b/Engine/Plugins/Experimental/TextureGraph/Source/TextureGraphInsight/Private/View/STextureGraphInsightBlobView.cpp
@@ -0,0 +1,178 @@
+// Copyright Epic Games, Inc. All Rights Reserved.
+#include "View/STextureGraphInsightBlobView.h"
+
+#include "2D/Tex.h"
+#include "TextureGraphInsight.h"
+#include "Model/TextureGraphInsightSession.h"
+
+#include "Device/FX/DeviceBuffer_FX.h"
+
+#include "Widgets/Layout/SScaleBox.h"
+#include "Materials/Material.h"
+
+#include "View/STextureGraphInsightDeviceBufferView.h"
+#include <Widgets/Layout/SUniformGridPanel.h>
+#include <Materials/MaterialInstanceDynamic.h>
+#include <Brushes/SlateColorBrush.h>
+#include <Widgets/SBoxPanel.h>
+#include <Widgets/Text/STextBlock.h>
+#include <Widgets/SOverlay.h>
+
+
+void STextureGraphInsightBlobView::Construct(const FArguments& InArgs)
+{
+	RecordID recordID = InArgs._recordID;
+	bool withHighlight = InArgs._withHighlight;
+	BoolTiles tilesMask = InArgs._tilesMask;
+
+	_recordID = recordID;
+
+	// Check for a valid blob id before anything
+	if (!(recordID.IsValid() && recordID.IsBlob()))
+	{
+		return;
+	}
+
+	uint32_t SIMAGE_WIDTH = 512;
+	uint32_t SIMAGE_HEIGHT = 512;
+
+	const float TEXT_PADDING = 3.0;
+	const float BLOB_PADDING = 0.5;
+
+	bool addTouchedBorder = withHighlight;
+
+	_borderBrush = MakeShareable(new FSlateColorBrush(FColor(255, 37, 18)));
+
+
+	// Describe the blob, identify if it is tiled or not
+	const auto& br = TextureGraphInsight::Instance()->GetSession()->GetRecord().GetBlob(recordID);
+
+	if (br.IsTiled()) {
+		_inspectedID = _recordID; // Default inspected is self if tiled.
+
+		FString s = br.Name;
+
+
+		if (br.TexWidth || br.TexHeight)
+			s += "\n" + FString::FromInt(br.TexWidth) + " x " + FString::FromInt(br.TexHeight);
+
+		{ // Add some meta data specific to a tiled blob
+			if (br.NumUniqueBlobs() < br.NumTiles())
+			{
+				s += "    " + FString::FromInt(br.NumUniqueBlobs()) + "unique tiles / ";
+				s += (br.Grid().IsUnique() ? "1" : FString::FromInt(br.Grid().Rows()) + " x " + FString::FromInt(br.Grid().Cols()));
+			}
+		}
+		if (br.ReplayCount > 0)
+			s += " replay #" + FString::FromInt(br.ReplayCount);
+
+		TSharedPtr<SUniformGridPanel> gridPanel;
+		ChildSlot
+			[
+				SNew(SVerticalBox)
+				+ SVerticalBox::Slot()
+				.Padding(TEXT_PADDING)
+				.AutoHeight()
+				[
+					SNew(STextBlock)
+					.Text(FText::FromString(s))
+					.AutoWrapText(true)
+				]
+				+ SVerticalBox::Slot()
+				[
+					SAssignNew(gridPanel, SUniformGridPanel)
+				]
+			];
+
+		bool isSingleTiledBlob = br.Grid().IsUnique();
+
+		br.Grid().ForEach([&](TiledGrid::Size tx, TiledGrid::Size ty)
+		{
+			bool highlight = false;
+			if (tilesMask.Grid().IsValidTile(tx, ty))
+			{
+				highlight = tilesMask(tx, ty) != 0;
+			}
+
+			auto subTileRid = br.GetTileBlob(tx, ty);
+
+			if (isSingleTiledBlob || (subTileRid.id == _recordID.id)) /// In some case, the sub tile is the same as the parent tile, capture this case
+			{
+				RecordID bufferID = TextureGraphInsight::Instance()->GetSession()->GetRecord().FindDeviceBufferRecordFromHash(br.HashValue);
+				bool blobContentInBrush = false;
+				_inspectedID = bufferID;
+
+				gridPanel->AddSlot(tx, ty)
+				[
+					SNew(SOverlay)
+					+ SOverlay::Slot()
+					[
+						SNew(STextureGraphInsightDeviceBufferView)
+						.recordID(bufferID)
+						.blobID(_recordID)
+						.withToolTip(true)
+					]
+				];
+			}
+			else
+			{
+				gridPanel->AddSlot(tx, ty)
+				[
+					SNew(STextureGraphInsightBlobView)
+					.recordID(subTileRid)
+					.withHighlight(highlight)
+				];
+			}
+		});
+			}
+	else
+	{
+		RecordID bufferID = TextureGraphInsight::Instance()->GetSession()->GetRecord().FindDeviceBufferRecordFromHash(br.HashValue);
+
+		bool blobContentInBrush = false;
+		FString s = br.Name;
+		if (br.TexWidth || br.TexHeight)
+			s += "\n" + FString::FromInt(br.TexWidth) + " x " + FString::FromInt(br.TexHeight);
+		if (br.ReplayCount > 0)
+			s += " replay #" + FString::FromInt(br.ReplayCount);
+		
+		_inspectedID = bufferID;
+
+		ChildSlot
+			[
+				SNew(SOverlay)
+				+ SOverlay::Slot()
+				.Padding(BLOB_PADDING )
+				[
+					SNew(STextureGraphInsightDeviceBufferView)
+					.recordID(bufferID)
+					.blobID(_recordID)
+					.withToolTip(true)
+				]
+			];
+	}
+
+	if (addTouchedBorder)
+	{
+		SetBorderImage(_borderBrush.Get());
+	}
+};
+
+FReply STextureGraphInsightBlobView::OnInspect()
+{
+	TextureGraphInsight::Instance()->GetSession()->SendToInspector(_inspectedID);
+
+	return FReply::Handled();
+}
+
+
+FReply STextureGraphInsightBlobView::OnMouseButtonDown(const FGeometry& MyGeometry, const FPointerEvent& MouseEvent)
+{
+	//return OnInspect();
+	return FReply::Unhandled();
+}
+
+FReply STextureGraphInsightBlobView::OnMouseButtonDoubleClick(const FGeometry& InMyGeometry, const FPointerEvent& InMouseEvent)
+{
+	return OnInspect();
+}
diff --git a/Engine/Plugins/Experimental/TextureGraph/Source/TextureGraphInsight/Private/View/STextureGraphInsightDeviceBufferView.cpp b/Engine/Plugins/Experimental/TextureGraph/Source/TextureGraphInsight/Private/View/STextureGraphInsightDeviceBufferView.cpp
new file mode 100644
--- /dev/null
+++ b/Engine/Plugins/Experimental/TextureGraph/Source/TextureGraphInsight/Private/View/STextureGraphInsightDeviceBufferView.cpp
@@ -0,0 +1,188 @@
+// Copyright Epic Games, Inc. All Rights Reserved.
+#include "View/STextureGraphInsightDeviceBufferView.h"
+
+#include "2D/Tex.h"
+#include "TextureGraphInsight.h"
+#include "Model/TextureGraphInsightSession.h"
+
+#include "TextureGraphEngine.h"
+#include "Device/DeviceManager.h"
+#include "Device/FX/DeviceBuffer_FX.h"
+
+#include "Widgets/Layout/SScaleBox.h"
+#include "Materials/Material.h"
+#include <Materials/MaterialInstanceDynamic.h>
+#include <Styling/SlateBrush.h>
+#include <Widgets/SOverlay.h>
+#include <Widgets/Images/SImage.h>
+#include <Widgets/Text/STextBlock.h>
+#include <Widgets/Layout/SBox.h>
+#include <Widgets/SBoxPanel.h>
+#include <Brushes/SlateDynamicImageBrush.h>
+void STextureGraphInsightDeviceBufferView::Construct(const FArguments& InArgs)
+{
+	RecordID recordID = InArgs._recordID;
+	RecordID blobID = InArgs._blobID;
+	bool withToolTip = InArgs._withToolTip;
+	bool withDescription = InArgs._withDescription;
+
+	_recordID = recordID;
+	_blobID = blobID;
+
+	// allocate the material brush where we can assign the texture (or a copy of ) async, or just reflect the status
+	MakeWidgetMaterial();
+
+	// Build the basic widget eventually showing the checker board if no buffer found
+	TSharedPtr<SOverlay> overlay;
+	ChildSlot
+		[
+			SAssignNew(overlay, SOverlay)
+			+ SOverlay::Slot()
+		[
+			SNew(SImage)
+			.Image(_brush.Get())
+		]
+		];
+
+	// Check for a valid buffer id before anything
+	bool useBlobId = false;
+	//if (!(recordID.IsValid() && recordID.IsBuffer() && (recordID.Buffer() != INVALID_INDEX)))
+	{
+		useBlobId = true;
+		// if the Buffer ID is invalid, try to the Blob Id
+		if (!(blobID.IsValid() && blobID.IsBlob() && (blobID.Blob() != INVALID_INDEX)))
+		{
+			UE_LOG(LogTemp, Warning, TEXT("BlobTexture Invalid recordID for the buffer ?"));
+			return;
+		}	
+	}
+
+	// Now let's go find the true live DeviceBuffer ptr
+	BlobPtr blob;
+	if (useBlobId)
+	{
+		blob = TextureGraphInsight::Instance()->GetSession()->GetCache().GetBlob(blobID);
+	}
+	DeviceBufferPtr buffer;
+	if (blob)
+	{
+		buffer = blob->GetBufferRef().GetPtr();
+	} else
+	{
+		buffer = TextureGraphInsight::Instance()->GetSession()->GetCache().GetDeviceBuffer(recordID);
+	}
+
+	// Hopefully found a device buffer if not early exit
+	if (!buffer)
+		return;
+
+	uint32_t SIMAGE_WIDTH = 512;
+	uint32_t SIMAGE_HEIGHT = 512;
+
+	const float TEXT_PADDING = 3.0;
+	const float BLOB_PADDING = 2.0;
+
+	bool blobContentInBrush = false;
+
+	const auto& d = buffer->Descriptor();
+	FString s = d.Name;
+	if (d.Width || d.Height)
+		s += "\n" + FString::FromInt(d.Width) + " x " + FString::FromInt(d.Height);
+
+	s += "\n" + FString::FromInt(d.ItemsPerPoint) + " x " + FString(BufferDescriptor::FormatToString(d.Format));
+	s += "\n" + HashToFString(buffer->Hash()->Value());
+	//if (buffer->Hash()->HasTempDependency())
+	//	s += "\n" + HashToFString(buffer->PrevHashValue());
+
+
+	_imageWidth = std::min(d.Width, SIMAGE_WIDTH);
+	_imageHeight = std::min(d.Height, SIMAGE_HEIGHT);
+	_imageName = d.Name;
+
+	 // Got the description, now need to retreive the true live buffer and the texture to configure the brush
+	 {
+		// Hopefully found a device buffer
+		if (buffer)
+		{
+			// Let's create a texture from the raw data
+			auto tex = std::make_shared<Tex>();
+
+			if (!buffer->IsNull())
+			{
+				buffer->Raw()
+					.then([this, tex](RawBufferPtr raw) mutable
+						{
+							return tex->LoadRaw(raw);
+						})
+					.then([this, tex](auto result) mutable
+						{
+							if (result->IsOK())
+							{
+								check(tex && tex->GetTexture());
+								tex->SetFilter(TextureFilter::TF_Nearest);
+								UTexture* blobTexture = tex->GetTexture();
+								CreateWidgetBrush(blobTexture);
+							}
+							else
+							{
+								UE_LOG(LogTemp, Warning, TEXT("BlobTexture failed to create the widget material"));
+							}
+						})
+					.fail([this, tex]() mutable
+						{
+							UE_LOG(LogTemp, Warning, TEXT("BlobTexture failed to create the widget material"));
+						});
+			}
+			else
+			{
+
+			}
+		}
+		else
+		{
+			UE_LOG(LogTemp, Warning, TEXT("BlobTexture failed to find the buffer ?"));
+			// no buffer ?
+		}
+	}
+
+	if (withToolTip)
+	{
+		overlay->SetToolTipText(FText::FromString(s));
+	}
+	if (withDescription)
+	{
+		overlay->AddSlot()
+			.Padding(TEXT_PADDING)
+			.VAlign(EVerticalAlignment::VAlign_Bottom)
+			[
+				SNew(STextBlock)
+				.Text(FText::FromString(s))
+			.AutoWrapText(true)
+			];
+	}
+};
+
+void STextureGraphInsightDeviceBufferView::CreateWidgetBrush(UTexture* blobTexture) {
+	if (blobTexture && _brush && _brushMaterial) {
+		_brushMaterial->SetTextureParameterValue("BlobTex", blobTexture);
+		_brushMaterial->SetScalarParameterValue("BlobTexIsValid", 1.0);
+	}
+}
+
+void STextureGraphInsightDeviceBufferView::MakeWidgetMaterial() {
+	_brush = MakeShareable(
+		new FSlateDynamicImageBrush(
+			(UTexture2D*) nullptr,
+			FVector2D(_imageWidth, _imageHeight),
+			//desiredSize,
+			FName(_imageName),
+			FLinearColor(1.0f, 1.0f, 1.0f, 1.0f),
+			ESlateBrushTileType::NoTile,
+			ESlateBrushImageType::Linear
+		)
+	);
+	UMaterial* material = LoadObject<UMaterial>(nullptr, TEXT("/TextureGraph/UI/Materials/BlobView"));
+	_brushMaterial = (UMaterialInstanceDynamic::Create(material, nullptr));
+	_brushMaterial->SetScalarParameterValue("BlobTexIsValid", 0.0);
+	_brush->SetResourceObject(_brushMaterial);
+}
diff --git a/Engine/Plugins/Experimental/TextureGraph/Source/TextureGraphInsight/Private/View/STextureGraphInsightDeviceView.cpp b/Engine/Plugins/Experimental/TextureGraph/Source/TextureGraphInsight/Private/View/STextureGraphInsightDeviceView.cpp
new file mode 100644
--- /dev/null
+++ b/Engine/Plugins/Experimental/TextureGraph/Source/TextureGraphInsight/Private/View/STextureGraphInsightDeviceView.cpp
@@ -0,0 +1,645 @@
+// Copyright Epic Games, Inc. All Rights Reserved.
+#include "View/STextureGraphInsightDeviceView.h"
+
+#include "TextureGraphInsight.h"
+#include "Model/TextureGraphInsightSession.h"
+
+#include "Device/DeviceManager.h"
+#include "Data/Blobber.h"
+#include "TextureGraphEngine.h"
+#include <Widgets/Views/ITableRow.h>
+#include <Widgets/Views/STableViewBase.h>
+#include <Widgets/Text/STextBlock.h>
+
+class STextureGraphInsightDeviceListViewRow : public SMultiColumnTableRow<STextureGraphInsightDeviceListView::FItem>
+{
+public:
+	using FItem = STextureGraphInsightDeviceListView::FItem;
+	using FItemW = STextureGraphInsightDeviceListView::FItemW;
+
+	SLATE_BEGIN_ARGS(STextureGraphInsightDeviceListViewRow) {}
+	SLATE_ARGUMENT(FItem, Item)
+	SLATE_END_ARGS()
+
+public:
+
+	enum Column {
+		Main = 0,
+		Count,
+		Hash,
+		Hash1,
+		Hash2,
+		MemSize,
+		Pixels,
+
+		NUM_COLUMNS,
+	};
+	using ColumnInfo = std::pair<FName, float>;
+	static ColumnInfo s_columnNames[NUM_COLUMNS];
+	static Column NameToColumn(const FName& name) {
+		for (int i = 0; i < NUM_COLUMNS; ++i)
+		{
+			if (name == s_columnNames[i].first)
+				return Column(i);
+		}
+		return NUM_COLUMNS;
+	}
+
+	void Construct(const FArguments& InArgs, const TSharedRef<STableViewBase>& InOwnerTableView)
+	{
+		_item = InArgs._Item;
+		_recordID = InArgs._Item->_recordID;
+		SMultiColumnTableRow<FItem >::Construct(FSuperRowType::FArguments(), InOwnerTableView);
+
+		Refresh();
+	}
+
+	virtual TSharedRef<SWidget> GenerateWidgetForColumn(const FName& columnName) override
+	{
+			Column column = NameToColumn(columnName);
+
+			if (column == Main) {
+				// Rows in a TreeView need an expander button and some indentation
+				return SNew(SHorizontalBox)
+					+ SHorizontalBox::Slot()
+					.AutoWidth()
+					.HAlign(HAlign_Right)
+					.VAlign(VAlign_Fill)
+					[
+						SNew(SExpanderArrow, SharedThis(this))
+						.StyleSet(ExpanderStyleSet)
+					]
+					+ SHorizontalBox::Slot()
+					.FillWidth(1)
+					.Padding(FMargin(4.0f, 0.0f))
+					.VAlign(VAlign_Center)
+					[
+						SAssignNew(_textBoxes[column], STextBlock)
+						.Text(GetTextForColumn(column))
+					];
+			}
+			else
+			{
+				return SNew(SBox)
+					.Padding(FMargin(4.0f, 0.0f))
+					.VAlign(VAlign_Center)
+					[
+						SAssignNew(_textBoxes[column], STextBlock)
+						.Text(GetTextForColumn(column))
+						.OnDoubleClicked(this, &STextureGraphInsightDeviceListViewRow::OnDoubleClickedResultColumn)
+					];
+			}
+
+		// default to null widget if property cannot be found
+		return SNullWidget::NullWidget;
+	}
+
+	FText GetTextForColumn(Column column) const
+	{
+		bool isDevice = (_recordID.IsDevice());
+
+		const auto& br = TextureGraphInsight::Instance()->GetSession()->GetRecord().GetBuffer(_recordID);
+		auto memSize = br.RawBufferMemSize;
+
+		FString s;
+		switch (column)
+		{
+		case Main:
+			if (isDevice)
+			{
+				auto item = _item.Pin();
+				if (item->_metaItem == 1)
+					s = FString(TEXT("erased"));
+				else
+					s = (Device::DeviceType_Names((DeviceType) _recordID.Buffer_DeviceType()));
+			}
+			else
+				s = br.Descriptor.Name;
+
+			//	s = FString::FromInt(_recordID.Buffer());
+			break;
+		case Count:
+			if (isDevice)
+			{
+				auto item = _item.Pin();
+				s = FString::FromInt(item->_children.Num());
+			}
+			else
+			{
+				s = FString::FromInt(br.Rank);
+			}
+			break;
+		case Hash:
+			if (isDevice)
+			{
+			}
+			else
+			{
+				auto& h = br.ID; // the key used to find this buffer (ptr)
+				s = HashToFString(h);
+			}
+			break;
+		case Hash1:
+			if (isDevice)
+			{
+			}
+			else
+			{
+				auto& h = br.HashValue;
+				s = HashToFString(h);
+			}
+			break;
+		case Hash2:
+			if (isDevice)
+			{
+			}
+			else
+			{
+				auto& h = br.PrevHashValue;
+				s = HashToFString(h);
+			}
+			break;
+		case MemSize:
+			if (isDevice)
+			{
+			}
+			else
+			{
+				s = FString::FromInt(memSize / 1024) + " kb";
+			}
+			break;
+		case Pixels:
+			if (isDevice)
+			{
+			}
+			else
+			{
+				if (br.Descriptor.Width || br.Descriptor.Height)
+				{
+					s = FString::FromInt(br.Descriptor.Width) + " x " + FString::FromInt(br.Descriptor.Height);
+				}
+			}
+			break;
+		}
+
+		return FText::FromString(s);
+	}
+
+	void Refresh()
+	{
+		if (_recordID.IsDevice())
+		{
+			if (!TextureGraphEngine::GetInstance())
+			{
+				SetColorAndOpacity(FLinearColor(0.3, 0.3, 0.3));
+			}
+			else
+			{
+				auto item = _item.Pin();
+				auto device = TextureGraphEngine::GetInstance()->GetDeviceManager()->GetDevice(_recordID.Buffer_DeviceType());
+				if (!device)
+				{
+					SetColorAndOpacity(FLinearColor(0.2, 0.2, 0.2));
+				}
+				else if (item->_metaItem == 1)
+				{
+					SetColorAndOpacity(FLinearColor(0.5, 0.5, 0.5));
+					_textBoxes[Count]->SetText(FText::FromString(FString::FromInt(item->_children.Num())));
+				}
+				else
+				{ // This is a true device item, update the data
+					auto count = device->GetNumDeviceBuffersUsed();
+					if (count == 0)
+						SetColorAndOpacity(FLinearColor(0.7, 0.7, 0.7));
+					if (count != item->_children.Num())
+						SetColorAndOpacity(FLinearColor(1.0, 0.5, 0.5));
+					else
+						SetColorAndOpacity(FLinearColor(1.0, 1.0, 1.0));
+
+					_textBoxes[Count]->SetText(FText::FromString(FString::FromInt(item->_children.Num())));
+					_textBoxes[Hash]->SetText(FText::FromString(FString::FromInt(count)));
+				}
+			}
+		}
+		else
+		{
+			const auto& br = TextureGraphInsight::Instance()->GetSession()->GetRecord().GetBuffer(_recordID);
+			if (br.bLeaked)
+				SetColorAndOpacity(FLinearColor(1.0, 0.8, 0.5)); 
+			else if (br.bErased)
+				SetColorAndOpacity(FLinearColor(0.5, 0.5, 0.5));
+			else
+				SetColorAndOpacity(FLinearColor(0.5, 1.0, 1.0));
+
+			_textBoxes[Count]->SetText(GetTextForColumn(Count));
+		}
+	}
+
+	FReply OnDoubleClickedResultColumn(/** The geometry of the widget*/
+		const FGeometry&,
+		/** The Mouse Event that we are processing */
+		const FPointerEvent&)
+	{
+		TextureGraphInsight::Instance()->GetSession()->SendToInspector(_recordID);
+		return FReply::Handled();
+	}
+
+protected:
+	FItemW					_item;
+	RecordID				_recordID;
+	TSharedPtr<STextBlock>	_textBoxes[NUM_COLUMNS];
+};
+
+STextureGraphInsightDeviceListViewRow::ColumnInfo STextureGraphInsightDeviceListViewRow::s_columnNames[] = {
+	{ FName(TEXT("Name")), 0.5},    // main
+	{ FName(TEXT("Rank")), 0.1},
+	{ FName(TEXT("Ptr")), 0.1},
+	{ FName(TEXT("Hash")), 0.1},
+	{ FName(TEXT("HashPrev")), 0.1},
+	{ FName(TEXT("Size")), 0.05},
+	{ FName(TEXT("Pixels")), 0.05}
+};
+
+void STextureGraphInsightDeviceListView::Construct(const FArguments& Args)
+{
+	_inspectOnSimpleClick = Args._inspectOnSimpleClick;
+	_inspectDevices = Args._inspectDevices;
+
+	TSharedPtr<SHeaderRow> headerRow = SNew(SHeaderRow);
+
+	for (int i = 0; i < STextureGraphInsightDeviceListViewRow::NUM_COLUMNS; i++)
+	{
+		headerRow->AddColumn(
+			SHeaderRow::Column(STextureGraphInsightDeviceListViewRow::s_columnNames[i].first)
+			.DefaultLabel(FText::FromString(STextureGraphInsightDeviceListViewRow::s_columnNames[i].first.ToString()))
+			.FillWidth(STextureGraphInsightDeviceListViewRow::s_columnNames[i].second)
+		);
+	}
+
+	ChildSlot
+	[
+		SAssignNew(_tableView, SItemTableView)
+		.ItemHeight(24)
+		.TreeItemsSource(&_rootItems)
+		.OnGenerateRow(this, &STextureGraphInsightDeviceListView::OnGenerateRowForView)
+		.OnGetChildren(this, &STextureGraphInsightDeviceListView::OnGetChildrenForView)
+
+		.OnMouseButtonDoubleClick(this, &STextureGraphInsightDeviceListView::OnDoubleClickItemForView)
+		.OnMouseButtonClick(this, &STextureGraphInsightDeviceListView::OnClickItemForView)
+		.HeaderRow(headerRow)
+	];
+
+	// If a valid recordID was specified, let's use it as the root device type,
+	// meaning we only care about that device type
+	// the value of the buffer index doesn't matter
+	if (Args._recordID.IsBuffer())
+		_rootDevice = Args._recordID;
+
+	// Allocate the root items
+	for (int i = 0; i < (int)DeviceType::Count; ++i)
+	{
+		auto deviceType = (DeviceType)i;
+		auto rid = RecordID::fromBuffer(deviceType, INVALID_INDEX);
+		FItem item = MakeShareable(new FItemData(rid));
+		_rootItems.Add(item);
+	}
+
+	// install the observer notifications:
+	auto sr = StaticCastSharedRef<STextureGraphInsightDeviceListView>(this->AsShared());
+	TextureGraphInsight::Instance()->GetSession()->OnDeviceBufferAdded().AddSP(sr, &STextureGraphInsightDeviceListView::OnDeviceBufferAdded);
+	TextureGraphInsight::Instance()->GetSession()->OnDeviceBufferRemoved().AddSP(sr, &STextureGraphInsightDeviceListView::OnDeviceBufferRemoved);
+	TextureGraphInsight::Instance()->GetSession()->OnEngineReset().AddSP(sr, &STextureGraphInsightDeviceListView::OnEngineReset);
+}
+
+void STextureGraphInsightDeviceListView::Reset()
+{
+	for (int i = 0; i < (int)DeviceType::Count; ++i)
+	{
+		auto deviceType = (DeviceType)i;
+		_rootItems[i]->_children.Empty();
+		_rootItems[i]->_removedRoot.Reset();
+		if (_rootItems[i]->_widget) _rootItems[i]->_widget->Refresh();
+
+		if (TextureGraphInsight::Instance()->GetSession())
+		{
+			auto deviceBufferIDs = TextureGraphInsight::Instance()->GetSession()->GetRecord().FetchActiveDeviceBufferIDs((DeviceType)i);
+			for (auto rid : deviceBufferIDs)
+			{
+				_rootItems[i]->_children.Add(MakeShareable(new FItemData(rid)));
+			}
+		}
+	}
+
+	if (_rootDevice.IsDevice())
+	{
+		_tableView->SetItemExpansion(_rootItems[_rootDevice.Buffer_DeviceType()], true);
+	}
+
+	_tableView->RequestListRefresh();
+
+	RefreshRootItems();
+}
+
+TSharedRef<ITableRow> STextureGraphInsightDeviceListView::OnGenerateRowForView(FItem item, const TSharedRef<STableViewBase>& OwnerTable)
+{
+	return SAssignNew(item->_widget, STextureGraphInsightDeviceListViewRow, OwnerTable).Item(item);
+}
+
+void STextureGraphInsightDeviceListView::OnClickItemForView(FItem item)
+{
+	if (_inspectOnSimpleClick)
+		OnInspect(item);
+}
+void STextureGraphInsightDeviceListView::OnDoubleClickItemForView(FItem item)
+{
+	OnInspect(item);
+}
+
+void STextureGraphInsightDeviceListView::OnInspect(FItem item)
+{
+	if (item->_recordID.IsDevice())
+	{
+		if (item->_metaItem > 0)
+			return;
+		if (!_inspectDevices)
+			return;
+	}
+
+	auto record = item->_recordID;
+	TextureGraphInsight::Instance()->GetSession()->SendToInspector(record);
+}
+
+void STextureGraphInsightDeviceListView::OnDeviceBufferAdded(const RecordIDArray& rids)
+{
+	for (const auto& rid : rids)
+	{
+		if (!rid.IsBuffer() || (rid.Buffer() == INVALID_INDEX))
+			continue;
+
+		auto i = rid.Buffer_DeviceType();
+		FItem item = MakeShareable(new FItemData(rid));
+		_rootItems[i]->_children.Add(item);
+	}
+
+
+	RefreshRootItems();
+
+	_tableView->RequestListRefresh();
+}
+
+void STextureGraphInsightDeviceListView::OnDeviceBufferRemoved(const RecordIDArray& rids)
+{
+	for (const auto& rid : rids)
+	{
+		if (!rid.IsBuffer() || (rid.Buffer() == INVALID_INDEX))
+			continue;
+
+		// First let's find the position of the record in the existing model:
+		auto& root = _rootItems[rid.Buffer_DeviceType()];
+		auto& children = root->_children;
+
+		auto index = children.IndexOfByKey(rid);
+		if (index == INDEX_NONE)
+			continue; /// This should not happen
+
+		// Then let's move the removed item to the sub group in removedRoot
+		// First make sure the removedRoot exists
+		if (!root->_removedRoot)
+		{
+			root->_removedRoot = MakeShareable(new FItemData(root->_recordID, 1));
+			//_tableView->SetItemExpansion(root->_removedRoot, true);
+		}
+
+		// Then move the item there
+		root->_removedRoot->_children.Add(children[index]);
+		children.RemoveAt(index);
+	}
+
+
+	RefreshRootItems();
+
+	_tableView->RequestListRefresh();
+
+}
+
+void STextureGraphInsightDeviceListView::OnEngineReset(int id)
+{
+	Reset();
+}
+
+bool operator <(const STextureGraphInsightDeviceListView::FItem& a, const  STextureGraphInsightDeviceListView::FItem& b)
+{
+	return ((uint32) a->_rank) < ((uint32) b->_rank);
+}
+
+void STextureGraphInsightDeviceListView::RefreshRootItems()
+{
+	for (auto& i: _rootItems)
+	{
+
+		for (auto& c : i->_children)
+		{
+			const auto& br = TextureGraphInsight::Instance()->GetSession()->GetRecord().GetBuffer(c->_recordID);
+			c->_rank = br.Rank;
+		}
+		i->_children.Sort();
+			
+		if (i->_widget)
+		{
+			i->_widget->Refresh();
+
+			for (auto& c : i->_children)
+			{
+				if (c->_widget)
+				{
+					c->_widget->Refresh();
+				}
+			}
+		}
+		if (i->_removedRoot && i->_removedRoot->_widget)
+		{
+			i->_removedRoot->_widget->Refresh();
+		}
+	}
+}
+
+void STextureGraphInsightDeviceView::Construct(const FArguments& Args)
+{
+	_deviceType = (DeviceType) Args._deviceType;
+	
+	ChildSlot
+	[
+		SNew(SVerticalBox)
+		+ SVerticalBox::Slot()
+		.AutoHeight()
+		[
+			SNew(SHorizontalBox)
+			+ SHorizontalBox::Slot()
+			[
+				SNew(STextBlock)
+				.Text(FText::FromString(Device::DeviceType_Names(_deviceType)))
+			]
+			+ SHorizontalBox::Slot()
+			[
+				SNew(STextBlock)
+				.Text(FText::FromString("Threads"))
+				.Justification(ETextJustify::Right)
+			]
+			+ SHorizontalBox::Slot()
+			[
+				SAssignNew(_threadCount, STextBlock)
+				.Text(FText::FromString("__"))
+				.AutoWrapText(true)
+			]
+		]
+		+ SVerticalBox::Slot()
+		.AutoHeight()
+		[
+			SNew(SHorizontalBox)
+			+ SHorizontalBox::Slot()
+			[
+				SNew(STextBlock)
+				.Text(FText::FromString("Memory"))
+			]
+			+ SHorizontalBox::Slot()
+			[
+				SAssignNew(_mem_percent, STextBlock)
+				.Text(FText::FromString("___%"))
+				.AutoWrapText(true)
+				.Justification(ETextJustify::Right)
+			]
+			+ SHorizontalBox::Slot()
+			[
+				SAssignNew(_mem_number, STextBlock)
+				.Text(FText::FromString("______ Kb"))
+				.AutoWrapText(true)
+				.Justification(ETextJustify::Right)
+			]
+			+ SHorizontalBox::Slot()
+			[
+				SAssignNew(_mem_budget, STextBlock)
+				.Text(FText::FromString("______ Kb"))
+				.AutoWrapText(true)
+				.Justification(ETextJustify::Right)
+			]
+		]
+		+ SVerticalBox::Slot()
+		.AutoHeight()
+		[
+			SNew(SHorizontalBox)
+			+ SHorizontalBox::Slot()
+			[
+				SNew(STextBlock)
+				.Text(FText::FromString("Buffers"))
+			]
+			+ SHorizontalBox::Slot()
+			[
+				SAssignNew(_buffer_percent, STextBlock)
+				.Text(FText::FromString("___%"))
+			.AutoWrapText(true)
+			.Justification(ETextJustify::Right)
+			]
+			+ SHorizontalBox::Slot()
+			[
+				SAssignNew(_buffer_number, STextBlock)
+				.Text(FText::FromString("______"))
+			.AutoWrapText(true)
+			.Justification(ETextJustify::Right)
+			]
+			+ SHorizontalBox::Slot()
+			[
+				SAssignNew(_buffer_budget, STextBlock)
+				.Text(FText::FromString("______"))
+			.AutoWrapText(true)
+			.Justification(ETextJustify::Right)
+			]
+		]
+	];
+		
+}
+
+void STextureGraphInsightDeviceView::Refresh()
+{
+	auto device = TextureGraphEngine::GetInstance()->GetDeviceManager()->GetDevice(_deviceType);
+	_threadCount->SetText(FText::FromString(FString::FromInt(device->GetMaxThreads())));
+
+	auto memMax = device->GetMaxMemory();
+	auto memUsed = device->GetMemUsed();
+	auto memPercent = 0;
+	if (memMax > 0)
+		memPercent = (memUsed * 100) / memMax;
+
+	auto bufferMax = device->GetNumDeviceBuffersMax();
+	auto bufferUsed = device->GetNumDeviceBuffersUsed();
+	auto bufferPercent = 0;
+	if (bufferMax > 0)
+		bufferPercent = (bufferUsed * 100) / bufferMax;
+
+
+	_mem_percent->SetText(FText::FromString(FString::FromInt(memPercent)));
+	_mem_number->SetText(FText::FromString(FString::FromInt(memUsed / (1024 * 1024)) + " Mb"));
+	_mem_budget->SetText(FText::FromString(FString::FromInt(memMax / (1024 * 1024)) + " Mb"));
+
+	_buffer_percent->SetText(FText::FromString(FString::FromInt(bufferPercent)));
+	_buffer_number->SetText(FText::FromString(FString::FromInt(bufferUsed)));
+	_buffer_budget->SetText(FText::FromString(FString::FromInt(bufferMax)));
+
+	
+}
+
+
+
+void STextureGraphInsightDeviceManagerView::Construct(const FArguments& Args)
+{
+	// install the observer notifications
+	auto sr = StaticCastSharedRef<STextureGraphInsightDeviceManagerView>(this->AsShared());
+	TextureGraphInsight::Instance()->GetSession()->OnUpdateIdle().AddSP(sr, &STextureGraphInsightDeviceManagerView::OnUpdate);
+	TextureGraphInsight::Instance()->GetSession()->OnEngineReset().AddSP(sr, &STextureGraphInsightDeviceManagerView::OnEngineReset);
+}
+
+void STextureGraphInsightDeviceManagerView::Reset()
+{
+	_deviceViews.clear();
+
+	TSharedPtr<SVerticalBox> vbox;
+	ChildSlot
+	[
+		SAssignNew(vbox, SVerticalBox)
+	]; 
+
+	if (TextureGraphEngine::GetInstance())
+	{
+		for (int i = 0; i < (int)DeviceType::Count; ++i)
+		{
+			//if (i == (int)DeviceType::Null)
+			//	continue;
+
+			auto deviceType = (DeviceType)i;
+			auto device = TextureGraphEngine::GetInstance()->GetDeviceManager()->GetDevice(deviceType);
+			if (device)
+			{
+				TSharedPtr< STextureGraphInsightDeviceView> deviceView;
+				vbox->AddSlot()
+					[
+						SAssignNew(deviceView, STextureGraphInsightDeviceView)
+						.deviceType(deviceType)
+					];
+
+				_deviceViews.push_back(deviceView);
+			}
+		}
+	}
+}
+
+void STextureGraphInsightDeviceManagerView::OnUpdate()
+{
+	for (auto& d : _deviceViews)
+	{
+		d->Refresh();
+	}
+}
+
+void STextureGraphInsightDeviceManagerView::OnEngineReset(int32 id)
+{
+	Reset();
+}
\ No newline at end of file
diff --git a/Engine/Plugins/Experimental/TextureGraph/Source/TextureGraphInsight/Private/View/STextureGraphInsightInspectorView.cpp b/Engine/Plugins/Experimental/TextureGraph/Source/TextureGraphInsight/Private/View/STextureGraphInsightInspectorView.cpp
new file mode 100644
--- /dev/null
+++ b/Engine/Plugins/Experimental/TextureGraph/Source/TextureGraphInsight/Private/View/STextureGraphInsightInspectorView.cpp
@@ -0,0 +1,293 @@
+// Copyright Epic Games, Inc. All Rights Reserved.
+#include "View/STextureGraphInsightInspectorView.h"
+
+#include "2D/Tex.h"
+#include "TextureGraphInsight.h"
+#include "Model/TextureGraphInsightSession.h"
+#include "View/STextureGraphInsightBlobView.h"
+#include "View/STextureGraphInsightDeviceBufferView.h"
+#include "View/STextureGraphInsightDeviceView.h"
+#include "View/STextureGraphInsightRecordTrailView.h"
+#include "View/STextureGraphInsightSessionView.h"
+#include "Widgets/SNullWidget.h"
+
+#include "Device/FX/DeviceBuffer_FX.h"
+#include <Widgets/Layout/SHeader.h>
+
+
+void STextureGraphInsightBatchInspectorView::Construct(const FArguments& Args)
+{
+	_rootItems.Empty(0);
+
+	InspectBatch(Args._recordID);
+};
+
+TSharedRef<ITableRow> STextureGraphInsightBatchInspectorView::OnGenerateTileForView(FItem item, const TSharedRef<STableViewBase>& OwnerTable)
+{
+	const auto& br = TextureGraphInsight::Instance()->GetSession()->GetRecord().GetBatch(item->_batchID);
+
+	RecordID targetRID = br.ResultBlobs[item->_targetIndex];
+	FString targetTypeName = TextureHelper::TextureTypeToString((TextureType) item->_targetIndex);
+	FString label = "Target " + targetTypeName + " : " + FString::FromInt(targetRID.Blob());
+	return SNew(STableRow<FItem>, OwnerTable)
+		.Padding(4)
+		.Content()
+		[
+			SNew(SVerticalBox)
+			+ SVerticalBox::Slot()
+			.AutoHeight()
+			[
+				SNew(SHeader)
+				[
+					SNew(STextBlock)
+					.Text(FText::FromString(label))
+				]
+			]
+			+ SVerticalBox::Slot()
+			[
+				SNew(STextureGraphInsightBlobView)
+				.recordID(targetRID)
+			]
+		];
+}
+
+void STextureGraphInsightBatchInspectorView::InspectBatch(RecordID batchID)
+{
+	if (!batchID.IsBatch())
+	{
+		batchID = RecordID();
+	}
+
+	_rootItems.Empty();
+
+	const auto& br = TextureGraphInsight::Instance()->GetSession()->GetRecord().GetBatch(batchID);
+
+	int32 resultIdx = 0;
+	for (auto inputBlobId : br.ResultBlobs)
+	{
+		if (inputBlobId.IsValid())
+			_rootItems.Add(MakeShareable(new FItemData(batchID, resultIdx)));
+		resultIdx++;
+	}
+
+	ChildSlot
+		[
+			SNew(SSplitter)
+			.Orientation(Orient_Vertical)
+		+ SSplitter::Slot()
+		.Value(3)
+		[
+
+			SAssignNew(_view, SItemTileView)
+			.ListItemsSource(&_rootItems)
+		.ItemWidth(256)
+		.ItemHeight(256)
+		.OnGenerateTile(this, &STextureGraphInsightBatchInspectorView::OnGenerateTileForView)
+		]
+		];
+
+
+	_view->RequestListRefresh();
+}
+
+
+void STextureGraphInsightJobInspectorView::Construct(const FArguments& Args)
+{
+	_rootItems.Empty(0);
+
+	InspectJob(Args._recordID);
+};
+
+TSharedRef<ITableRow> STextureGraphInsightJobInspectorView::OnGenerateTileForView(FItem item, const TSharedRef<STableViewBase>& OwnerTable)
+{
+	const auto& br = TextureGraphInsight::Instance()->GetSession()->GetRecord().GetBatch(item->_jobID);
+	const auto& jr = br.GetJob(item->_jobID);
+
+	RecordID inputRID = jr.InputBlobIds[item->_inputIndex];
+
+	FString label = jr.InputArgNames[item->_inputIndex] + " : " + FString::FromInt(inputRID.Blob());
+	return SNew(STableRow<FItem>, OwnerTable)
+		.Padding(4)
+		.Content()
+		[
+			SNew(SVerticalBox)
+			+ SVerticalBox::Slot()
+			.AutoHeight()
+			[
+				SNew(SHeader)
+				[
+					SNew(STextBlock)
+					.Text(FText::FromString(label))
+				]
+			]
+			+ SVerticalBox::Slot()
+			[
+				SNew(STextureGraphInsightBlobView)
+				.recordID(inputRID)
+				.tilesMask(item->_tilesMask)
+			]
+		];
+}
+
+void STextureGraphInsightJobInspectorView::InspectJob(RecordID jobID)
+{
+	if (!jobID.IsJob())
+	{
+		jobID = RecordID();
+	}
+
+	_rootItems.Empty();
+
+	const auto& br = TextureGraphInsight::Instance()->GetSession()->GetRecord().GetBatch(jobID);
+	const auto& jr = br.GetJob(jobID);
+
+	int32 inputIdx = 0;
+	for (auto inputBlobId : jr.InputBlobIds)
+	{
+		if (inputBlobId.IsValid())
+			_rootItems.Add(MakeShareable(new FItemData(jobID, inputIdx)));
+		inputIdx++;
+	}
+
+	ChildSlot
+		[
+		SNew(SSplitter)
+		.Orientation(Orient_Vertical)
+		+ SSplitter::Slot()
+		.Value(3)
+		[
+			SAssignNew(_outBlobView, STextureGraphInsightBlobView)
+			.recordID(jr.ResultBlobId)
+			.tilesMask(jr.Tiles)
+		]
+		+ SSplitter::Slot()
+		[
+			SAssignNew(_view, SItemTileView)
+			.ListItemsSource(&_rootItems)
+			.ItemWidth(256)
+			.ItemHeight(256)
+			.OnGenerateTile(this, &STextureGraphInsightJobInspectorView::OnGenerateTileForView)
+		]
+		];
+
+
+	_view->RequestListRefresh();
+}
+
+
+void STextureGraphInsightBlobInspectorView::Construct(const FArguments& Args)
+{
+	InspectBlob(Args._recordID);
+};
+
+
+void STextureGraphInsightBlobInspectorView::InspectBlob(RecordID blobID)
+{
+	ChildSlot
+		[
+			SAssignNew(_blobView, STextureGraphInsightBlobView)
+				.recordID(blobID)
+		];
+}
+
+void STextureGraphInsightDeviceBufferInspectorView::Construct(const FArguments& Args)
+{
+	InspectBuffer(Args._recordID);
+};
+
+
+void STextureGraphInsightDeviceBufferInspectorView::InspectBuffer(RecordID rid)
+{
+	if (rid.IsDevice())
+	{
+		ChildSlot
+		[
+			SNew(STextureGraphInsightDeviceView)
+			.deviceType((DeviceType) rid.Buffer_DeviceType())
+		];
+	}
+	else
+	{
+		ChildSlot
+		[
+			SNew(STextureGraphInsightDeviceBufferView)
+			.recordID(rid)
+			.withDescription(true)
+		];
+	}
+}
+
+
+void STextureGraphInsightInspectorView::Construct(const FArguments& Args)
+{
+	ChildSlot
+		[
+			SNew(SVerticalBox)
+			+ SVerticalBox::Slot()
+			.AutoHeight()
+			[
+				SAssignNew(_recordTrail, STextureGraphInsightRecordTrailView)
+			]
+			+ SVerticalBox::Slot()
+			[
+				SAssignNew(_stack, SOverlay)
+			]
+		];
+
+	// install the observer notifications
+	auto sr = StaticCastSharedRef<STextureGraphInsightInspectorView>(this->AsShared());
+	TextureGraphInsight::Instance()->GetSession()->OnBatchInspected().AddSP(sr, &STextureGraphInsightInspectorView::OnInspected);
+	TextureGraphInsight::Instance()->GetSession()->OnJobInspected().AddSP(sr, &STextureGraphInsightInspectorView::OnInspected);
+	TextureGraphInsight::Instance()->GetSession()->OnBlobInspected().AddSP(sr, &STextureGraphInsightInspectorView::OnInspected);
+	TextureGraphInsight::Instance()->GetSession()->OnBufferInspected().AddSP(sr, &STextureGraphInsightInspectorView::OnInspected);
+
+	TextureGraphInsight::Instance()->GetSession()->OnEngineReset().AddSP(sr, &STextureGraphInsightInspectorView::OnEngineReset);
+};
+
+
+void STextureGraphInsightInspectorView::OnInspected(RecordID rid)
+{
+	_recordTrail->Refresh(rid);
+	_stack->ClearChildren();
+
+	if (rid.IsBuffer())
+	{
+		_stack->AddSlot()
+			[
+				SNew(STextureGraphInsightDeviceBufferInspectorView)
+				.recordID(rid)
+			];
+	}
+	else if (rid.IsBlob())
+	{
+		_stack->AddSlot()
+			[
+				SNew(STextureGraphInsightBlobInspectorView)
+				.recordID(rid)
+			];
+	}
+	else if (rid.IsJob())
+	{
+		_stack->AddSlot()
+			[
+				SNew(STextureGraphInsightJobInspectorView)
+				.recordID(rid)
+			];
+	}
+	else if (rid.IsBatch())
+	{
+		_stack->AddSlot()
+			[
+				SNew(STextureGraphInsightBatchInspectorView)
+				.recordID(rid)
+			];
+	}
+}
+
+
+void STextureGraphInsightInspectorView::OnEngineReset(int32 id)
+{
+	// clear the inspector _trail and _stack when engine is destroyed
+	_recordTrail->Refresh(RecordID());
+	_stack->ClearChildren();
+}
diff --git a/Engine/Plugins/Experimental/TextureGraph/Source/TextureGraphInsight/Private/View/STextureGraphInsightMixView.cpp b/Engine/Plugins/Experimental/TextureGraph/Source/TextureGraphInsight/Private/View/STextureGraphInsightMixView.cpp
new file mode 100644
--- /dev/null
+++ b/Engine/Plugins/Experimental/TextureGraph/Source/TextureGraphInsight/Private/View/STextureGraphInsightMixView.cpp
@@ -0,0 +1,313 @@
+// Copyright Epic Games, Inc. All Rights Reserved.
+#include "View/STextureGraphInsightMixView.h"
+
+#include "TextureGraphInsight.h"
+#include "Model/TextureGraphInsightSession.h"
+
+#include "Device/DeviceManager.h"
+#include "Data/Blobber.h"
+#include "TextureGraphEngine.h"
+
+
+class STextureGraphInsightMixListViewRow : public SMultiColumnTableRow<STextureGraphInsightMixListView::FItem>
+{
+public:
+	using FItem = STextureGraphInsightMixListView::FItem;
+
+	SLATE_BEGIN_ARGS(STextureGraphInsightMixListViewRow) {}
+	SLATE_ARGUMENT(FItem, Item)
+		SLATE_END_ARGS()
+
+public:
+
+	enum Column {
+		Main = 0,
+		Name,
+		Parent,
+		Instances,
+
+		NUM_COLUMNS,
+	};
+	using ColumnInfo = std::pair<FName, float>;
+	static ColumnInfo s_columnNames[NUM_COLUMNS];
+	static Column NameToColumn(const FName& name) {
+		for (int i = 0; i < NUM_COLUMNS; ++i)
+		{
+			if (name == s_columnNames[i].first)
+				return Column(i);
+		}
+		return NUM_COLUMNS;
+	}
+
+	virtual TSharedRef<SWidget> GenerateWidgetForColumn(const FName& columnName) override
+	{
+		if (_recordID.IsValid()) {
+
+			Column column = NameToColumn(columnName);
+
+			if (column == Main) {
+				// Rows in a TreeView need an expander button and some indentation
+				return SNew(SHorizontalBox)
+					+ SHorizontalBox::Slot()
+					.AutoWidth()
+					.HAlign(HAlign_Right)
+					.VAlign(VAlign_Fill)
+					[
+						SNew(SExpanderArrow, SharedThis(this))
+						.StyleSet(ExpanderStyleSet)
+					]
+				+ SHorizontalBox::Slot()
+					.FillWidth(1)
+					.Padding(FMargin(4.0f, 0.0f))
+					.VAlign(VAlign_Center)
+					[
+						SAssignNew(_textBoxes[column], STextBlock)
+						.Text(GetTextForColumn(column))
+					];
+			} else
+			{
+				return SNew(SBox)
+					.Padding(FMargin(4.0f, 0.0f))
+					.VAlign(VAlign_Center)
+					[
+						SAssignNew(_textBoxes[column], STextBlock)
+						.Text(GetTextForColumn(column))
+					.OnDoubleClicked(this, &STextureGraphInsightMixListViewRow::OnDoubleClickedResultColumn)
+					];
+			}
+		}
+		// default to null widget if property cannot be found
+		return SNullWidget::NullWidget;
+	}
+
+	FText GetTextForColumn(Column column) const
+	{
+		FString s;
+		if (_recordID.IsMix())
+		{
+			const auto& mr = TextureGraphInsight::Instance()->GetSession()->GetRecord().GetMix(_recordID);
+
+			switch (column)
+			{
+			case Main:
+				s = "Mix " + FString::FromInt(_recordID.Mix());
+				break;
+
+			case Name:
+				s = mr.Name;
+				break;
+
+			case Parent:
+				if (mr.IsInstanceMix())
+				{
+					s = "Mix " + FString::FromInt(mr.ParentMixID.Mix());
+				}
+				break;
+
+			case Instances:
+				if (!mr.IsInstanceMix())
+				{
+					s = FString::FromInt(mr.InstanceMixIDs.size());
+				}
+				break;
+			}
+		}
+		else if (_recordID.IsBatch())
+		{
+			const auto& br = TextureGraphInsight::Instance()->GetSession()->GetRecord().GetBatch(_recordID);
+			switch (column)
+			{
+			case Main:
+				s = "Batch " + FString::FromInt(_recordID.Batch());
+				break;
+
+			case Name:
+				s = br.Action;
+				break;
+
+			case Parent:
+				break;
+			case Instances:
+				break;
+			}
+		}
+		return FText::FromString(s);
+	}
+
+
+	void Construct(const FArguments& InArgs, const TSharedRef<STableViewBase>& InOwnerTableView)
+	{
+		_recordID = InArgs._Item->_recordID;
+		SMultiColumnTableRow<FItem >::Construct(FSuperRowType::FArguments(), InOwnerTableView);
+
+		Refresh();
+	}
+
+	void Refresh()
+	{
+		const auto& mr = TextureGraphInsight::Instance()->GetSession()->GetRecord().GetMix(_recordID);	
+		_textBoxes[Name]->SetText(GetTextForColumn(Name));
+	}
+
+	FReply OnDoubleClickedResultColumn(/** The geometry of the widget*/
+		const FGeometry&,
+		/** The Mouse Event that we are processing */
+		const FPointerEvent&)
+	{
+		TextureGraphInsight::Instance()->GetSession()->SendToInspector(_recordID);
+		return FReply::Handled();
+	}
+
+protected:
+	RecordID _recordID;
+	TSharedPtr<STextBlock> _textBoxes[NUM_COLUMNS];
+};
+
+STextureGraphInsightMixListViewRow::ColumnInfo STextureGraphInsightMixListViewRow::s_columnNames[] = {
+	{ FName(TEXT("Id")), 0.1},    // main
+	{ FName(TEXT("Name")), 0.5},
+	{ FName(TEXT("Parent")), 0.2},
+	{ FName(TEXT("Instances")), 0.2},
+};
+
+void STextureGraphInsightMixListView::Construct(const FArguments& Args)
+{
+
+	TSharedPtr<SHeaderRow> headerRow = SNew(SHeaderRow);
+
+	for (int i = 0; i < STextureGraphInsightMixListViewRow::NUM_COLUMNS; i++)
+	{
+		headerRow->AddColumn(
+			SHeaderRow::Column(STextureGraphInsightMixListViewRow::s_columnNames[i].first)
+			.DefaultLabel(FText::FromString(STextureGraphInsightMixListViewRow::s_columnNames[i].first.ToString()))
+			.FillWidth(STextureGraphInsightMixListViewRow::s_columnNames[i].second)
+		);
+	}
+
+	ChildSlot
+		[
+			SAssignNew(_tableView, SItemTableView)
+			.ItemHeight(24)
+		//.ListItemsSource(&_rootItems)
+		.TreeItemsSource(&_rootItems)
+		.OnGenerateRow(this, &STextureGraphInsightMixListView::OnGenerateRowForView)
+		.OnGetChildren(this, &STextureGraphInsightMixListView::OnGetChildrenForView)
+
+		.OnMouseButtonDoubleClick(this, &STextureGraphInsightMixListView::OnDoubleClickItemForView)
+		.OnMouseButtonClick(this, &STextureGraphInsightMixListView::OnClickItemForView)
+		.HeaderRow(headerRow)
+		];
+
+	// install the observer notifications
+	auto sr = StaticCastSharedRef<STextureGraphInsightMixListView>(this->AsShared());
+	TextureGraphInsight::Instance()->GetSession()->OnMixAdded().AddSP(sr, &STextureGraphInsightMixListView::OnMixNew);
+	TextureGraphInsight::Instance()->GetSession()->OnMixUpdated().AddSP(sr, &STextureGraphInsightMixListView::OnMixUpdate);
+	TextureGraphInsight::Instance()->GetSession()->OnEngineReset().AddSP(sr, &STextureGraphInsightMixListView::OnEngineReset);
+
+}
+
+TSharedRef<ITableRow> STextureGraphInsightMixListView::OnGenerateRowForView(FItem item, const TSharedRef<STableViewBase>& OwnerTable)
+{
+	return SAssignNew(item->_widget, STextureGraphInsightMixListViewRow, OwnerTable).Item(item);
+}
+
+void STextureGraphInsightMixListView::OnClickItemForView(FItem item)
+{
+	// Item clicked, let's send it to the inspector
+	auto record = item->_recordID;
+	TextureGraphInsight::Instance()->GetSession()->SendToInspector(record);
+}
+void STextureGraphInsightMixListView::OnDoubleClickItemForView(FItem item)
+{
+	// Item double clicked, let's send it to the inspector
+	auto record = item->_recordID;
+	TextureGraphInsight::Instance()->GetSession()->SendToInspector(record);
+}
+
+void STextureGraphInsightMixListView::OnMixNew(RecordID rid)
+{
+	const auto& mr = TextureGraphInsight::Instance()->GetSession()->GetRecord().GetMix(rid);
+
+	FItem item = MakeShareable(new FItemData(rid));
+
+	_rootItems.Add(item);
+
+	UpdateItemData(*_rootItems.Last());
+
+	_tableView->RequestListRefresh();
+}
+
+bool operator ==(const STextureGraphInsightMixListView::FItem& a, const  RecordID& i)
+{
+	return a->_recordID.id == i.id;
+}
+
+void STextureGraphInsightMixListView::RefreshRootItems()
+{
+	for (auto& i : _rootItems)
+	{
+		if (i->_widget)
+		{
+			i->_widget->Refresh();
+
+			for (auto& c : i->_children)
+			{
+				if (c->_widget)
+				{
+					c->_widget->Refresh();
+				}
+			}
+		}
+	}
+}
+
+void STextureGraphInsightMixListView::OnMixUpdate(RecordID rid)
+{
+	auto itemPtr = _rootItems.FindByKey(rid);
+	if (itemPtr)
+	{
+		UpdateItemData(*(* itemPtr));
+	}
+
+	_tableView->RequestListRefresh();
+}
+
+void STextureGraphInsightMixListView::UpdateItemData(FItemData& item)
+{
+	if (item._recordID.IsMix())
+	{
+		const auto& mr = TextureGraphInsight::Instance()->GetSession()->GetRecord().GetMix(item._recordID);
+
+		for (int i = 0; i < mr.Batches.size(); ++i)
+		{
+			auto bid = mr.Batches[i];
+			if (i < item._children.Num() && item._children[i]->_recordID.id == bid.id)
+			{
+			}
+			else
+			{
+				if (i < item._children.Num())
+				{
+					// erase the existing array starting at i on the first not equal item
+					item._children.SetNum(i);
+				}
+
+				auto newItem = MakeShareable(new FItemData(bid));
+				item._children.Push(newItem);
+				UpdateItemData(*item._children.Last());
+			}
+		}
+
+		if (item._widget)
+		{
+			item._widget->Refresh();
+		}
+	}
+}
+
+void STextureGraphInsightMixListView::OnEngineReset(int id)
+{
+	_rootItems.Empty();
+	_tableView->RequestListRefresh();
+}
+
diff --git a/Engine/Plugins/Experimental/TextureGraph/Source/TextureGraphInsight/Private/View/STextureGraphInsightRecordTrailView.cpp b/Engine/Plugins/Experimental/TextureGraph/Source/TextureGraphInsight/Private/View/STextureGraphInsightRecordTrailView.cpp
new file mode 100644
--- /dev/null
+++ b/Engine/Plugins/Experimental/TextureGraph/Source/TextureGraphInsight/Private/View/STextureGraphInsightRecordTrailView.cpp
@@ -0,0 +1,102 @@
+// Copyright Epic Games, Inc. All Rights Reserved.
+#include "View/STextureGraphInsightRecordTrailView.h"
+
+#include "2D/Tex.h"
+#include "TextureGraphInsight.h"
+#include "Model/TextureGraphInsightSession.h"
+#include "Widgets/SNullWidget.h"
+
+void STextureGraphInsightRecordTrailView::Construct(const FArguments& Args)
+{
+	ChildSlot
+	[
+		SAssignNew(_trail, SBreadcrumbTrail<RecordID>)
+		.OnCrumbClicked(this, &STextureGraphInsightRecordTrailView::OnCrumbClicked)
+	];
+
+	Refresh(Args._recordID);
+}
+
+
+RecordID STextureGraphInsightRecordTrailView::FindRootRecord(RecordID rid)
+{
+	if (rid.IsBuffer())
+	{
+		const auto& bufferRecord = TextureGraphInsight::Instance()->GetSession()->GetRecord().GetBuffer(rid);
+
+		return TextureGraphInsight::Instance()->GetSession()->GetRecord().FindBlobRecord(bufferRecord.HashValue);
+	}
+	else if (rid.IsBlob())
+	{
+		const auto& blobRecord = TextureGraphInsight::Instance()->GetSession()->GetRecord().GetBlob(rid);
+		return blobRecord.SourceID;
+	}
+	else if (rid.IsJob())
+	{
+		return RecordID::fromBatch(rid.Batch());
+	}
+	else if (rid.IsBatch())
+	{
+		const auto& batchRecord = TextureGraphInsight::Instance()->GetSession()->GetRecord().GetBatch(rid);
+	}
+	else if (rid.IsAction())
+	{
+	}
+
+	return RecordID();
+}
+
+void STextureGraphInsightRecordTrailView::Refresh(RecordID rid)
+{
+	_trail->ClearCrumbs(true);
+
+	//build a trail of rid navigating the records
+	RecordIDArray trail;
+	do
+	{
+		trail.emplace_back(rid);
+		rid = FindRootRecord(rid);
+	} while (rid.IsValid());
+
+	// then create the crumbs as needed
+	while (!trail.empty())
+	{
+		auto id = trail.back();
+		trail.pop_back();
+
+		if (id.IsBuffer())
+		{
+			const auto& bufferRecord = TextureGraphInsight::Instance()->GetSession()->GetRecord().GetBuffer(id);
+			_trail->PushCrumb(FText::FromString("Buffer " + HashToFString(bufferRecord.HashValue)), id);
+		}
+		else if (id.IsBlob())
+		{
+			const auto& blobRecord = TextureGraphInsight::Instance()->GetSession()->GetRecord().GetBlob(id);
+			_trail->PushCrumb(FText::FromString("Blob " + HashToFString(blobRecord.HashValue)), id);
+		}
+		else if (id.IsJob())
+		{
+			const auto& jobRecord = TextureGraphInsight::Instance()->GetSession()->GetRecord().GetJob(id);
+			_trail->PushCrumb(FText::FromString(jobRecord.TransformName), id);
+		}
+		else if (id.IsBatch())
+		{
+			const auto& batchRecord = TextureGraphInsight::Instance()->GetSession()->GetRecord().GetBatch(id);
+			_trail->PushCrumb(FText::FromString("Batch " + FString::FromInt(batchRecord.BatchID)), id);
+		}
+		else if (id.IsAction())
+		{
+			const auto& actionRecord = TextureGraphInsight::Instance()->GetSession()->GetRecord().GetAction(id);
+			_trail->PushCrumb(FText::FromString("Action " + actionRecord.Name), id);
+		}
+		else // Invalid
+		{
+			_trail->PushCrumb(FText::FromString(FString("Null")), RecordID());
+		}
+	}
+}
+
+void STextureGraphInsightRecordTrailView::OnCrumbClicked(const RecordID& rid)
+{
+	TextureGraphInsight::Instance()->GetSession()->SendToInspector(rid);
+}
diff --git a/Engine/Plugins/Experimental/TextureGraph/Source/TextureGraphInsight/Private/View/STextureGraphInsightResourceView.cpp b/Engine/Plugins/Experimental/TextureGraph/Source/TextureGraphInsight/Private/View/STextureGraphInsightResourceView.cpp
new file mode 100644
--- /dev/null
+++ b/Engine/Plugins/Experimental/TextureGraph/Source/TextureGraphInsight/Private/View/STextureGraphInsightResourceView.cpp
@@ -0,0 +1,306 @@
+// Copyright Epic Games, Inc. All Rights Reserved.
+#include "View/STextureGraphInsightResourceView.h"
+
+#include "TextureGraphInsight.h"
+#include "Model/TextureGraphInsightSession.h"
+
+#include "Device/DeviceManager.h"
+#include "Data/Blobber.h"
+#include "TextureGraphEngine.h"
+#include "Widgets/Text/STextBlock.h"
+#include "Widgets/Views/STableRow.h"
+#include "Widgets/SBoxPanel.h"
+
+
+class STextureGraphInsightResourceViewRow : public SMultiColumnTableRow<STextureGraphInsightResourceView::FItem>
+{
+public:
+	using FItem = STextureGraphInsightResourceView::FItem;
+
+	SLATE_BEGIN_ARGS(STextureGraphInsightResourceViewRow) {}
+	SLATE_ARGUMENT(FItem, Item)
+		SLATE_END_ARGS()
+
+public:
+
+	enum Column {
+		Main = 0,
+		Name,
+		Hash,
+		RefCount,
+		NumMapped,
+		Mapped,
+		Tiles,
+		Pixels,
+
+		NUM_COLUMNS,
+	};
+	using ColumnInfo = std::pair<FName, float>;
+	static ColumnInfo s_columnNames[NUM_COLUMNS];
+	static Column NameToColumn(const FName& name) {
+		for (int i = 0; i < NUM_COLUMNS; ++i)
+		{
+			if (name == s_columnNames[i].first)
+				return Column(i);
+		}
+		return NUM_COLUMNS;
+	}
+
+	virtual TSharedRef<SWidget> GenerateWidgetForColumn(const FName& columnName) override
+	{
+		if (_recordID.IsValid()) {
+
+			Column column = NameToColumn(columnName);
+
+			if (column == Main) {
+				// Rows in a TreeView need an expander button and some indentation
+				return SNew(SHorizontalBox)
+					+ SHorizontalBox::Slot()
+					.AutoWidth()
+					.HAlign(HAlign_Right)
+					.VAlign(VAlign_Fill)
+					[
+						SNew(SExpanderArrow, SharedThis(this))
+						.StyleSet(ExpanderStyleSet)
+					]
+				+ SHorizontalBox::Slot()
+					.FillWidth(1)
+					.Padding(FMargin(4.0f, 0.0f))
+					.VAlign(VAlign_Center)
+					[
+						SAssignNew(_textBoxes[column], STextBlock)
+						.Text(GetTextForColumn(column))
+					];
+			}
+			else
+			{
+				return SNew(SBox)
+					.Padding(FMargin(4.0f, 0.0f))
+					.VAlign(VAlign_Center)
+					[
+						SAssignNew(_textBoxes[column], STextBlock)
+						.Text(GetTextForColumn(column))
+					.OnDoubleClicked(this, &STextureGraphInsightResourceViewRow::OnDoubleClickedResultColumn)
+					];
+			}
+		}
+		// default to null widget if property cannot be found
+		return SNullWidget::NullWidget;
+	}
+
+	FText GetTextForColumn(Column column) const
+	{
+		const auto& br = TextureGraphInsight::Instance()->GetSession()->GetRecord().GetBlob(_recordID);
+
+		FString s;
+		switch (column)
+		{
+		case Main:
+				s = FString::FromInt(_recordID.Blob());
+				if (br.ReplayCount > 0)
+					s += " ...#" + FString::FromInt(br.ReplayCount);
+			break;
+		case Hash:
+			s = HashToFString(br.HashValue);
+			break;
+		case Name:
+				s = br.Name;
+			break;
+		case RefCount:
+			{
+				auto refCount = TextureGraphInsight::Instance()->GetSession()->GetCache().GetBlob(_recordID).use_count();
+				s = FString::FromInt(refCount);
+			}
+			break;
+		case NumMapped:
+			if (br.NumMapped > 0)
+				s = FString::FromInt(br.NumMapped);
+			break;
+		case Mapped:
+			if (br.MappedID.IsValid())
+				s = FString::FromInt(br.MappedID.Blob());
+			break;
+		case Tiles:
+			if (br.NumTiles())
+			{
+				s = (br.Grid().IsUnique() ? "1" : FString::FromInt(br.Grid().Rows()) + " x " + FString::FromInt(br.Grid().Cols()));
+			}
+			break;
+		case Pixels:
+			if (br.TexWidth || br.TexHeight)
+			{
+				s = FString::FromInt(br.TexWidth) + " x " + FString::FromInt(br.TexHeight);
+			}
+			break;
+		}
+
+		return FText::FromString(s);
+	}
+
+
+	void Construct(const FArguments& InArgs, const TSharedRef<STableViewBase>& InOwnerTableView)
+	{
+		_recordID = InArgs._Item->_recordID;
+		SMultiColumnTableRow<FItem >::Construct(FSuperRowType::FArguments(), InOwnerTableView);
+
+		Refresh();
+	}
+
+	void Refresh()
+	{
+		const auto& br = TextureGraphInsight::Instance()->GetSession()->GetRecord().GetBlob(_recordID);
+		if (br.ReplayCount > 0)
+			SetColorAndOpacity(FLinearColor(1, 0, 0));
+		if (br.MappedID.IsValid())
+			SetColorAndOpacity(FLinearColor(0.5, 0.5, 0.5));
+		if (br.NumMapped > 0)
+			SetColorAndOpacity(FLinearColor(0.5, 1.0, 1.0));
+
+		_textBoxes[Main]->SetText(GetTextForColumn(Main));
+		_textBoxes[RefCount]->SetText(GetTextForColumn(RefCount));
+		_textBoxes[NumMapped]->SetText(GetTextForColumn(NumMapped));
+		_textBoxes[Mapped]->SetText(GetTextForColumn(Mapped));
+		_textBoxes[Tiles]->SetText(GetTextForColumn(Tiles));
+		_textBoxes[Pixels]->SetText(GetTextForColumn(Pixels));
+	}
+
+	FReply OnDoubleClickedResultColumn(/** The geometry of the widget*/
+		const FGeometry&,
+		/** The Mouse Event that we are processing */
+		const FPointerEvent&)
+	{
+		TextureGraphInsight::Instance()->GetSession()->SendToInspector(_recordID);
+		return FReply::Handled();
+	}
+
+protected:
+	RecordID _recordID;
+	TSharedPtr<STextBlock> _textBoxes[NUM_COLUMNS];
+};
+
+STextureGraphInsightResourceViewRow::ColumnInfo STextureGraphInsightResourceViewRow::s_columnNames[] = {
+	{ FName(TEXT("Blob Id")), 0.05},    // main
+	{ FName(TEXT("Name")), 0.3},
+	{ FName(TEXT("Hash")), 0.1},
+	{ FName(TEXT("RefCount")), 0.05},
+	{ FName(TEXT("NumMapped")), 0.05},
+	{ FName(TEXT("Mapped")), 0.05},
+	{ FName(TEXT("Tiles")), 0.05},
+	{ FName(TEXT("Pixels")), 0.05},
+};
+
+void STextureGraphInsightResourceView::Construct(const FArguments& Args)
+{
+
+	TSharedPtr<SHeaderRow> headerRow = SNew(SHeaderRow);
+
+	for (int i = 0; i < STextureGraphInsightResourceViewRow::NUM_COLUMNS; i++)
+	{
+		headerRow->AddColumn(
+			SHeaderRow::Column(STextureGraphInsightResourceViewRow::s_columnNames[i].first)
+			.DefaultLabel(FText::FromString(STextureGraphInsightResourceViewRow::s_columnNames[i].first.ToString()))
+			.FillWidth(STextureGraphInsightResourceViewRow::s_columnNames[i].second)
+		);
+	}
+
+	ChildSlot
+	[
+		SAssignNew(_tableView, SItemTableView)
+		.ItemHeight(24)
+		//.ListItemsSource(&_rootItems)
+		.TreeItemsSource(&_rootItems)
+		.OnGenerateRow(this, &STextureGraphInsightResourceView::OnGenerateRowForView)
+		.OnGetChildren(this, &STextureGraphInsightResourceView::OnGetChildrenForView)
+
+		.OnMouseButtonDoubleClick(this, &STextureGraphInsightResourceView::OnDoubleClickItemForView)
+		.OnMouseButtonClick(this, &STextureGraphInsightResourceView::OnClickItemForView)
+		.HeaderRow(headerRow)
+	];
+
+	// install the observer notifications
+	auto sr = StaticCastSharedRef<STextureGraphInsightResourceView>(this->AsShared());
+	TextureGraphInsight::Instance()->GetSession()->OnBlobAdded().AddSP(sr, &STextureGraphInsightResourceView::OnBlobNew);
+	TextureGraphInsight::Instance()->GetSession()->OnBlobMapped().AddSP(sr, &STextureGraphInsightResourceView::OnBlobMapped);
+	TextureGraphInsight::Instance()->GetSession()->OnEngineReset().AddSP(sr, &STextureGraphInsightResourceView::OnEngineReset);
+
+}
+
+TSharedRef<ITableRow> STextureGraphInsightResourceView::OnGenerateRowForView(FItem item, const TSharedRef<STableViewBase>& OwnerTable)
+{
+	return SAssignNew(item->_widget, STextureGraphInsightResourceViewRow, OwnerTable).Item(item);
+}
+
+void STextureGraphInsightResourceView::OnClickItemForView(FItem item)
+{
+	// Item clicked, let's send it to the inspector
+	auto record = item->_recordID;
+	TextureGraphInsight::Instance()->GetSession()->SendToInspector(record);
+}
+void STextureGraphInsightResourceView::OnDoubleClickItemForView(FItem item)
+{
+	// Item double clicked, let's send it to the inspector
+	auto record = item->_recordID;
+	TextureGraphInsight::Instance()->GetSession()->SendToInspector(record);
+}
+
+void STextureGraphInsightResourceView::OnBlobNew(const RecordIDArray& rids)
+{
+	for (auto& rid : rids)
+	{
+		const auto& br = TextureGraphInsight::Instance()->GetSession()->GetRecord().GetBlob(rid);
+
+		if (br.SourceID.IsJob() || !br.SourceID.IsValid())
+		{
+			FItem item = MakeShareable(new FItemData(rid));
+			_rootItems.Add(item);
+
+			if (br.IsTiled())
+			{
+				for (int i = 0; i < br.NumUniqueBlobs(); ++i)
+				{
+					item->_children.Add(MakeShareable(new FItemData(br.GetUniqueBlob(i))));
+				}
+			}
+		}
+	}
+
+	_tableView->RequestListRefresh();
+}
+
+
+bool operator ==(const STextureGraphInsightResourceView::FItem& a, const  RecordID& i)
+{
+	return a->_recordID.id == i.id;
+}
+
+void STextureGraphInsightResourceView::OnBlobMapped(const RecordIDArray& rids)
+{
+	RefreshRootItems();
+}
+
+		
+void STextureGraphInsightResourceView::RefreshRootItems()
+{
+	for (auto& i : _rootItems)
+	{
+		if (i->_widget)
+		{
+			i->_widget->Refresh();
+
+			for (auto& c : i->_children)
+			{
+				if (c->_widget)
+				{
+					c->_widget->Refresh();
+				}
+			}
+		}
+	}
+}
+
+void STextureGraphInsightResourceView::OnEngineReset(int id)
+{
+	_rootItems.Empty();
+	_tableView->RequestListRefresh();
+}
+
diff --git a/Engine/Plugins/Experimental/TextureGraph/Source/TextureGraphInsight/Private/View/STextureGraphInsightSessionView.cpp b/Engine/Plugins/Experimental/TextureGraph/Source/TextureGraphInsight/Private/View/STextureGraphInsightSessionView.cpp
new file mode 100644
--- /dev/null
+++ b/Engine/Plugins/Experimental/TextureGraph/Source/TextureGraphInsight/Private/View/STextureGraphInsightSessionView.cpp
@@ -0,0 +1,652 @@
+// Copyright Epic Games, Inc. All Rights Reserved.
+#include "View/STextureGraphInsightSessionView.h"
+
+#include "TextureGraphInsight.h"
+#include "Model/TextureGraphInsightSession.h"
+#include <Widgets/Views/STableViewBase.h>
+#include <Widgets/Views/ITableRow.h>
+#include <Brushes/SlateColorBrush.h>
+#include <Framework/MultiBox/MultiBoxBuilder.h>
+#include <Framework/Commands/UIAction.h>
+
+
+class STextureGraphInsightBatchJobViewRow : public SMultiColumnTableRow<STextureGraphInsightBatchJobView::FItem>
+{
+public:
+	using FItem = STextureGraphInsightBatchJobView::FItem;
+
+	SLATE_BEGIN_ARGS(STextureGraphInsightBatchJobViewRow) {}
+	SLATE_ARGUMENT(FItem, Item)
+	SLATE_END_ARGS()
+
+public:
+
+	enum Column {
+		Main = 0,
+		Name,
+		Hash,
+		Tiles,
+		Pixels,
+		Fillrate,
+		Timeline,
+
+		NUM_COLUMNS,
+	};
+	using ColumnInfo = std::tuple<FName, float, ETextJustify::Type>;
+	static ColumnInfo s_columnNames[NUM_COLUMNS];
+	static Column NameToColumn(const FName& name) {
+		for (int i = 0; i < NUM_COLUMNS; ++i)
+		{
+			if (name == std::get<0>(s_columnNames[i]))
+				return Column(i);
+		}
+		return NUM_COLUMNS;
+	}
+
+	static FLinearColor GetColorForBatchState(const BatchRecord& br)
+	{
+		auto color = FLinearColor(1, 0, 0);
+		if (br.bIsDone || br.bIsJobsDone)
+		{
+			if (br.bIsJobsDone)
+			{
+				color = FLinearColor(1, 1, 1); // Back to normal
+			}
+			else
+			{
+				color = FLinearColor(1, 0.5, 0); // Waiting on jobs
+			}
+		}
+
+		if (br.bIsNoCache)
+			color *= 0.6;
+
+		return color;
+	}
+
+	static FLinearColor GetColorForJobState(bool isDone, bool isNoOp, bool isMainPhase)
+	{
+		if (!isMainPhase)
+		{
+			return FLinearColor(0.5, 0.5, 0.8); // Not main phase
+		}
+		else if (isNoOp)
+		{
+			return FLinearColor(0.5, 0.8, 0.5); // No op
+		}
+		else
+		{
+			return FLinearColor(1, 1, 1); // Back to normal
+		}
+	}
+
+	void Construct(const FArguments& InArgs, const TSharedRef<STableViewBase>& InOwnerTableView)
+	{
+		_recordID = InArgs._Item->_recordID;
+		_isMainPhase = InArgs._Item->_isMainPhase;
+
+		const auto& br = TextureGraphInsight::Instance()->GetSession()->GetRecord().GetBatch(_recordID);
+		if (_recordID.IsBatch())
+		{
+			_accentColor = GetColorForBatchState(br);
+		}
+
+		if (_recordID.IsJob())
+		{
+			const auto& jr = br.GetJob(_recordID);
+			_accentColor = GetColorForJobState(jr.bIsDone, jr.IsNoOp(), jr.IsMainPhase());
+		}
+
+		SMultiColumnTableRow<FItem >::Construct(FSuperRowType::FArguments(), InOwnerTableView);
+	}
+
+	virtual TSharedRef<SWidget> GenerateWidgetForColumn(const FName& columnName) override
+	{
+		static FSlateBrush* s_timelineBrush = new FSlateColorBrush(FColor(125, 125, 255));
+		static FSlateBrush* s_timelineBrush2 = new FSlateColorBrush(FColor(125,  255, 125));
+		if (_recordID.IsValid()) {
+
+
+			Column column = NameToColumn(columnName);
+			if (column == Main) {
+				// Rows in a TreeView need an expander button and some indentation
+				return SNew(SHorizontalBox)
+						+ SHorizontalBox::Slot()
+						.AutoWidth()
+						.HAlign(HAlign_Right)
+						.VAlign(VAlign_Fill)
+						[
+							SNew(SExpanderArrow, SharedThis(this))
+							.StyleSet(ExpanderStyleSet)
+						]
+						+ SHorizontalBox::Slot()
+						.FillWidth(1)
+						.Padding(FMargin(4.0f, 0.0f))
+						.VAlign(VAlign_Center)
+						[
+							SAssignNew(_textBoxes[column], STextBlock)
+							.Text(GetTextForColumn(column))
+							.ColorAndOpacity(_accentColor)
+						];
+			}
+			else if (column == Timeline)
+			{
+
+				auto& sr = TextureGraphInsight::Instance()->GetSession()->GetRecord();
+				const auto& br = sr.GetBatch(_recordID);
+				const auto& jr = br.GetJob(_recordID);
+
+				if (_recordID.IsJob())
+				{
+					auto positionPre = FVector2D((jr.GetBeginTimeMS() - br.BeginTimeMS) / br.ScopeTime_ms(), 0.7);
+					auto sizePre = FVector2D(jr.GetPreRunTimeMS() / br.ScopeTime_ms(), 0.2);
+
+					auto position = FVector2D((jr.GetBeginRunTimeMS() - br.BeginTimeMS) / br.ScopeTime_ms(), 0.1);
+					auto sizeRun = FVector2D(jr.GetRunTimeMS() / br.ScopeTime_ms(), 0.7);
+				
+					return	
+						SNew(SVerticalBox)
+
+						// First row
+						// Display time top, justified aright at the end of the runtime
+						+SVerticalBox::Slot()
+						.FillHeight(sizeRun.Y)
+						
+						[
+							SNew(SHorizontalBox)
+
+							+ SHorizontalBox::Slot()
+							.FillWidth(position.X)
+
+							+ SHorizontalBox::Slot()
+							.FillWidth(1 - position.X)
+							.Padding(FMargin(2.0f, 0.0f))
+							.VAlign(EVerticalAlignment::VAlign_Top)
+							[
+								SAssignNew(_textBoxes[column], STextBlock)
+								.Text(GetTextForColumn(column))
+								//.Justification(ETextJustify::Right)
+								.Justification(ETextJustify::Left)
+								.ColorAndOpacity(_accentColor)
+							]
+						
+						]
+
+						// Second row
+						// Display rectangle for pretime, then for runtime
+						+SVerticalBox::Slot()
+						.FillHeight(sizePre.Y)
+						[
+							SNew(SHorizontalBox)
+
+							+ SHorizontalBox::Slot()
+							.FillWidth(positionPre.X)
+						
+							+ SHorizontalBox::Slot()
+							.FillWidth(sizePre.X)
+							[
+								SNew(SBorder)
+								.BorderImage(s_timelineBrush2)
+							]
+
+							+ SHorizontalBox::Slot()
+							.FillWidth(sizeRun.X)
+							[
+								SNew(SBorder)
+								.BorderImage(s_timelineBrush)
+							]
+
+							+ SHorizontalBox::Slot()
+							.FillWidth(1.0 - position.X - sizeRun.X)
+						]
+
+						+ SVerticalBox::Slot()
+						.FillHeight(1.0 - sizeRun.Y - sizePre.Y)
+						;
+				} else {
+					auto sizeX = 0.5;
+					auto sizeY = 0.1;
+					auto positionY = 0.45;
+
+					return	SNew(SHorizontalBox)
+						+ SHorizontalBox::Slot()
+						.FillWidth(sizeX)
+						[
+							SNew(SVerticalBox)
+							+ SVerticalBox::Slot()
+							.FillHeight(positionY)
+							+ SVerticalBox::Slot()
+							.FillHeight(sizeY)
+							[
+								SNew(SBorder)
+								.BorderImage(s_timelineBrush)
+							]
+							+ SVerticalBox::Slot()
+							.FillHeight(1.0 - positionY - sizeY)
+						]
+						+ SHorizontalBox::Slot()
+						.AutoWidth()
+						.Padding(2.0)
+						.VAlign(EVerticalAlignment::VAlign_Center)
+						.HAlign(HAlign_Center)
+						[
+							SAssignNew(_textBoxes[column], STextBlock)
+							.Text(GetTextForColumn(Timeline))
+							.ColorAndOpacity(_accentColor)
+						]
+						+ SHorizontalBox::Slot()
+						.FillWidth(0.5)
+						[
+							SNew(SVerticalBox)
+							+ SVerticalBox::Slot()
+							.FillHeight(positionY)
+							+ SVerticalBox::Slot()
+							.FillHeight(sizeY)
+							[
+								SNew(SBorder)
+								.BorderImage(s_timelineBrush)
+							]
+							+ SVerticalBox::Slot()
+							.FillHeight(1.0 - positionY - sizeY)
+						];
+				}
+			}
+			else if (column == Name && _recordID.IsBatch())
+			{
+				return 	SNew(SHorizontalBox)
+					+ SHorizontalBox::Slot()
+					.Padding(2.0)
+					.FillWidth(1)
+					[
+						SAssignNew(_textBoxes[column], STextBlock)
+						.Text(GetTextForColumn(column))
+						.Justification(ETextJustify::Left)
+						.ColorAndOpacity(_accentColor)
+						.OnDoubleClicked(this, &STextureGraphInsightBatchJobViewRow::OnDoubleClickedResultColumn)
+					];
+			}
+			else
+			{
+				return 	SNew(SHorizontalBox)
+					+ SHorizontalBox::Slot()
+					.Padding(2.0)
+					.FillWidth(1)
+					[	
+						SAssignNew(_textBoxes[column], STextBlock)
+						.Text(GetTextForColumn(column))
+						.Justification(std::get<2>(STextureGraphInsightBatchJobViewRow::s_columnNames[column]))
+						.ColorAndOpacity(_accentColor)
+						.OnDoubleClicked(this, &STextureGraphInsightBatchJobViewRow::OnDoubleClickedResultColumn)
+					];
+			}
+		}
+		// default to null widget if property cannot be found
+		return SNullWidget::NullWidget;
+	}
+
+	FText GetTextForColumn(Column column) const
+	{
+		auto& sr = TextureGraphInsight::Instance()->GetSession()->GetRecord();
+		const auto& br = sr.GetBatch(_recordID);
+		const auto& jr = br.GetJob(_recordID);
+		const auto& mr = sr.GetMix(br.MixID);
+
+		FString s;
+		switch (column)
+		{
+		case Main:
+			if (_recordID.IsBatch())
+			{
+				s = (_isMainPhase ? "Batch " : "     Sub batch ") + FString::FromInt(br.BatchID);
+				if (br.bIsNoCache)
+					s += " - NO CACHE"; 
+				if (br.bIsFromIdle)
+					s += " - IDLE";
+				if (br.ReplayCount > 0)
+					s += " ...#" + FString::FromInt(br.ReplayCount);
+			}
+			else
+			{
+				s = "Job " + FString::FromInt(jr.JobIdx);
+				if (jr.PhaseIdx)
+					s += " : " + FString::FromInt(jr.PhaseIdx);
+				if (br.ReplayCount > 0)
+					s += " ...#" + FString::FromInt(br.ReplayCount);
+			}
+			break;
+		case Name:
+			if (_recordID.IsBatch())
+			{
+				s = mr.Name;
+				if (br.Action.Len())
+					s += " - " + br.Action;
+			}
+			else
+			{
+				s = jr.TransformName;
+			}
+			break;
+		case Hash:
+			if (_recordID.IsBatch())
+			{
+			} else
+			{
+				s = HashToFString(jr.JobHash);
+			}
+			break;
+		case Tiles:
+			if (_recordID.IsBatch())
+			{
+				s = FString::Printf(TEXT("%8d"), br.NumInvalidatedTiles) + " / " + FString::FromInt(br.NumTiles);
+			}
+			else
+			{
+				if (jr.GetNumTiles())
+					s = FString::Printf(TEXT("%8d"), jr.GetNumInvalidated()) + " / " + (jr.GetGrid().IsUnique() ? "1" : FString::FromInt(jr.GetGrid().Rows()) + " x " + FString::FromInt(jr.GetGrid().Cols()));
+			}
+			break;
+		case Pixels:
+			if (_recordID.IsBatch())
+			{
+			}
+			else
+			{
+				if (jr.TexWidth || jr.TexHeight)
+					s = FString::FromInt(jr.TexWidth) + " x " + FString::FromInt(jr.TexHeight);
+			}
+			break;
+		case Fillrate:
+			if (_recordID.IsBatch())
+			{
+			}
+			else
+			{
+				if (jr.TexWidth || jr.TexHeight)
+					s = FString::FromInt(jr.GetFillRate() * 0.000001); /// Fillrate expressed in Million pix per sec
+			}
+		break;
+		case Timeline:
+			if (_recordID.IsBatch())
+			{
+				//s = FString::FromInt(br.ScopeTime_ms());
+				s = FString::Printf(TEXT("%.1f"), br.ScopeTime_ms());
+			}
+			else
+			{
+				//s = FString::FromInt(jr.RunTime_ms());
+				s = FString::Printf(TEXT("%.1f"), jr.GetRunTimeMS());
+			}
+			break;
+		}
+
+		return FText::FromString(s);
+	}
+
+	void OnUpdate()
+	{
+		const auto& br = TextureGraphInsight::Instance()->GetSession()->GetRecord().GetBatch(_recordID);
+		SetColorAndOpacity(GetColorForBatchState(br));
+
+		if (_textBoxes[Main])
+		{
+			_textBoxes[Main]->SetText(GetTextForColumn(Main));
+			if (_recordID.IsBatch())
+				_textBoxes[Name]->SetText(GetTextForColumn(Name));
+			_textBoxes[Tiles]->SetText(GetTextForColumn(Tiles));
+			_textBoxes[Pixels]->SetText(GetTextForColumn(Pixels));
+			_textBoxes[Fillrate]->SetText(GetTextForColumn(Fillrate));
+			_textBoxes[Timeline]->SetText(GetTextForColumn(Timeline));
+		}
+	}
+
+	FReply OnDoubleClickedResultColumn(/** The geometry of the widget*/
+		const FGeometry&,
+		/** The Mouse Event that we are processing */
+		const FPointerEvent&)
+	{
+		TextureGraphInsight::Instance()->GetSession()->SendToInspector(_recordID);
+		return FReply::Handled();
+	}
+
+protected:
+	RecordID _recordID;
+	TSharedPtr<STextBlock> _textBoxes[NUM_COLUMNS];
+	FLinearColor _accentColor = FLinearColor(1, 1, 1);
+	bool _isMainPhase = true;
+};
+
+STextureGraphInsightBatchJobViewRow::ColumnInfo STextureGraphInsightBatchJobViewRow::s_columnNames[] = {
+	{ FName(TEXT("Id")), 0.1, ETextJustify::Left},    // main
+	{ FName(TEXT("Name")), 0.15, ETextJustify::Right},
+	{ FName(TEXT("Hash")), 0.1, ETextJustify::Center},
+	{ FName(TEXT("Tiles")), 0.08, ETextJustify::Left},
+	{ FName(TEXT("Pixels")), 0.08, ETextJustify::Center},
+	{ FName(TEXT("Fillrate [Mpx/s]")), 0.08, ETextJustify::Right},
+	{ FName(TEXT("Timeline [ms]")), 0.3, ETextJustify::Center},
+};
+
+void STextureGraphInsightBatchJobView::Construct(const FArguments& Args)
+{
+	_inspectOnSimpleClick = Args._inspectOnSimpleClick;
+
+	TSharedPtr<SHeaderRow> headerRow = SNew(SHeaderRow);
+
+	for (int i = 0; i < STextureGraphInsightBatchJobViewRow::NUM_COLUMNS; i++)
+	{
+		headerRow->AddColumn(
+			SHeaderRow::Column(std::get<0>(STextureGraphInsightBatchJobViewRow::s_columnNames[i]))
+			.DefaultLabel(FText::FromString(std::get<0>(STextureGraphInsightBatchJobViewRow::s_columnNames[i]).ToString()))
+			.HAlignHeader(EHorizontalAlignment::HAlign_Center)
+			.HAlignCell(EHorizontalAlignment::HAlign_Fill)
+			.FillWidth(std::get<1>(STextureGraphInsightBatchJobViewRow::s_columnNames[i]))
+		);
+	}
+
+	ChildSlot
+	[
+		
+	SAssignNew(_treeView, SItemTreeView)
+		.ItemHeight(24)
+		.TreeItemsSource(&_rootItems)
+		.OnGenerateRow(this, &STextureGraphInsightBatchJobView::OnGenerateRowForTree)
+		.OnGetChildren(this, &STextureGraphInsightBatchJobView::OnGetChildrenForView)
+		.OnMouseButtonDoubleClick(this, &STextureGraphInsightBatchJobView::OnDoubleClickItemForTree)
+		.OnMouseButtonClick(this, &STextureGraphInsightBatchJobView::OnClickItemForTree)
+		.OnContextMenuOpening(this, &STextureGraphInsightBatchJobView::OnContextMenuOpeningForTree)
+		.HeaderRow(headerRow)
+	];
+
+	OnBatchNew(Args._recordID);
+
+	if (Args._recordID.IsBatch())
+	{
+		_treeView->SetItemExpansion(_rootItems[0], true);
+	}
+}
+
+TSharedRef<ITableRow> STextureGraphInsightBatchJobView::OnGenerateRowForTree(FItem item, const TSharedRef<STableViewBase>& OwnerTable)
+{
+	return SAssignNew(item->_widget, STextureGraphInsightBatchJobViewRow, OwnerTable).Item(item);
+}
+
+void STextureGraphInsightBatchJobView::OnClickItemForTree(FItem item)
+{
+	// Item clicked, let's send it to the inspector
+	auto record = item->_recordID;
+	if (_inspectOnSimpleClick)
+		TextureGraphInsight::Instance()->GetSession()->SendToInspector(record);
+}
+void STextureGraphInsightBatchJobView::OnDoubleClickItemForTree(FItem item)
+{
+	// Item double clicked, let's send it to the inspector
+	auto record = item->_recordID;
+	TextureGraphInsight::Instance()->GetSession()->SendToInspector(record);
+}
+
+TSharedPtr<SWidget> STextureGraphInsightBatchJobView::OnContextMenuOpeningForTree()
+{
+	// Only if a batch is focused then we would show a context menu
+	auto selection = _treeView->GetSelectedItems();
+	if (selection.Num())
+	{
+		if (selection[0]->_recordID.IsBatch())
+		{
+			return OnContextMenuBatch(selection[0]);
+		}
+		if (selection[0]->_recordID.IsJob())
+		{
+			return OnContextMenuJob(selection[0]);
+		}
+	}
+
+	return nullptr;
+}
+TSharedPtr<SWidget> STextureGraphInsightBatchJobView::OnContextMenuBatch(FItem batchItem)
+{
+	const bool closeAfterSelection = true;
+	FMenuBuilder menuBuilder(closeAfterSelection, NULL);
+
+	// Begin menu section
+	menuBuilder.BeginSection("BatchContextMenu");
+	{
+		menuBuilder.AddWidget(SNew(STextBlock).Text(FText::FromString("Batch " + FString::FromInt(batchItem->_recordID.Batch()) + " ...")), FText(), true);
+		// Create action delegate and Add menu entry
+		{
+			FUIAction action = FUIAction(FExecuteAction::CreateRaw(this, &STextureGraphInsightBatchJobView::OnReplayBatch, batchItem->_recordID, false));
+			menuBuilder.AddMenuEntry(FText::FromString(FString(TEXT("Replay Batch"))), FText(), FSlateIcon(), action);
+		}
+		{
+			FUIAction action = FUIAction(FExecuteAction::CreateRaw(this, &STextureGraphInsightBatchJobView::OnReplayBatch, batchItem->_recordID, true));
+			menuBuilder.AddMenuEntry(FText::FromString(FString(TEXT("Replay Batch & Capture RenderDoc"))), FText(), FSlateIcon(), action);
+		}
+	}
+	menuBuilder.EndSection();
+
+	return menuBuilder.MakeWidget();
+}
+
+
+TSharedPtr<SWidget> STextureGraphInsightBatchJobView::OnContextMenuJob(FItem jobItem)
+{
+	const bool closeAfterSelection = true;
+	FMenuBuilder menuBuilder(closeAfterSelection, NULL);
+
+	// Begin menu section
+	menuBuilder.BeginSection("JobContextMenu");
+	{
+		menuBuilder.AddWidget(SNew(STextBlock).Text(FText::FromString("Batch " + FString::FromInt(jobItem->_recordID.Batch()) + " - Job " + FString::FromInt(jobItem->_recordID.Job()) + " ...")), FText(), true);
+		// Create action delegate and Add menu entry
+		{
+			FUIAction action = FUIAction(FExecuteAction::CreateRaw(this, &STextureGraphInsightBatchJobView::OnReplayJob, jobItem->_recordID, false));
+			menuBuilder.AddMenuEntry(FText::FromString(FString(TEXT("Replay Job"))), FText(), FSlateIcon(), action);
+		}
+		{
+			FUIAction action = FUIAction(FExecuteAction::CreateRaw(this, &STextureGraphInsightBatchJobView::OnReplayJob, jobItem->_recordID, true));
+			menuBuilder.AddMenuEntry(FText::FromString(FString(TEXT("Replay Job & Capture RenderDoc"))), FText(), FSlateIcon(), action);
+		}
+	}
+	menuBuilder.EndSection();
+
+	return menuBuilder.MakeWidget();
+}
+
+
+void STextureGraphInsightBatchJobView::OnReplayBatch(RecordID batchId, bool captureRenderDoc)
+{
+	if (!batchId.IsBatch())
+		return;
+
+	// Trigger a replay!
+	bool result = TextureGraphInsight::Instance()->GetSession()->ReplayBatch(batchId, captureRenderDoc);
+
+}
+
+void STextureGraphInsightBatchJobView::OnReplayJob(RecordID jobId, bool captureRenderDoc)
+{
+	if (!jobId.IsJob())
+		return;
+
+	// Trigger a replay!
+	bool result = TextureGraphInsight::Instance()->GetSession()->ReplayJob(jobId, captureRenderDoc);
+
+}
+
+void STextureGraphInsightBatchJobView::OnBatchNew(RecordID batchRecordID)
+{
+
+	// Fetch the actual data from the record and add job's items
+	const auto& br = TextureGraphInsight::Instance()->GetSession()->GetRecord().GetBatch(batchRecordID);
+
+	FItem batchItem = MakeShareable(new FItemData(batchRecordID));
+	
+
+	FItem batchItemSubs = MakeShareable(new FItemData(batchRecordID, false));
+
+	uint32_t jobIdx = 0;
+	for (const auto& j : br.Jobs)
+	{
+		if (j.IsMainPhase())
+		{
+			batchItem->_children.Add(MakeShareable(new FItemData(RecordID::fromBatchJob(batchRecordID.Batch(), jobIdx))));
+		}
+		else
+		{
+			batchItemSubs->_children.Add(MakeShareable(new FItemData(RecordID::fromBatchJob(batchRecordID.Batch(), jobIdx), false)));
+		}
+
+		jobIdx++;
+	}
+
+	batchItem->_children.Add(batchItemSubs);
+
+	_rootItems.Add(batchItem);
+
+	_treeView->RequestTreeRefresh();
+
+}
+
+void STextureGraphInsightBatchJobView::OnBatchUpdate(RecordID batchRecordID)
+{
+	const auto& bri = _rootItems[batchRecordID.Batch()];
+	if (bri->_widget) {
+		bri->_widget->OnUpdate();
+	}
+	for (const auto& j : bri->_children)
+	{
+		if (j->_widget) {
+			j->_widget->OnUpdate();
+		}
+	}
+
+}
+
+void STextureGraphInsightBatchJobView::OnEngineReset(int32 id)
+{
+	_rootItems.Empty();
+	_treeView->RequestTreeRefresh();
+}
+
+
+
+
+
+
+
+void STextureGraphInsightSessionView::Construct(const FArguments& Args)
+{
+	ChildSlot
+	[
+		SAssignNew(_batchJobView, STextureGraphInsightBatchJobView)
+		.inspectOnSimpleClick(true)
+	];
+
+
+	// install the observer notifications
+	auto sr = StaticCastSharedRef<STextureGraphInsightBatchJobView>(this->_batchJobView->AsShared());
+
+	TextureGraphInsight::Instance()->GetSession()->OnBatchAdded().AddSP(sr, &STextureGraphInsightBatchJobView::OnBatchNew);
+	TextureGraphInsight::Instance()->GetSession()->OnBatchDone().AddSP(sr, &STextureGraphInsightBatchJobView::OnBatchUpdate);
+	TextureGraphInsight::Instance()->GetSession()->OnBatchJobsDone().AddSP(sr, &STextureGraphInsightBatchJobView::OnBatchUpdate);
+	TextureGraphInsight::Instance()->GetSession()->OnEngineReset().AddSP(sr, &STextureGraphInsightBatchJobView::OnEngineReset);
+}
+
diff --git a/Engine/Plugins/Experimental/TextureGraph/Source/TextureGraphInsight/Public/Model/TextureGraphInsightObserver.h b/Engine/Plugins/Experimental/TextureGraph/Source/TextureGraphInsight/Public/Model/TextureGraphInsightObserver.h
new file mode 100644
--- /dev/null
+++ b/Engine/Plugins/Experimental/TextureGraph/Source/TextureGraphInsight/Public/Model/TextureGraphInsightObserver.h
@@ -0,0 +1,79 @@
+// Copyright Epic Games, Inc. All Rights Reserved.
+#pragma once
+
+#include "TextureGraphEngine.h"
+#include "Device/Device.h"
+#include "Device/DeviceObserverSource.h"
+#include "Data/Blobber.h"
+#include "Job/Scheduler.h"
+
+using HashArray = DeviceObserverSource::HashArray;
+
+DECLARE_LOG_CATEGORY_EXTERN(LogTextureGraphInsightObserver, Log, All);
+
+/// Concrete DeviceObserver interface
+class TEXTUREGRAPHINSIGHT_API TextureGraphInsightDeviceObserver : public DeviceObserverSource
+{
+protected:
+	DeviceType DevType;
+
+	/// Protected interface of emitters called by the device to notify the observers
+	virtual void DeviceBuffersUpdated(HashNDescArray&& AddedBuffers, HashArray&& RemovedBuffers) override;
+public:
+	TextureGraphInsightDeviceObserver();
+	explicit TextureGraphInsightDeviceObserver(DeviceType);
+	virtual ~TextureGraphInsightDeviceObserver() override;
+};
+
+/// Concrete BlobberObserver interface
+class TEXTUREGRAPHINSIGHT_API TextureGraphInsightBlobberObserver : public BlobberObserverSource
+{
+private:
+protected:
+	/// Protected interface of emitters called by the device to notify the observers
+	virtual void BlobberUpdated(HashArray&& AddedHashes, HashArray&& RemappedHashes) override;
+
+public:
+	TextureGraphInsightBlobberObserver();
+	virtual ~TextureGraphInsightBlobberObserver() override;
+};
+
+/// Concrete SchedulerObserver interface
+class TEXTUREGRAPHINSIGHT_API TextureGraphInsightSchedulerObserver : public SchedulerObserverSource
+{
+private:
+protected:
+	/// Protected interface of emitters called by the scheduler to notify the observers
+	virtual void Start() override;
+	virtual void UpdateIdle() override;
+	virtual void Stop() override;
+	virtual void BatchAdded(JobBatchPtr Batch) override;
+	virtual void BatchDone(JobBatchPtr Batch) override;
+	virtual void BatchJobsDone(JobBatchPtr Batch) override;
+
+public:
+
+	TextureGraphInsightSchedulerObserver();
+	virtual ~TextureGraphInsightSchedulerObserver() override;
+};
+
+/// Concrete EngineObserver interface
+/// Responsible for:
+///	  1/ watching the engine life cycle
+///   2/ owning the other system observers, and installing them appropirately when an Engine is active
+///	  3/ notifying Insight
+class TEXTUREGRAPHINSIGHT_API TextureGraphInsightEngineObserver : public EngineObserverSource
+{
+protected:
+	/// Protected interface of emitters called by the engine to notify the observers
+	virtual void Created() override;
+	virtual void Destroyed() override;
+
+public:
+	TextureGraphInsightEngineObserver();
+	virtual ~TextureGraphInsightEngineObserver() override;
+
+	std::shared_ptr<TextureGraphInsightDeviceObserver>			_deviceObservers[(uint32)DeviceType::Count];
+	std::shared_ptr<TextureGraphInsightBlobberObserver>		BlobberObserver;
+	std::shared_ptr<TextureGraphInsightSchedulerObserver>		SchedulerObserver;
+};
diff --git a/Engine/Plugins/Experimental/TextureGraph/Source/TextureGraphInsight/Public/Model/TextureGraphInsightRecord.h b/Engine/Plugins/Experimental/TextureGraph/Source/TextureGraphInsight/Public/Model/TextureGraphInsightRecord.h
new file mode 100644
--- /dev/null
+++ b/Engine/Plugins/Experimental/TextureGraph/Source/TextureGraphInsight/Public/Model/TextureGraphInsightRecord.h
@@ -0,0 +1,462 @@
+// Copyright Epic Games, Inc. All Rights Reserved.
+#pragma once
+
+#include <vector>
+#include <unordered_map>
+
+#include "Data/Blob.h"
+#include "Device/Device.h"
+
+using FStringArray = std::vector<FString>;
+
+template <typename DimType>
+class TEXTUREGRAPHINSIGHT_API T_TiledGrid
+{
+public:
+	using Size = DimType;
+	using Index = int32_t;
+
+	Size							_rows = 0;
+	Size							_cols = 0;
+
+	T_TiledGrid() {}
+	T_TiledGrid(Size size) : T_TiledGrid(size, size) {}
+	T_TiledGrid(Size rows, Size cols) : _rows(rows), _cols(cols) {}
+
+	FORCEINLINE	Index				operator()(Size tx, Size ty) const { check(IsValidTile(tx, ty)); return tx * _cols + ty; }
+	FORCEINLINE Size				Rows() const { return _rows; }
+	FORCEINLINE Size				Cols() const { return _cols; }
+	FORCEINLINE Index				NumTiles() const { return _rows * _cols; }
+	FORCEINLINE bool				IsEmpty() const { return (_rows < 1 || _cols < 1); }
+	FORCEINLINE bool				IsUnique() const { return ( _rows * _cols == 1);  }
+	FORCEINLINE bool				IsTiled() const { return (_rows * _cols > 0); }
+
+	FORCEINLINE bool				IsValidTile(Size tx, Size ty) const { return (tx < _rows && ty < _cols); }
+
+	template<typename Functor>
+	FORCEINLINE void				ForEach(Functor functor) const
+	{
+		for (Size tx = 0; tx < _rows; ++tx)
+		{
+			for (Size ty = 0; ty < _cols; ++ty)
+			{
+				functor(tx, ty);
+			}
+		}
+	}
+};
+using TiledGrid = T_TiledGrid<uint8_t>;
+
+template <typename Type>
+class TEXTUREGRAPHINSIGHT_API T_Tiled
+{
+public:
+	using Size = typename TiledGrid::Size;
+	using TypeArray = std::vector<Type>;
+
+
+	TiledGrid						_grid;
+	TypeArray						_tiles;
+
+	T_Tiled() {}
+	T_Tiled(Size size) : T_Tiled(size, size) {}
+	T_Tiled(Size rows, Size cols) : _grid(rows, cols) { _tiles.resize(_grid.NumTiles()); }
+
+	FORCEINLINE void	 			Resize(Size size) { Resize(size, size); }
+	FORCEINLINE void	 			Resize(Size rows, Size cols) { _grid = TiledGrid(rows, cols); _tiles.resize(_grid.NumTiles()); }
+	FORCEINLINE void	 			Resize(TiledGrid grid) { _grid = grid; _tiles.resize(_grid.NumTiles()); }
+	FORCEINLINE void				Fill(const Type& value) { for (auto& v : _tiles) v = value; }
+
+	FORCEINLINE const Type&			operator()(Size tx, Size ty) const { return _tiles[_grid(tx, ty)]; }
+	FORCEINLINE Type&				operator()(Size tx, Size ty) { return _tiles[_grid(tx, ty)]; }
+	FORCEINLINE const TiledGrid&	Grid() const { return _grid; }
+
+	//FORCEINLINE const TypeArray& Tiles() const { return _tiles; }
+	//FORCEINLINE TypeArray& Tiles() { return _tiles; }
+
+	template<typename Functor>
+	FORCEINLINE void				ForEach(Functor functor) { _grid.ForEach(functor); }
+	template<typename Functor>
+	FORCEINLINE void				ForEachGetTile(Functor functor) const { _grid.ForEach([&] (Size tx, Size ty) { functor(tx, ty, _tiles[_grid(tx, ty)]); }); 	}
+};
+using BoolTiles = T_Tiled<uint8_t>;
+using Uint8Tiles = T_Tiled<uint8_t>;
+
+template<typename Type>
+int32_t TiledSum(const T_Tiled<Type>& tiles)
+{
+	int32_t s = 0;
+	tiles.ForEachGetTile([&s](typename T_Tiled<Type>::Size tx, typename T_Tiled<Type>::Size ty, const Type& v)
+	{
+		s += v;
+	});
+	return s;
+}
+
+using DeviceBufferID = uint64;
+using DeviceBufferIDArray = std::vector<DeviceBufferID>;
+
+using BlobID = uint64;
+
+using BatchID = uint64;
+using FrameID = uint64;
+using QueueID = uint32;
+
+using ActionID = uint64;
+
+using MixID = uint64;
+
+
+static const uint16_t INVALID_SHORT = 0xFFFF;
+static const uint32_t INVALID_INDEX = 0xFFFFFFFF;
+static const uint64_t INVALID_RECORD_ID = 0xFFFFFFFFFFFFFFFF;
+
+struct RecordID
+{
+	enum Type : uint8
+	{
+		kInvalid = 0xFF,
+		kBatch = 0,
+		kJob,
+		kBlob,
+		kDeviceBuffer,
+		kAction,
+		kMix,
+	};
+	union
+	{
+		struct 
+		{
+			uint8			type;
+			uint32_t		prim;
+			uint32_t		second : 24;
+		};
+		uint64			id = INVALID_RECORD_ID;
+	};
+	RecordID(uint8 t, uint32_t p, uint32_t s) :type(t), prim(p), second(s) {};
+	RecordID(uint64 recordId) :id(recordId) {};
+	RecordID() = default;
+
+	FORCEINLINE bool IsValid() const { return id != INVALID_RECORD_ID; }
+	FORCEINLINE bool IsBuffer() const { return (type == kDeviceBuffer); }
+	FORCEINLINE bool IsDevice() const { return IsBuffer() && (prim == INVALID_INDEX); }
+	FORCEINLINE bool IsBlob() const { return (type == kBlob); }
+	FORCEINLINE bool IsJob() const { return (type == kJob); }
+	FORCEINLINE bool IsBatch() const { return (type == kBatch); }
+	FORCEINLINE bool IsAction() const { return (type == kAction); }
+	FORCEINLINE bool IsMix() const { return (type == kMix); }
+
+	FORCEINLINE uint32 Buffer() const { return prim; }
+	FORCEINLINE uint32 Buffer_DeviceType() const { return second; }
+	FORCEINLINE uint32 Blob() const { return prim; }
+	FORCEINLINE uint32 Batch() const { return prim; }
+	FORCEINLINE uint32 Job() const { return second; }
+	FORCEINLINE uint32 Action() const { return prim; }
+	FORCEINLINE uint32 Mix() const { return prim; }
+
+	static RecordID fromBuffer(DeviceType t, uint32_t b) { return RecordID(kDeviceBuffer, b, (uint32)t); }
+	static RecordID fromBlob(uint32_t b) { return RecordID(kBlob, b, INVALID_INDEX); }
+	static RecordID fromBatch(uint32_t b) { return RecordID(kBatch, b, INVALID_INDEX); }
+	static RecordID fromBatchJob(uint32_t b, uint32_t j) { return RecordID(kJob, b, j); }
+	static RecordID fromAction(uint32_t a) { return RecordID(kAction, a, INVALID_INDEX); }
+	static RecordID fromMix(uint32_t a) { return RecordID(kMix, a, INVALID_INDEX); }
+};
+using RecordIDArray = std::vector<RecordID>;
+using RecordIDTiles = T_Tiled<RecordID>;
+
+/// 
+/// DeviceBuffer Record
+///
+class TEXTUREGRAPHINSIGHT_API DeviceBufferRecord
+{
+public:
+	// adding a default constructor to enforce this to be a non POD class to successfully initialize
+	DeviceBufferRecord() {}
+	//RecordID						_selfID;
+	RecordID						SourceID;
+
+	static const DeviceBufferRecord null;
+	DeviceType						DevType = DeviceType::Null;
+	uint64							ID = 0;
+	HashType						HashValue = 0;
+	HashType						PrevHashValue = 0;
+	BufferDescriptor				Descriptor;
+
+	mutable uint64					RawBufferMemSize = 0; /// The true size as informed
+	mutable uint64					Rank = 0;
+
+	mutable bool					bErased = false;
+	mutable bool					bLeaked = false;
+};
+using DeviceBufferRecordArray = std::vector< DeviceBufferRecord >;
+using RecordFromDeviceBufferIDMap = std::unordered_map<DeviceBufferID, RecordID>;
+using RecordFromHashMap = std::unordered_map<HashType, RecordID>;
+
+
+
+/// Blob Record 
+class TEXTUREGRAPHINSIGHT_API BlobRecord
+{
+public:
+	// adding a default constructor to enforce this to be a non POD class to successfully initialize
+	BlobRecord() {}
+	RecordID						SelfID;
+
+	static const BlobRecord			null;
+	uint32							ReplayCount = 0;
+
+	FString							Name;
+
+	RecordID						SourceID; // Source job or action or which has generated this Blob
+
+	HashType						HashValue = 0;
+	uint32							TexWidth = 0;
+	uint32							TexHeight = 0;
+
+	uint32							NumMapped = 0;			/// if this blob is mapped by other blobs, here is the number of mapping to this
+	RecordID						MappedID;				/// if this blob hash is remapped, here is the remapped to blob ID
+	uint64							RawBufferMemSize = 0;	/// The true size as informed
+
+	/// A tiled blob reference its sub tile blobs in the following fields _tileIdxs and _uniqueBlobs
+	/// The tiled information for the blob and the indices of the unique blobs which make the tiles grid 
+	Uint8Tiles						TileIdxs;
+	std::vector<uint16_t>			UniqueBlobs;
+
+	FORCEINLINE const TiledGrid&	Grid() const { return TileIdxs.Grid(); }
+	FORCEINLINE bool				IsTiled() const { return TileIdxs.Grid().IsTiled(); }
+	FORCEINLINE int32_t				NumTiles() const { return TileIdxs.Grid().NumTiles(); }
+	FORCEINLINE RecordID			GetTileBlob(uint8_t x, uint8_t y) const { return GetUniqueBlob(TileIdxs(x, y)); }
+
+	FORCEINLINE int32_t				NumUniqueBlobs() const { return UniqueBlobs.size(); }
+	FORCEINLINE RecordID			GetUniqueBlob(uint8_t i) const { return RecordID::fromBlob(UniqueBlobs[i]); }
+};
+using BlobRecordArray = std::vector< BlobRecord >;
+using RecordFromBlobIDMap = std::unordered_map<BlobID, RecordID>;
+
+
+/// Individual Job Record is the leaf of the data model
+/// It gather meta information to present what the job has done
+/// the inputs and output blob
+class TEXTUREGRAPHINSIGHT_API JobRecord
+{
+public:
+	// adding a default constructor to enforce this to be a non POD class to successfully initialize
+	JobRecord() {}
+	static const JobRecord	GNull;
+
+	RecordID				SelfID;
+
+	uint32					JobIdx = -1;	/// Index of the job in the batch
+	int32					PhaseIdx = 0;  /// Phase index of the Job within the job
+													/// If 0 it is THE job, < 0 a prior job, > 0 an after job
+
+	int						Priority = 0;
+	FString					TransformName;
+
+	HashType				JobHash = 0;
+
+	uint32					TexWidth = 0;
+	uint32					TexHeight = 0;
+
+	double					BeginTimeMS = 0;
+	double					BeginRunTimeMS = 0;
+	double					EndTimeMS = 0;
+
+	RecordID				ResultBlobId;
+	RecordIDArray			InputBlobIds;
+	FStringArray			InputArgNames;
+
+	// Information about the tiled job
+	BoolTiles				Tiles;
+	
+	FORCEINLINE uint64				GetNumPixels() const { return TexWidth * TexHeight; }
+	FORCEINLINE uint64				GetNumPixelsPerTile() const { return GetNumPixels() / GetNumTiles(); }
+
+	FORCEINLINE const TiledGrid&	GetGrid() const { return Tiles.Grid(); }
+	FORCEINLINE bool				IsTiled() const { return Tiles.Grid().IsTiled(); }
+	FORCEINLINE int32				GetNumTiles() const { return Tiles.Grid().NumTiles(); }
+
+	FORCEINLINE bool				IsTileInvalidated(uint8_t x, uint8_t y) const { return Tiles(x, y) != 0; }
+	FORCEINLINE int32				GetNumInvalidated() const { return TiledSum(Tiles); }
+
+	FORCEINLINE double				GetBeginTimeMS() const { return BeginTimeMS; }
+	FORCEINLINE double				GetBeginRunTimeMS() const { return (BeginRunTimeMS != 0 ? BeginRunTimeMS : BeginTimeMS); } // Ideal begin runtime, use begin TIme if none assigned
+	FORCEINLINE double				GetRunTimeMS() const { return (BeginRunTimeMS != 0 ? EndTimeMS - BeginRunTimeMS : EndTimeMS - BeginTimeMS); } // Runtime must be positive, clamp to 0
+	FORCEINLINE double				GetPreRunTimeMS() const { return GetBeginRunTimeMS() - BeginTimeMS; }
+
+	FORCEINLINE double				GetFillRate() const { return GetNumPixelsPerTile() * GetNumInvalidated() / (0.001 * GetRunTimeMS()); }
+
+	FORCEINLINE bool				IsNoOp() const { return bIsDone && (GetNumInvalidated() == 0); }
+
+	FORCEINLINE bool				IsMainPhase() const { return PhaseIdx == 0; }
+
+
+	bool					bIsDone = false;
+	uint32					ReplayCount = 0;
+};
+using JobRecordArray = std::vector< JobRecord >;
+
+
+/// Batch Record
+/// Hold the Job Records
+class TEXTUREGRAPHINSIGHT_API BatchRecord
+{
+public:
+	// adding a default constructor to enforce this to be a non POD class to successfully initialize
+	BatchRecord() {}
+
+	RecordID						SelfID;
+
+	static const BatchRecord null;
+
+	RecordID						MixID;
+
+	BatchID							BatchID = INVALID_RECORD_ID;
+	FrameID							FrameID = INVALID_RECORD_ID;
+	JobRecordArray					Jobs;
+
+	RecordIDArray					ResultBlobs;
+
+	FString							Action;
+
+	double							BeginTimeMS = 0;
+	double							EndTimeMS = 0;
+
+	uint16_t						NumTiles = 0;
+	uint16_t						NumInvalidatedTiles = 0;
+
+	uint32_t						ReplayCount = 0;
+
+	bool							bIsDone = false;
+	bool							bIsJobsDone = false;
+
+	bool							bIsNoCache = false;
+	bool							bIsFromIdle = false;
+
+	FORCEINLINE const JobRecord&	GetJob(RecordID ID) const { return ((ID.Job() < Jobs.size()) ? Jobs[ID.Job()] : JobRecord::GNull); }
+	FORCEINLINE double				ScopeTime_ms() const { return EndTimeMS - BeginTimeMS; }
+
+};
+using BatchRecordArray = std::vector< BatchRecord >;
+using RecordFromBatchIDMap = std::unordered_map<BatchID, RecordID>;
+
+
+/// Action Record
+/// Hold the Action Records
+class TEXTUREGRAPHINSIGHT_API ActionRecord
+{
+public:
+	// adding a default constructor to enforce this to be a non POD class to successfully initialize
+	ActionRecord() {}
+	static const ActionRecord null;
+
+	FString							Name;
+	FString							Meta;
+
+	mutable RecordID				ParentAction;
+	RecordIDArray					SubActions;
+};
+using ActionRecordArray = std::vector< ActionRecord >;
+using RecordFromActionIDMap = std::unordered_map<ActionID, RecordID>;
+
+
+/// Mix Record
+/// Hold the Mix Records
+class TEXTUREGRAPHINSIGHT_API MixRecord
+{
+public:
+	// adding a default constructor to enforce this to be a non POD class to successfully initialize
+	MixRecord() {}
+	RecordID						SelfID;
+
+	static const MixRecord null;
+
+	RecordID						ParentMixID; // If mix is an instance, it has a parentMix
+	bool							IsInstanceMix() const { return ParentMixID.IsValid(); }
+
+	RecordIDArray					InstanceMixIDs; // Or it may have instances
+	bool							HasInstances() const { return InstanceMixIDs.size() > 0; }
+
+	FString							Name;
+	FString							AssetName;
+
+	RecordIDArray					Batches;
+
+};
+using MixRecordArray = std::vector< MixRecord >;
+using RecordFromMixIDMap = std::unordered_map<MixID, RecordID>;
+
+
+
+/// Session Record container defining the namespace for all the recordIDs
+class TEXTUREGRAPHINSIGHT_API SessionRecord
+{
+private:
+	BatchRecordArray				Batches;
+	RecordFromBatchIDMap			RecordFromBatchID;
+
+	DeviceBufferRecordArray			DeviceBuffers[(uint32) DeviceType::Count];
+	RecordFromDeviceBufferIDMap		BufferFromDeviceBufferIDActive[(uint32)DeviceType::Count];	/// Store the current active (or leaked) device buffer in a per device map
+	RecordFromDeviceBufferIDMap		BufferFromDeviceBufferIDErased;  /// Once erased, store all the device buffer in this map
+	RecordFromHashMap				BufferFromHash;
+
+	BlobRecordArray					Blobs;
+	RecordFromBlobIDMap				BlobFromBlobID;
+	
+	ActionRecordArray				Actions;
+	RecordFromActionIDMap			ActionFromActionID;
+
+	MixRecordArray					Mixes;
+	RecordFromMixIDMap				RecordFromMixID;
+
+public:
+	double							StartTime_ms() const { return (Batches.empty() ? 0 : Batches.front().BeginTimeMS); }
+	
+	// Batch records
+	RecordID						FindBatchRecord(BatchID InBatchID) const;
+	using BatchRecordConstructor	= std::function<BatchRecord (RecordID RecId)>;
+	using BatchRecordEditor			= std::function<bool (BatchRecord& r)>;
+	RecordID						NewBatchRecord(BatchID InBatchID, BatchRecordConstructor Constructor);
+	bool							EditBatchRecord(RecordID RecId, BatchRecordEditor Editor);
+
+	// Device Buffer records
+	std::tuple<RecordID, bool>		NewDeviceBufferRecord(DeviceType DevType, DeviceBufferID BufferID, const DeviceBufferRecord& Record);
+	RecordID						EraseDeviceBufferRecord(DeviceType DevType, DeviceBufferID BufferID); /// Assuming the db is in the current map, then move its key in the erased map and tag it erased
+
+	/// Find the records from the array of bufferIDs & output the leaked buffers
+	RecordIDArray    				FindActiveDeviceBufferRecords(DeviceType DevType, const DeviceBufferIDArray& BufferIDs, RecordIDArray& LeakedBuffers);
+	
+	RecordID						FindDeviceBufferRecord(DeviceBufferID BufferID) const;
+	RecordIDArray					FindDeviceBufferRecords(const DeviceBufferIDArray& BufferIDs) const;
+	RecordIDArray					FetchActiveDeviceBufferIDs(DeviceType DevType) const;
+
+	RecordID						FindDeviceBufferRecordFromHash(HashType InHash) const;
+
+	// Blob records
+	RecordID						FindBlobRecord(BlobID InBlobID) const;
+	using BlobRecordConstructor		= std::function<BlobRecord(RecordID RecId)>;
+	using BlobRecordEditor			= std::function<bool(BlobRecord& r)>;
+	RecordID						NewBlobRecord(BlobID InBlobID, BlobRecordConstructor Constructor);
+	bool							EditBlobRecord(RecordID RecId, BlobRecordEditor Editor);
+
+	RecordID						NewActionRecord(ActionID InActionID, const ActionRecord& Record);
+	RecordID						FindActionRecord(ActionID InActionID) const;
+
+	// Mix records
+	RecordID						FindMixRecord(MixID MixID) const;
+	using MixRecordConstructor		= std::function<MixRecord(RecordID RecId)>;
+	using MixRecordEditor			= std::function<bool(MixRecord& r)>;
+	RecordID						NewMixRecord(BatchID InBatchID, MixRecordConstructor constructor);
+	bool							EditMixRecord(RecordID RecId, MixRecordEditor editor);
+
+	FORCEINLINE const DeviceBufferRecord&	GetBuffer(RecordID RecId) const { return ((RecId.Buffer() < DeviceBuffers[RecId.Buffer_DeviceType()].size()) ? DeviceBuffers[RecId.Buffer_DeviceType()][RecId.Buffer()] : DeviceBufferRecord::null); }
+	FORCEINLINE const BlobRecord&			GetBlob(RecordID RecId) const { return ((RecId.Blob() < Blobs.size()) ? Blobs[RecId.Blob()] : BlobRecord::null); }
+	FORCEINLINE const BatchRecord&			GetBatch(RecordID RecId) const { return ((RecId.Batch() < Batches.size()) ? Batches[RecId.Batch()] : BatchRecord::null); }
+	FORCEINLINE const JobRecord&			GetJob(RecordID RecId) const { return GetBatch(RecId).GetJob(RecId); }
+	FORCEINLINE const ActionRecord&			GetAction(RecordID RecId) const { return ((RecId.Action() < Actions.size()) ? Actions[RecId.Action()] : ActionRecord::null); }
+	FORCEINLINE const MixRecord&			GetMix(RecordID RecId) const { return ((RecId.Mix() < Mixes.size()) ? Mixes[RecId.Mix()] : MixRecord::null); }
+
+};
+
+using SessionRecordPtr = std::shared_ptr<SessionRecord>;
+
diff --git a/Engine/Plugins/Experimental/TextureGraph/Source/TextureGraphInsight/Public/Model/TextureGraphInsightSession.h b/Engine/Plugins/Experimental/TextureGraph/Source/TextureGraphInsight/Public/Model/TextureGraphInsightSession.h
new file mode 100644
--- /dev/null
+++ b/Engine/Plugins/Experimental/TextureGraph/Source/TextureGraphInsight/Public/Model/TextureGraphInsightSession.h
@@ -0,0 +1,170 @@
+// Copyright Epic Games, Inc. All Rights Reserved.
+#pragma once
+
+#include "TextureGraphInsightRecord.h"
+#include "TextureGraphInsightObserver.h"
+
+class UMixInterface;
+using MixPtr = UMixInterface*;
+using MixPtrW = UMixInterface*;
+//using MixPtr = std::shared_ptr<UMixInterface>;
+//using MixPtrW = std::weak_ptr<UMixInterface>;
+
+class TEXTUREGRAPHINSIGHT_API TextureGraphInsightSession
+{
+public:
+	/// Keep all the pointers to "real" Mixer Engine objects in the live runtime cache
+	class LiveCache
+	{
+	public:
+		LiveCache(TextureGraphInsightSession* InSession) : Session(InSession) {}
+		TextureGraphInsightSession*		Session = nullptr;
+
+		// TODO: Find a better way to keep the Batch pointer alive, for now, only the last one will still be valid
+		std::vector<JobBatchPtrW>	Batches; 
+		using						BlobEntry= HashType;
+		std::vector<BlobEntry>		Blobs;
+		std::vector<MixPtrW>		Mixes;
+
+		void						AddBlob(RecordID RecId, BlobPtr BlobObj);
+		BlobPtr						GetBlob(RecordID RecId) const;
+
+		void						AddBatch(RecordID RecId, JobBatchPtr Batch);
+		JobBatchPtr					GetBatch(RecordID RecId) const;
+
+		uint64						DeviceCacheVersions[(uint32)DeviceType::Count] = { 0, 0, 0, 0, 0, 0 , 0};
+
+		DeviceBufferPtr				GetDeviceBuffer(RecordID RecId) const;
+
+		void						AddMix(RecordID RecId, MixPtr MixObj);
+		MixPtr						GetMix(RecordID RecId) const;
+	};
+
+private:
+protected:
+	friend class TextureGraphInsightEngineObserver;
+	/// Protected interface of emitters called by the engine observer
+	void							EngineCreated();
+	void							EngineDestroyed();
+	void							EmitOnEngineReset(bool);
+
+	/// manage DeviceBuffer update
+	friend class TextureGraphInsightDeviceObserver;
+	uint32							ParseDeviceBuffersCache(DeviceType type, RecordIDArray& addedBufferIds);
+	void							DeviceBuffersUpdated(DeviceType DevType, TextureGraphInsightDeviceObserver::HashNDescArray&& AddedBuffers, TextureGraphInsightDeviceObserver::HashArray&& RemovedBuffers);
+
+	/// Manage Blobber update
+	friend class TextureGraphInsightBlobberObserver;
+	void							BlobberHashesUpdated(TextureGraphInsightBlobberObserver::HashArray&& AddedHashes, TextureGraphInsightBlobberObserver::HashArray&& RemappedHashes);
+
+	friend class TextureGraphInsightSchedulerObserver;
+	/// Protected interface of emitters called by the scheduler observer
+	void							UpdateIdle();
+	void							BatchAdded(JobBatchPtr Batch);
+	void							BatchDone(JobBatchPtr Batch);
+	void							BatchJobsDone(JobBatchPtr Batch);
+
+	/// manage creation and update of the Batch job records
+	BatchRecord						MakeBatchRecord(RecordID RecId, JobBatchPtr Batch, bool isDone = true);
+	void							UpdateBatchRecord(JobBatchPtr Batch, BatchRecord& record);
+	JobRecord						MakeJobRecord(JobPtrW job, uint32 idx, int32 phaseIdx = 0);
+	void							UpdateJobRecord(JobPtrW job, JobRecord& record, double batchBeginTime_ms, double batchEndTime_ms);
+
+
+	/// manage creation and update of the BlobObj records
+	RecordID						FindOrCreateBlobRecordFromHash(HashType hash);
+	RecordID						RemapBlobHash(HashType oldHash, HashType newHash);
+	RecordID						FindOrCreateBlobRecord(BlobPtr BlobObj, RecordID sourceID = RecordID());
+	BlobRecord						MakeBlobRecord(RecordID RecId, BlobPtr BlobObj, RecordID sourceID = RecordID());
+	void							EditTiledBlobRecord(BlobRecord& rd, const RecordIDTiles& subTiles);
+	void							UpdateBlobRecord(RecordID RecId, BlobPtr BlobObj);
+	RecordIDArray					NewBlobsToNotify;
+	RecordIDArray					MappedBlobsToNotify;
+	void							EmitOnNewBlob();
+	void							EmitOnMappedBlob();
+
+	/// manage creation and update of the MixObj records
+	RecordID						AssociateBatchToMix(RecordID Batch, MixPtr MixObj);
+	RecordID						FindOrCreateMixRecord(MixPtr MixObj);
+
+	/// manage creation and update of the buffer records
+	RecordIDArray					NewBuffersToNotify;
+
+	/// The session record, collecting all the VALUES
+	/// A session record is renewed for each engine
+	SessionRecordPtr				Record;
+	std::mutex						RecordMutex;
+
+	/// Keep all the pointers to "real" Mixer Engine objects in the live cache
+	/// A session live cache is renewed for each engine
+	std::shared_ptr<LiveCache>		Cache;
+
+	int32							EngineID = -1;
+public:
+
+	TextureGraphInsightSession();
+    ~TextureGraphInsightSession();
+
+	FORCEINLINE SessionRecord&		GetRecord() { std::lock_guard<std::mutex> Lock(RecordMutex); return (*Record); }
+	FORCEINLINE LiveCache&			GetCache() { std::lock_guard<std::mutex> Lock(RecordMutex); return (*Cache); }
+
+	/// Send the specified record to be inspected
+	void							SendToInspector(RecordID RecId) const;
+
+	/// Notifier subscription
+	DECLARE_MULTICAST_DELEGATE_OneParam(FOnEngineReset, int32);
+	FOnEngineReset& OnEngineReset()		{ return OnEngineResetEvent; }
+
+	DECLARE_MULTICAST_DELEGATE(FOnUpdateIdle);
+	FOnUpdateIdle& OnUpdateIdle()		{ return OnUpdateIdleEvent; }
+
+	DECLARE_MULTICAST_DELEGATE_OneParam(FOnRecordIDArray, const RecordIDArray&);
+
+	FOnRecordIDArray& OnDeviceBufferAdded()		{ return OnDeviceBufferAddedEvent; }
+	FOnRecordIDArray& OnDeviceBufferRemoved()	{ return OnDeviceBufferRemovedEvent; }
+	FOnRecordIDArray& OnBlobAdded()				{ return OnBlobAddedEvent; }
+	FOnRecordIDArray& OnBlobMapped()			{ return OnBlobMappedEvent; }
+
+	DECLARE_MULTICAST_DELEGATE_OneParam(FOnRecordID, RecordID);
+
+	FOnRecordID& OnBatchAdded()			{ return OnBatchAddedEvent; }
+	FOnRecordID& OnBatchDone()			{ return OnBatchDoneEvent; }
+	FOnRecordID& OnBatchJobsDone()		{ return OnBatchJobsDoneEvent; }
+
+	FOnRecordID& OnMixAdded()			{ return OnMixAddedEvent; }
+	FOnRecordID& OnMixUpdated()			{ return OnMixUpdatedEvent; }
+	
+	FOnRecordID& OnBatchInspected()		{ return OnBatchInspectedEvent; }
+	FOnRecordID& OnJobInspected()		{ return OnJobInspectedEvent; }
+	FOnRecordID& OnBlobInspected()		{ return OnBlobInspectedEvent; }
+	FOnRecordID& OnBufferInspected()	{ return OnBufferInspectedEvent; }
+
+
+	// Trigger Replay of Batch and jobs
+	// this call is async and returns if the scheduling of the replay is successful or not
+	bool							ReplayBatch(RecordID batchRecordID, bool captureRenderDoc);
+	bool							ReplayJob(RecordID jobRecordID, bool captureRenderDoc);
+
+protected:
+	FOnEngineReset		OnEngineResetEvent;
+	FOnUpdateIdle		OnUpdateIdleEvent;
+
+	FOnRecordIDArray	OnDeviceBufferAddedEvent;
+	FOnRecordIDArray	OnDeviceBufferRemovedEvent;
+
+	FOnRecordIDArray	OnBlobAddedEvent;
+	FOnRecordIDArray	OnBlobMappedEvent;
+
+	FOnRecordID			OnBatchAddedEvent;
+	FOnRecordID			OnBatchDoneEvent;
+	FOnRecordID			OnBatchJobsDoneEvent;
+
+	FOnRecordID			OnMixAddedEvent;
+	FOnRecordID			OnMixUpdatedEvent;
+
+	FOnRecordID			OnBatchInspectedEvent;
+	FOnRecordID			OnJobInspectedEvent;
+	FOnRecordID			OnBlobInspectedEvent;
+	FOnRecordID			OnBufferInspectedEvent;
+};
+
diff --git a/Engine/Plugins/Experimental/TextureGraph/Source/TextureGraphInsight/Public/TextureGraphInsight.h b/Engine/Plugins/Experimental/TextureGraph/Source/TextureGraphInsight/Public/TextureGraphInsight.h
new file mode 100644
--- /dev/null
+++ b/Engine/Plugins/Experimental/TextureGraph/Source/TextureGraphInsight/Public/TextureGraphInsight.h
@@ -0,0 +1,47 @@
+// Copyright Epic Games, Inc. All Rights Reserved.
+#pragma once
+
+#include "CoreMinimal.h"
+#include <memory>
+#include "Modules/ModuleManager.h"
+
+DECLARE_LOG_CATEGORY_EXTERN(LogTextureGraphInsight, Log, All);
+
+class TextureGraphInsightSchedulerObserver;
+class TextureGraphInsightSession;
+using TextureGraphInsightSessionPtr = std::shared_ptr<TextureGraphInsightSession>;
+
+class FTextureGraphInsightModule : public FDefaultGameModuleImpl
+{
+public:
+	virtual void StartupModule() override;
+	virtual void ShutdownModule() override;
+};
+
+struct RecordID;
+
+class TEXTUREGRAPHINSIGHT_API TextureGraphInsight
+{
+private:
+	static TextureGraphInsight* GInstance;					/// Static instance. Use Create to create a new instance. Destroy the current instance first
+	TextureGraphInsight();
+	~TextureGraphInsight();
+
+	TextureGraphInsightSessionPtr					Session;
+public:
+	static bool								Create();  /// Create the Insight Instance ONLY if no other currently created AND if a TextureGraphEngine is already created
+
+	/// Destroy the current instance
+	static bool								Destroy(); /// Destroy the instance of Insight if it exists.
+
+	//////////////////////////////////////////////////////////////////////////
+	/// Inline functions
+	//////////////////////////////////////////////////////////////////////////
+	FORCEINLINE static TextureGraphInsight*		Instance() { return GInstance; }
+
+	FORCEINLINE TextureGraphInsightSessionPtr		GetSession() const { return Session; }
+};
+
+// Macro to control in one place how we display hash value
+//#define HashToFString(h) FString::Printf(TEXT("%20llu"), (h))
+#define HashToFString(h) FString::Printf(TEXT("%llX"), (h))
diff --git a/Engine/Plugins/Experimental/TextureGraph/Source/TextureGraphInsight/Public/View/STextureGraphInsightActionView.h b/Engine/Plugins/Experimental/TextureGraph/Source/TextureGraphInsight/Public/View/STextureGraphInsightActionView.h
new file mode 100644
--- /dev/null
+++ b/Engine/Plugins/Experimental/TextureGraph/Source/TextureGraphInsight/Public/View/STextureGraphInsightActionView.h
@@ -0,0 +1,47 @@
+// Copyright Epic Games, Inc. All Rights Reserved.
+#pragma once
+
+//#include "SlateBasics.h"
+
+#include "Model/TextureGraphInsightSession.h"
+#include <Widgets/Views/STreeView.h>
+
+class STextureGraphInsightActionViewRow; /// Declare the concrete type of widget used for the raws of the view. Defined in the cpp file
+
+class TEXTUREGRAPHINSIGHT_API STextureGraphInsightActionView : public SCompoundWidget
+{
+public:
+	SLATE_BEGIN_ARGS(STextureGraphInsightActionView) {}
+	SLATE_END_ARGS()
+
+	void Construct(const FArguments& Args);
+
+	// Item Types
+	class FItemData;
+	using FItem = TSharedPtr<FItemData>;
+	using FItemArray = TArray< FItem >;
+	class FItemData
+	{
+	public:
+		FItemData(RecordID rid) : _recordID(rid) {}
+		TSharedPtr < STextureGraphInsightActionViewRow > _widget;
+		RecordID _recordID;
+		FItemArray _children;
+	};
+	using SItemTableView = STreeView< FItem >;
+
+	// Standard delegates for the view
+	TSharedRef<ITableRow>	OnGenerateRowForView(FItem item, const TSharedRef<STableViewBase>& OwnerTable);
+	FORCEINLINE void		OnGetChildrenForView(FItem item, FItemArray& children) { children = item->_children; };
+	void					OnClickItemForView(FItem item);
+	void					OnDoubleClickItemForView(FItem item);
+
+	/// The list of root items
+	FItemArray _rootItems;
+
+	// The TreeView widget
+	TSharedPtr<SItemTableView> _tableView;
+
+	void OnActionNew(RecordID rid);
+	void OnEngineReset(int);
+};
diff --git a/Engine/Plugins/Experimental/TextureGraph/Source/TextureGraphInsight/Public/View/STextureGraphInsightBlobView.h b/Engine/Plugins/Experimental/TextureGraph/Source/TextureGraphInsight/Public/View/STextureGraphInsightBlobView.h
new file mode 100644
--- /dev/null
+++ b/Engine/Plugins/Experimental/TextureGraph/Source/TextureGraphInsight/Public/View/STextureGraphInsightBlobView.h
@@ -0,0 +1,54 @@
+// Copyright Epic Games, Inc. All Rights Reserved.
+#pragma once
+
+//#include "SlateBasics.h"
+
+#include "Model/TextureGraphInsightRecord.h"
+#include <Widgets/Layout/SBorder.h>
+#include <UObject/GCObject.h>
+
+class STextureGraphInsightDeviceBufferView;
+class UMaterialInstanceDynamic;
+struct FSlateBrush;
+struct FSlateColorBrush;
+class TEXTUREGRAPHINSIGHT_API STextureGraphInsightBlobView : public SBorder, public FGCObject /// Need FGCObject to control garbage collection of objects
+{
+public:
+	SLATE_BEGIN_ARGS(STextureGraphInsightBlobView) :
+		_withHighlight(false) {}
+		SLATE_ARGUMENT(RecordID, recordID)
+		SLATE_ARGUMENT(bool, withHighlight)
+
+		SLATE_ARGUMENT(BoolTiles, tilesMask)
+	SLATE_END_ARGS()
+
+	void Construct(const FArguments& Args);
+
+	RecordID _recordID;
+	RecordID _inspectedID;
+
+	// Needs our own brush to display the blob content
+	TSharedPtr<FSlateBrush> _brush;
+	TObjectPtr<UMaterialInstanceDynamic> _brushMaterial = nullptr; /// 
+
+	// Border brush
+	TSharedPtr<FSlateColorBrush> _borderBrush;
+
+	FReply OnInspect();
+
+	// Capture mouse clicks here
+	virtual FReply OnMouseButtonDown(const FGeometry& MyGeometry, const FPointerEvent& MouseEvent) override;
+	virtual FReply OnMouseButtonDoubleClick(const FGeometry& InMyGeometry, const FPointerEvent& InMouseEvent) override;
+
+
+	virtual void AddReferencedObjects(FReferenceCollector& Collector) override
+	{
+		Collector.AddReferencedObject(_brushMaterial);
+	}
+
+	virtual FString GetReferencerName() const override
+	{
+		return TEXT("STextureGraphInsightBlobView");
+	}
+
+};
diff --git a/Engine/Plugins/Experimental/TextureGraph/Source/TextureGraphInsight/Public/View/STextureGraphInsightDeviceBufferView.h b/Engine/Plugins/Experimental/TextureGraph/Source/TextureGraphInsight/Public/View/STextureGraphInsightDeviceBufferView.h
new file mode 100644
--- /dev/null
+++ b/Engine/Plugins/Experimental/TextureGraph/Source/TextureGraphInsight/Public/View/STextureGraphInsightDeviceBufferView.h
@@ -0,0 +1,48 @@
+// Copyright Epic Games, Inc. All Rights Reserved.
+#pragma once
+
+//#include "SlateBasics.h"
+
+#include "Model/TextureGraphInsightRecord.h"
+#include <Widgets/Layout/SBorder.h>
+#include <UObject/GCObject.h>
+
+class UMaterialInstanceDynamic;
+
+class TEXTUREGRAPHINSIGHT_API STextureGraphInsightDeviceBufferView : public SBorder, public FGCObject /// Need FGCObject to control garbage collection of objects
+{
+public:
+	SLATE_BEGIN_ARGS(STextureGraphInsightDeviceBufferView) :
+		_withToolTip(false),
+		_withDescription(false) {}
+		SLATE_ARGUMENT(RecordID, recordID)
+		SLATE_ARGUMENT(RecordID, blobID)
+		SLATE_ARGUMENT(bool, withToolTip)
+		SLATE_ARGUMENT(bool, withDescription)
+	SLATE_END_ARGS()
+
+	void Construct(const FArguments& Args);
+
+	RecordID _recordID;
+	RecordID _blobID;
+	uint32 _imageWidth = 1;
+	uint32 _imageHeight = 1;
+	FString _imageName;
+
+	// Needs our own brush to display the blob content
+	TSharedPtr<FSlateBrush> _brush;
+	TObjectPtr<UMaterialInstanceDynamic> _brushMaterial = nullptr;
+
+	void CreateWidgetBrush(UTexture* blobTexture);
+	void MakeWidgetMaterial();
+
+	virtual void AddReferencedObjects(FReferenceCollector& Collector) override
+	{
+		Collector.AddReferencedObject(_brushMaterial);
+	}
+
+	virtual FString GetReferencerName() const override
+	{
+		return TEXT("STextureGraphInsightDeviceBufferView");
+	}
+};
diff --git a/Engine/Plugins/Experimental/TextureGraph/Source/TextureGraphInsight/Public/View/STextureGraphInsightDeviceView.h b/Engine/Plugins/Experimental/TextureGraph/Source/TextureGraphInsight/Public/View/STextureGraphInsightDeviceView.h
new file mode 100644
--- /dev/null
+++ b/Engine/Plugins/Experimental/TextureGraph/Source/TextureGraphInsight/Public/View/STextureGraphInsightDeviceView.h
@@ -0,0 +1,113 @@
+// Copyright Epic Games, Inc. All Rights Reserved.
+#pragma once
+
+//#include "SlateBasics.h"
+
+#include "Model/TextureGraphInsightSession.h"
+#include <Widgets/SCompoundWidget.h>
+#include <Widgets/Views/STreeView.h>
+
+class STextureGraphInsightDeviceListViewRow; /// Declare the concrete type of widget used for the raws of the view. Defined in the cpp file
+class STableViewBase;
+class ITableRow;
+class STextBlock;
+
+class TEXTUREGRAPHINSIGHT_API STextureGraphInsightDeviceListView : public SCompoundWidget
+{
+public:
+	SLATE_BEGIN_ARGS(STextureGraphInsightDeviceListView)
+	: _inspectOnSimpleClick(false), _inspectDevices(true) {}
+	SLATE_ARGUMENT(RecordID, recordID) /// Can use an empty RecordID with type device and the expected Device type to only care about that device
+	SLATE_ARGUMENT(bool, inspectOnSimpleClick)
+	SLATE_ARGUMENT(bool, inspectDevices)
+	SLATE_END_ARGS()
+
+	void Construct(const FArguments& Args);
+
+	// TreeView Item Types
+	class FItemData;
+	using FItem = TSharedPtr<FItemData>;
+	using FItemW = TWeakPtr<FItemData>;
+	using FItemArray = TArray< FItem >;
+	class FItemData
+	{
+	public:
+		FItemData(RecordID rid, int32 metaItem = 0) : _recordID(rid), _metaItem(metaItem) {}
+		TSharedPtr < STextureGraphInsightDeviceListViewRow > _widget;
+		RecordID	_recordID;
+		FItemArray	_children;
+		FItem		_removedRoot;
+		const int32 _metaItem = 0;
+		int32		_rank = 0;
+	};
+	using SItemTableView = STreeView< FItem >;
+
+	// Standard delegates for the view
+	TSharedRef<ITableRow>	OnGenerateRowForView(FItem item, const TSharedRef<STableViewBase>& OwnerTable);
+	FORCEINLINE void		OnGetChildrenForView(FItem item, FItemArray& children) { children = item->_children; if (item->_removedRoot) children.Add(item->_removedRoot); };
+	void OnClickItemForView(FItem item);
+	void OnDoubleClickItemForView(FItem item);
+	void OnInspect(FItem item);
+	
+	RecordID	_rootDevice; /// Can use an empty RecordID parameter with type DeviceBuffer and the expected Device type to only care about that device
+
+	// Inspect on double click always
+	// eventually trigger inspect on simple click
+	bool _inspectOnSimpleClick = false;
+	
+	bool _inspectDevices = true;
+
+	/// The list of root items
+	FItemArray _rootItems;
+
+	// The TreeView widget
+	TSharedPtr<SItemTableView> _tableView;
+
+	void OnDeviceBufferAdded(const RecordIDArray& rids);
+	void OnDeviceBufferRemoved(const RecordIDArray& rids);
+	void OnEngineReset(int);
+
+	void Reset();
+
+	void RefreshRootItems();
+};
+
+inline bool operator ==(const STextureGraphInsightDeviceListView::FItem& i,  const RecordID& rid) { return i->_recordID.Buffer() == rid.Buffer(); }
+
+class TEXTUREGRAPHINSIGHT_API STextureGraphInsightDeviceView : public SCompoundWidget
+{
+	SLATE_BEGIN_ARGS(STextureGraphInsightDeviceView) {}
+	SLATE_ARGUMENT(DeviceType, deviceType)
+	SLATE_END_ARGS()
+
+	void Construct(const FArguments& Args);
+
+	DeviceType _deviceType;
+	
+	TSharedPtr<STextBlock> _mem_percent;
+	TSharedPtr<STextBlock> _mem_number;
+	TSharedPtr<STextBlock> _mem_budget;
+
+	TSharedPtr<STextBlock> _buffer_percent;
+	TSharedPtr<STextBlock> _buffer_number;
+	TSharedPtr<STextBlock> _buffer_budget;
+
+	TSharedPtr<STextBlock> _threadCount;
+	
+	void Refresh();
+};
+
+class TEXTUREGRAPHINSIGHT_API STextureGraphInsightDeviceManagerView : public SCompoundWidget
+{
+	SLATE_BEGIN_ARGS(STextureGraphInsightDeviceManagerView) {}
+	SLATE_END_ARGS()
+
+	void Construct(const FArguments& Args);
+
+	std::vector < TSharedPtr< STextureGraphInsightDeviceView >> _deviceViews;
+
+	void Reset(); // Rebuild the full ui of this widget, called by update
+	void OnUpdate();
+	void OnEngineReset(int32);
+};
+
diff --git a/Engine/Plugins/Experimental/TextureGraph/Source/TextureGraphInsight/Public/View/STextureGraphInsightInspectorView.h b/Engine/Plugins/Experimental/TextureGraph/Source/TextureGraphInsight/Public/View/STextureGraphInsightInspectorView.h
new file mode 100644
--- /dev/null
+++ b/Engine/Plugins/Experimental/TextureGraph/Source/TextureGraphInsight/Public/View/STextureGraphInsightInspectorView.h
@@ -0,0 +1,138 @@
+// Copyright Epic Games, Inc. All Rights Reserved.
+#pragma once
+
+//#include "SlateBasics.h"
+#include "Model/TextureGraphInsightSession.h"
+#include <Widgets/Views/STileView.h>
+
+class STextureGraphInsightDeviceBufferView;
+class STextureGraphInsightBlobView;
+class STextureGraphInsightRecordTrailView;
+
+class TEXTUREGRAPHINSIGHT_API STextureGraphInsightBatchInspectorView : public SCompoundWidget
+{
+public:
+	SLATE_BEGIN_ARGS(STextureGraphInsightBatchInspectorView) {}
+	SLATE_ARGUMENT(RecordID, recordID)
+		SLATE_END_ARGS()
+
+		void Construct(const FArguments& Args);
+
+	// ListView Item Types presenting various results of a batch
+	class FItemData;
+	using FItem = TSharedPtr<FItemData>;
+	using FItemArray = TArray< FItem >;
+	class FItemData
+	{
+	public:
+		FItemData(RecordID batchID) : _batchID(batchID) {}
+		FItemData(RecordID batchID, int32 targetIdx) : _batchID(batchID), _targetIndex(targetIdx) {}
+
+		RecordID	_batchID;
+		int32		_targetIndex = -1; // -1 meaning it is not an input
+	};
+	using SItemTileView = STileView< FItem >;
+
+	// Standard delegates for the tree view
+	TSharedRef<ITableRow> OnGenerateTileForView(FItem item, const TSharedRef<STableViewBase>& OwnerTable);
+
+	/// The list of root items
+	FItemArray _rootItems;
+	// The TreeView widget
+	TSharedPtr<SItemTileView> _view;
+
+
+	// Inspect the specified batch (Called by Constructor with recordID param)
+	void InspectBatch(RecordID job);
+};
+
+class TEXTUREGRAPHINSIGHT_API STextureGraphInsightJobInspectorView : public SCompoundWidget
+{
+public:
+	SLATE_BEGIN_ARGS(STextureGraphInsightJobInspectorView) {}
+	SLATE_ARGUMENT(RecordID, recordID)
+	SLATE_END_ARGS()
+
+	void Construct(const FArguments& Args);
+
+	// ListView Item Types presenting various input of a job
+	class FItemData;
+	using FItem = TSharedPtr<FItemData>;
+	using FItemArray = TArray< FItem >;
+	class FItemData
+	{
+	public:
+		FItemData(RecordID jobID) : _jobID(jobID) {}
+		FItemData(RecordID jobID, BoolTiles mask) : _jobID(jobID), _tilesMask(mask) {}
+		FItemData(RecordID jobID, int32 inputIdx) :  _jobID(jobID), _inputIndex(inputIdx){}
+		
+		RecordID	_jobID;
+		BoolTiles	_tilesMask;
+		int32		_inputIndex = -1; // -1 meaning it is not an input
+	};
+	using SItemTileView = STileView< FItem >;
+
+	// Standard delegates for the tree view
+	TSharedRef<ITableRow> OnGenerateTileForView(FItem item, const TSharedRef<STableViewBase>& OwnerTable);
+
+	// The output blob view generated by the job
+	TSharedPtr<STextureGraphInsightBlobView> _outBlobView;
+
+	/// The list of root items
+	FItemArray _rootItems;
+	// The TreeView widget
+	TSharedPtr<SItemTileView> _view;
+
+
+	// Inspect the specified Job (Called by Constructor with recordID param)
+	void InspectJob(RecordID job);
+};
+
+
+class TEXTUREGRAPHINSIGHT_API STextureGraphInsightBlobInspectorView : public SCompoundWidget
+{
+public:
+	SLATE_BEGIN_ARGS(STextureGraphInsightBlobInspectorView) {}
+	SLATE_ARGUMENT(RecordID, recordID)
+	SLATE_END_ARGS()
+
+	void Construct(const FArguments& Args);
+
+	// The main blob view generated by the job
+	TSharedPtr<STextureGraphInsightBlobView> _blobView;
+
+	// Inspect the specified Blob (Called by Constructor with recordID param)
+	void InspectBlob(RecordID blob);
+};
+
+class TEXTUREGRAPHINSIGHT_API STextureGraphInsightDeviceBufferInspectorView : public SCompoundWidget
+{
+public:
+	SLATE_BEGIN_ARGS(STextureGraphInsightDeviceBufferInspectorView) {}
+	SLATE_ARGUMENT(RecordID, recordID)
+	SLATE_END_ARGS()
+
+	void Construct(const FArguments& Args);
+
+
+	// Inspect the specified buffer (Called by Constructor with recordID param)
+	void InspectBuffer(RecordID buffer);
+};
+
+class TEXTUREGRAPHINSIGHT_API STextureGraphInsightInspectorView : public SCompoundWidget
+{
+public:
+	SLATE_BEGIN_ARGS(STextureGraphInsightInspectorView) {}
+	SLATE_END_ARGS()
+
+	void Construct(const FArguments& Args);
+
+	TSharedPtr<STextureGraphInsightRecordTrailView> _recordTrail;
+
+	TSharedPtr<SOverlay> _stack;
+
+
+
+	void OnInspected(RecordID rid);
+	void OnEngineReset(int32);
+};
\ No newline at end of file
diff --git a/Engine/Plugins/Experimental/TextureGraph/Source/TextureGraphInsight/Public/View/STextureGraphInsightMixView.h b/Engine/Plugins/Experimental/TextureGraph/Source/TextureGraphInsight/Public/View/STextureGraphInsightMixView.h
new file mode 100644
--- /dev/null
+++ b/Engine/Plugins/Experimental/TextureGraph/Source/TextureGraphInsight/Public/View/STextureGraphInsightMixView.h
@@ -0,0 +1,56 @@
+// Copyright Epic Games, Inc. All Rights Reserved.
+#pragma once
+
+//#include "SlateBasics.h"
+
+#include "Model/TextureGraphInsightSession.h"
+#include "Widgets/Views/STreeView.h"
+#include "Widgets/SCompoundWidget.h"
+class STextureGraphInsightMixListViewRow; /// Declare the concrete type of widget used for the raws of the view. Defined in the cpp file
+class STableViewBase;
+class ITableRow;
+
+class TEXTUREGRAPHINSIGHT_API STextureGraphInsightMixListView : public SCompoundWidget
+{
+public:
+	SLATE_BEGIN_ARGS(STextureGraphInsightMixListView) {}
+	SLATE_END_ARGS()
+
+		void Construct(const FArguments& Args);
+
+	// TreeView Item Types
+	class FItemData;
+	using FItem = TSharedPtr<FItemData>;
+	using FItemArray = TArray< FItem >;
+	class FItemData
+	{
+	public:
+		FItemData(RecordID rid) : _recordID(rid) {}
+		TSharedPtr < STextureGraphInsightMixListViewRow > _widget;
+		RecordID _recordID;
+		FItemArray	_children;
+	};
+	//using SItemTableView = SListView< FItem >;
+	using SItemTableView = STreeView< FItem >;
+
+
+	// Standard delegates for the view
+	TSharedRef<ITableRow>	OnGenerateRowForView(FItem item, const TSharedRef<STableViewBase>& OwnerTable);
+	FORCEINLINE void		OnGetChildrenForView(FItem item, FItemArray& children) { children = item->_children; };
+	void					OnClickItemForView(FItem item);
+	void					OnDoubleClickItemForView(FItem item);
+
+	/// The list of root items
+	FItemArray _rootItems;
+
+	// The TreeView widget
+	TSharedPtr<SItemTableView> _tableView;
+
+	void RefreshRootItems();
+
+	void OnMixNew(RecordID rid);
+	void OnMixUpdate(RecordID rid);
+	void OnEngineReset(int id);
+
+	void UpdateItemData(FItemData& item);
+};
\ No newline at end of file
diff --git a/Engine/Plugins/Experimental/TextureGraph/Source/TextureGraphInsight/Public/View/STextureGraphInsightRecordTrailView.h b/Engine/Plugins/Experimental/TextureGraph/Source/TextureGraphInsight/Public/View/STextureGraphInsightRecordTrailView.h
new file mode 100644
--- /dev/null
+++ b/Engine/Plugins/Experimental/TextureGraph/Source/TextureGraphInsight/Public/View/STextureGraphInsightRecordTrailView.h
@@ -0,0 +1,27 @@
+// Copyright Epic Games, Inc. All Rights Reserved.
+#pragma once
+
+//#include "SlateBasics.h"
+#include "Widgets/Navigation/SBreadcrumbTrail.h"
+
+#include "Model/TextureGraphInsightSession.h"
+
+
+class TEXTUREGRAPHINSIGHT_API STextureGraphInsightRecordTrailView : public SCompoundWidget
+{
+	TSharedPtr< SBreadcrumbTrail<RecordID> > _trail;
+public:
+	SLATE_BEGIN_ARGS(STextureGraphInsightRecordTrailView) {}
+		SLATE_ARGUMENT(RecordID, recordID)
+	SLATE_END_ARGS()
+
+	void Construct(const FArguments& Args);
+
+	void Refresh(RecordID rid);
+
+	static RecordID FindRootRecord(RecordID rid);
+
+	void OnCrumbClicked(const RecordID& blob);
+};
+
+
diff --git a/Engine/Plugins/Experimental/TextureGraph/Source/TextureGraphInsight/Public/View/STextureGraphInsightResourceView.h b/Engine/Plugins/Experimental/TextureGraph/Source/TextureGraphInsight/Public/View/STextureGraphInsightResourceView.h
new file mode 100644
--- /dev/null
+++ b/Engine/Plugins/Experimental/TextureGraph/Source/TextureGraphInsight/Public/View/STextureGraphInsightResourceView.h
@@ -0,0 +1,54 @@
+// Copyright Epic Games, Inc. All Rights Reserved.
+#pragma once
+
+//#include "SlateBasics.h"
+
+#include "Model/TextureGraphInsightSession.h"
+#include "Widgets/SCompoundWidget.h"
+#include "Widgets/Views/STreeView.h"
+class STextureGraphInsightResourceViewRow; /// Declare the concrete type of widget used for the raws of the view. Defined in the cpp file
+class ITableRow;
+class STableViewBase;
+
+class TEXTUREGRAPHINSIGHT_API STextureGraphInsightResourceView : public SCompoundWidget
+{
+public:
+	SLATE_BEGIN_ARGS(STextureGraphInsightResourceView) {}
+	SLATE_END_ARGS()
+
+	void Construct(const FArguments& Args);
+
+	// TreeView Item Types
+	class FItemData;
+	using FItem = TSharedPtr<FItemData>;
+	using FItemArray = TArray< FItem >;
+	class FItemData
+	{
+	public:
+		FItemData(RecordID rid) : _recordID(rid) {}
+		TSharedPtr < STextureGraphInsightResourceViewRow > _widget;
+		RecordID _recordID;
+		FItemArray	_children;
+	};
+	//using SItemTableView = SListView< FItem >;
+	using SItemTableView = STreeView< FItem >;
+
+
+	// Standard delegates for the view
+	TSharedRef<ITableRow>	OnGenerateRowForView(FItem item, const TSharedRef<STableViewBase>& OwnerTable);
+	FORCEINLINE void		OnGetChildrenForView(FItem item, FItemArray& children) { children = item->_children; };
+	void					OnClickItemForView(FItem item);
+	void					OnDoubleClickItemForView(FItem item);
+
+	/// The list of root items
+	FItemArray _rootItems;
+
+	// The TreeView widget
+	TSharedPtr<SItemTableView> _tableView;
+
+	void RefreshRootItems();
+
+	void OnBlobNew(const RecordIDArray& rids);
+	void OnBlobMapped(const RecordIDArray& rids);
+	void OnEngineReset(int id);
+};
diff --git a/Engine/Plugins/Experimental/TextureGraph/Source/TextureGraphInsight/Public/View/STextureGraphInsightSessionView.h b/Engine/Plugins/Experimental/TextureGraph/Source/TextureGraphInsight/Public/View/STextureGraphInsightSessionView.h
new file mode 100644
--- /dev/null
+++ b/Engine/Plugins/Experimental/TextureGraph/Source/TextureGraphInsight/Public/View/STextureGraphInsightSessionView.h
@@ -0,0 +1,78 @@
+// Copyright Epic Games, Inc. All Rights Reserved.
+#pragma once
+
+//#include "SlateBasics.h"
+
+#include "Model/TextureGraphInsightSession.h"
+#include "Widgets/SCompoundWidget.h"
+#include "Widgets/Views/STreeView.h"
+
+class STextureGraphInsightBatchJobViewRow; /// Declare the concrete type of widget used for the raws of the tree view. Defined in the cpp file
+class ITableRow;
+class STableViewBase;
+
+class TEXTUREGRAPHINSIGHT_API STextureGraphInsightBatchJobView : public SCompoundWidget
+{
+public:
+	SLATE_BEGIN_ARGS(STextureGraphInsightBatchJobView)
+		: _inspectOnSimpleClick(false) {}
+	SLATE_ARGUMENT(RecordID, recordID)
+	SLATE_ARGUMENT(bool, inspectOnSimpleClick)
+	SLATE_END_ARGS()
+
+	void Construct(const FArguments& Args);
+
+	// TreeView Item Types
+	class FItemData;
+	using FItem = TSharedPtr<FItemData>;
+	using FItemArray = TArray< FItem >;
+	class FItemData
+	{
+	public:
+		FItemData(RecordID id, bool isMainPhase = true) : _recordID(id), _isMainPhase(isMainPhase) {}
+		TSharedPtr < STextureGraphInsightBatchJobViewRow > _widget;
+		RecordID _recordID;
+		FItemArray _children;
+		bool _isMainPhase = true;
+	};
+	using SItemTreeView = STreeView< FItem >;
+
+	// Standard delegates for the tree view
+	TSharedRef<ITableRow>	OnGenerateRowForTree(FItem item, const TSharedRef<STableViewBase>& OwnerTable);
+	FORCEINLINE void		OnGetChildrenForView(FItem item, FItemArray& children) { children = item->_children; };
+	void					OnClickItemForTree(FItem item);
+	void					OnDoubleClickItemForTree(FItem item);
+	TSharedPtr<SWidget>		OnContextMenuOpeningForTree();
+	TSharedPtr<SWidget>		OnContextMenuBatch(FItem batchItem);
+	TSharedPtr<SWidget>		OnContextMenuJob(FItem jobItem);
+
+	// Inspect on double click always
+	// eventually trigger inspect on simple click
+	bool _inspectOnSimpleClick = false;
+
+	/// The list of root items
+	FItemArray _rootItems;
+
+	// The TreeView widget
+	TSharedPtr<SItemTreeView> _treeView;
+
+
+	void OnReplayBatch(RecordID batchId, bool captureRenderDoc);
+	void OnReplayJob(RecordID batchId, bool captureRenderDoc);
+
+	void OnBatchNew(RecordID batch);
+	void OnBatchUpdate(RecordID batch);
+	void OnEngineReset(int32);
+};
+
+
+class TEXTUREGRAPHINSIGHT_API STextureGraphInsightSessionView : public SCompoundWidget
+{
+public:
+	SLATE_BEGIN_ARGS(STextureGraphInsightSessionView) {}
+	SLATE_END_ARGS()
+
+	void Construct(const FArguments& Args);
+
+	TSharedPtr<STextureGraphInsightBatchJobView> _batchJobView;
+};
diff --git a/Engine/Plugins/Experimental/TextureGraph/Source/TextureGraphInsight/TextureGraphInsight.Build.cs b/Engine/Plugins/Experimental/TextureGraph/Source/TextureGraphInsight/TextureGraphInsight.Build.cs
new file mode 100644
--- /dev/null
+++ b/Engine/Plugins/Experimental/TextureGraph/Source/TextureGraphInsight/TextureGraphInsight.Build.cs
@@ -0,0 +1,36 @@
+// Copyright Epic Games, Inc. All Rights Reserved.
+using UnrealBuildTool;
+
+public class TextureGraphInsight : ModuleRules
+{
+	public TextureGraphInsight(ReadOnlyTargetRules Target) : base(Target)
+	{
+		PCHUsage = PCHUsageMode.UseExplicitOrSharedPCHs;
+		bUseRTTI = true;
+		bEnableExceptions = true;
+
+		PublicDependencyModuleNames.AddRange(new string[] {
+			"Core",
+			"CoreUObject",
+			"Engine",
+			"InputCore",
+			"RenderCore",
+			"RHI"
+		});
+
+		PrivateDependencyModuleNames.AddRange(new string[]
+		{
+			"Slate",
+			"SlateCore",
+			"UMG"
+		});
+
+		PrivateDependencyModuleNames.AddRange(new string[]
+		{
+			"TextureGraphEngine",
+			"FreeImage",
+			"Function2",
+			"Continuable",
+		});
+	}
+}
diff --git a/Engine/Plugins/Experimental/TextureGraph/Source/TextureGraphInsightEditor/Private/TextureGraphInsightEditor.cpp b/Engine/Plugins/Experimental/TextureGraph/Source/TextureGraphInsightEditor/Private/TextureGraphInsightEditor.cpp
new file mode 100644
--- /dev/null
+++ b/Engine/Plugins/Experimental/TextureGraph/Source/TextureGraphInsightEditor/Private/TextureGraphInsightEditor.cpp
@@ -0,0 +1,415 @@
+// Copyright Epic Games, Inc. All Rights Reserved.
+
+#include "TextureGraphInsightEditor.h"
+#include "Framework/Application/SlateApplication.h"
+#include "Layout/WidgetPath.h"
+#include "LevelEditor.h"
+#include "TextureGraphEngine.h"
+#include "TextureGraphInsight.h"
+#include "TextureGraphInsightEditorCommands.h"
+#include "TextureGraphInsightEditorStyle.h"
+#include "ToolMenus.h"
+#include "View/STextureGraphInsightActionView.h"
+#include "View/STextureGraphInsightDeviceView.h"
+#include "View/STextureGraphInsightInspectorView.h"
+#include "View/STextureGraphInsightMixView.h"
+#include "View/STextureGraphInsightResourceView.h"
+#include "View/STextureGraphInsightSessionView.h"
+#include "Widgets/Docking/SDockTab.h"
+#include "Widgets/Layout/SBox.h"
+#include "Widgets/Text/STextBlock.h"
+
+
+static const FName TextureGraphInsightEditorTabName("TextureGraphInsightEditor");
+static const FName TextureGraphInsightEditorTabName_Mixes("TextureGraphInsightEditor_Mixes");
+static const FName TextureGraphInsightEditorTabName_Actions("TextureGraphInsightEditor_Actions");
+static const FName TextureGraphInsightEditorTabName_JobBatches("TextureGraphInsightEditor_JobBatches");
+static const FName TextureGraphInsightEditorTabName_Resources("TextureGraphInsightEditor_Resources");
+static const FName TextureGraphInsightEditorTabName_Devices("TextureGraphInsightEditor_Devices");
+static const FName TextureGraphInsightEditorTabName_Inspector("TextureGraphInsightEditor_Inspector");
+
+#define LOCTEXT_NAMESPACE "FTextureGraphInsightEditorModule"
+
+
+struct FInsightTabCommands : public TCommands<FInsightTabCommands>
+{
+	FInsightTabCommands()
+		: TCommands<FInsightTabCommands>(
+			TEXT("TextureGraphInsightTab"), // Context name for fast lookup
+			LOCTEXT("TextureGraphInsightTab", "TextureGraphInsightTab Debugger"), // Localized context name for displaying
+			NAME_None, // Parent
+			FCoreStyle::Get().GetStyleSetName() // Icon Style Set
+			)
+	{
+	}
+
+	// TCommand<> interface
+	virtual void RegisterCommands() override;
+	// End of TCommand<> interface
+
+	TSharedPtr<FUICommandInfo> ShowMixesTab;
+	TSharedPtr<FUICommandInfo> ShowActionsTab;
+	TSharedPtr<FUICommandInfo> ShowJobBatchesTab;
+	TSharedPtr<FUICommandInfo> ShowResourcesTab;
+	TSharedPtr<FUICommandInfo> ShowDevicesTab;
+	TSharedPtr<FUICommandInfo> ShowInspectorTab;
+};
+
+void FInsightTabCommands::RegisterCommands()
+{
+	UI_COMMAND(ShowMixesTab, "Mixes", "Toggles visibility of the Mixes tab", EUserInterfaceActionType::Check, FInputChord());
+	UI_COMMAND(ShowActionsTab, "Actions", "Toggles visibility of the Actions tab", EUserInterfaceActionType::Check, FInputChord());
+	UI_COMMAND(ShowJobBatchesTab, "JobBatches", "Toggles visibility of the JobBatches tab", EUserInterfaceActionType::Check, FInputChord());
+	UI_COMMAND(ShowResourcesTab, "Resources", "Toggles visibility of the Resources tab", EUserInterfaceActionType::Check, FInputChord());
+	UI_COMMAND(ShowDevicesTab, "Devices", "Toggles visibility of the Devices tab", EUserInterfaceActionType::Check, FInputChord());
+	UI_COMMAND(ShowInspectorTab, "Inspector", "Toggles visibility of the Inspector tab", EUserInterfaceActionType::Check, FInputChord());
+}
+
+void FTextureGraphInsightEditorModule::StartupModule()
+{
+	// This code will execute after your module is loaded into memory; the exact timing is specified in the .uplugin file per-module
+	
+	FTextureGraphInsightEditorStyle::Initialize();
+	FTextureGraphInsightEditorStyle::ReloadTextures();
+
+	FTextureGraphInsightEditorCommands::Register();
+	
+	Commands = MakeShareable(new FUICommandList);
+
+	// Commands->MapAction(
+	// 	FTextureGraphInsightEditorCommands::Get().OpenPluginWindow,
+	// 	FExecuteAction::CreateRaw(this, &FTextureGraphInsightEditorModule::PluginButtonClicked),
+	// 	FCanExecuteAction());
+
+	// UToolMenus::RegisterStartupCallback(FSimpleMulticastDelegate::FDelegate::CreateRaw(this, &FTextureGraphInsightEditorModule::RegisterMenus));
+
+	FInsightTabCommands::Register();
+
+	FGlobalTabmanager::Get()->RegisterNomadTabSpawner(TextureGraphInsightEditorTabName, FOnSpawnTab::CreateRaw(this, &FTextureGraphInsightEditorModule::OnSpawnPluginTab))
+		.SetDisplayName(LOCTEXT("FTextureGraphInsightEditorTabTitle", "Mixer Insight"))
+		.SetMenuType(ETabSpawnerMenuType::Hidden);
+
+
+	// in StartupModule()
+	_tickDelegate = FTickerDelegate::CreateRaw(this, &FTextureGraphInsightEditorModule::Tick);
+	_tickDelegateHandle = FTSTicker::GetCoreTicker().AddTicker(_tickDelegate);
+}
+
+void FTextureGraphInsightEditorModule::ShutdownModule()
+{
+	// This function may be called during shutdown to clean up your module.  For modules that support dynamic reloading,
+	// we call this function before unloading the module.
+
+	FTSTicker::GetCoreTicker().RemoveTicker(_tickDelegateHandle);
+
+	UToolMenus::UnRegisterStartupCallback(this);
+
+	UToolMenus::UnregisterOwner(this);
+
+	TextureGraphInsight::Destroy();
+
+	FTextureGraphInsightEditorStyle::Shutdown();
+
+	FTextureGraphInsightEditorCommands::Unregister();
+
+	FGlobalTabmanager::Get()->UnregisterNomadTabSpawner(TextureGraphInsightEditorTabName);
+}
+bool FTextureGraphInsightEditorModule::Tick(float DeltaTime)
+{
+	return true;
+}
+TSharedRef<SDockTab> FTextureGraphInsightEditorModule::OnSpawnPluginTab(const FSpawnTabArgs& SpawnTabArgs)
+{
+	// (re)Instantiate the Mixer Insight singleton
+	if (TextureGraphInsight::Instance())
+	{
+		TextureGraphInsight::Destroy();
+	}
+	TextureGraphInsight::Create();
+
+
+	const TSharedRef<SDockTab> NomadTab = SNew(SDockTab)
+		.TabRole(ETabRole::NomadTab)
+		.Label(NSLOCTEXT("TextureGraphInsight", "Mixer Insight", "Mixer Insight"));
+
+	if (!TabManager.IsValid())
+	{
+		TabManager = FGlobalTabmanager::Get()->NewTabManager(NomadTab);
+		// on persist layout will handle saving layout if the editor is shut down:
+		TabManager->SetOnPersistLayout(
+			FTabManager::FOnPersistLayout::CreateStatic(
+				[](const TSharedRef<FTabManager::FLayout>& InLayout)
+				{
+					if (InLayout->GetPrimaryArea().Pin().IsValid())
+					{
+						FLayoutSaveRestore::SaveToConfig(GEditorLayoutIni, InLayout);
+					}
+				}
+			)
+		);
+	}
+	else
+	{
+		ensure(Layout.IsValid());
+	}
+
+	TWeakPtr<FTabManager> tabManagerWeak = TabManager;
+	// On tab close will save the layout if the debugging window itself is closed,
+	// this handler also cleans up any floating debugging controls. If we don't close
+	// all areas we need to add some logic to the tab manager to reuse existing tabs:
+	NomadTab->SetOnTabClosed(SDockTab::FOnTabClosedCallback::CreateStatic(
+		[](TSharedRef<SDockTab> Self, TWeakPtr<FTabManager> InTabManager)
+		{
+			TSharedPtr<FTabManager> OwningTabManager = InTabManager.Pin();
+			if (OwningTabManager.IsValid())
+			{
+				FLayoutSaveRestore::SaveToConfig(GEditorLayoutIni, OwningTabManager->PersistLayout());
+				OwningTabManager->CloseAllAreas();
+			}
+		}
+		, tabManagerWeak
+	));
+
+	if (!Layout.IsValid())
+	{
+		TabManager->RegisterTabSpawner(TextureGraphInsightEditorTabName_Mixes, FOnSpawnTab::CreateStatic(
+			[](const FSpawnTabArgs&)->TSharedRef<SDockTab>
+			{
+				return SNew(SDockTab)
+					.TabRole(ETabRole::PanelTab)
+					.Label(LOCTEXT("FTextureGraphInsightEditorTabTitleMixes", "Mixer Mixes"))
+					[
+						SNew(STextureGraphInsightMixListView)
+					];
+			}));
+		TabManager->RegisterTabSpawner(TextureGraphInsightEditorTabName_Actions, FOnSpawnTab::CreateStatic(
+			[](const FSpawnTabArgs&)->TSharedRef<SDockTab>
+			{
+				return SNew(SDockTab)
+					.TabRole(ETabRole::PanelTab)
+					.Label(LOCTEXT("FTextureGraphInsightEditorTabTitleAction", "Mixer Actions"))
+					[
+						SNew(STextureGraphInsightActionView)
+					];
+			}));
+		TabManager->RegisterTabSpawner(TextureGraphInsightEditorTabName_JobBatches, FOnSpawnTab::CreateStatic(
+			[](const FSpawnTabArgs&)->TSharedRef<SDockTab>
+			{
+				return SNew(SDockTab)
+					.TabRole(ETabRole::PanelTab)
+					.Label(LOCTEXT("FTextureGraphInsightEditorTabTitleJobBatches", "Mixer Jobs & Batches"))
+					[
+						SNew(STextureGraphInsightSessionView)
+					];
+			}));
+		TabManager->RegisterTabSpawner(TextureGraphInsightEditorTabName_Resources, FOnSpawnTab::CreateStatic(
+			[](const FSpawnTabArgs&)->TSharedRef<SDockTab>
+			{
+				return SNew(SDockTab)
+					.TabRole(ETabRole::PanelTab)
+					.Label(LOCTEXT("FTextureGraphInsightEditorTabTitleResources", "Mixer Resources"))
+					[
+						SNew(STextureGraphInsightResourceView)
+					]; 
+			}));
+		TabManager->RegisterTabSpawner(TextureGraphInsightEditorTabName_Devices, FOnSpawnTab::CreateStatic(
+			[](const FSpawnTabArgs&)->TSharedRef<SDockTab>
+			{
+				return SNew(SDockTab)
+					.TabRole(ETabRole::PanelTab)
+					.Label(LOCTEXT("FTextureGraphInsightEditorTabTitleDevices", "Mixer Devices"))
+					[
+						SNew(STextureGraphInsightDeviceListView)
+					];
+			}));
+		TabManager->RegisterTabSpawner(TextureGraphInsightEditorTabName_Inspector, FOnSpawnTab::CreateStatic(
+			[](const FSpawnTabArgs&)->TSharedRef<SDockTab>
+			{
+				return SNew(SDockTab)
+					.TabRole(ETabRole::PanelTab)
+					.Label(LOCTEXT("FTextureGraphInsightEditorTabTitleInspector", "Mixer Inspector"))
+					[
+						SNew(STextureGraphInsightInspectorView)
+					];
+			}));
+
+		Layout = FTabManager::NewLayout("Standalone_TextureGraphInsight_Layout_v2")
+			->AddArea
+			(
+				FTabManager::NewPrimaryArea()
+				->SetOrientation(Orient_Horizontal)
+				->Split
+				(
+					FTabManager::NewStack()
+					->SetSizeCoefficient(.4f)
+					->SetHideTabWell(true)
+					->AddTab(TextureGraphInsightEditorTabName_Mixes, ETabState::OpenedTab)
+					->AddTab(TextureGraphInsightEditorTabName_Actions, ETabState::OpenedTab)
+					->AddTab(TextureGraphInsightEditorTabName_JobBatches, ETabState::OpenedTab)
+					->AddTab(TextureGraphInsightEditorTabName_Resources, ETabState::OpenedTab)
+					->AddTab(TextureGraphInsightEditorTabName_Devices, ETabState::OpenedTab)
+					->SetForegroundTab(TextureGraphInsightEditorTabName_JobBatches)
+				)
+				->Split
+				(
+					FTabManager::NewStack()
+					->SetSizeCoefficient(.4f)
+					->SetHideTabWell(true)
+					->AddTab(TextureGraphInsightEditorTabName_Inspector, ETabState::OpenedTab)
+					->SetForegroundTab(TextureGraphInsightEditorTabName_Inspector)
+				)
+			);
+	}
+
+	Layout = FLayoutSaveRestore::LoadFromConfig(GEditorLayoutIni, Layout.ToSharedRef());
+
+	TSharedRef<SWidget> TabContents = TabManager->RestoreFrom(Layout.ToSharedRef(), TSharedPtr<SWindow>()).ToSharedRef();
+
+
+	// build command list for tab restoration menu:
+	TSharedPtr<FUICommandList> CommandList = MakeShareable(new FUICommandList());
+
+	TWeakPtr<FTabManager> DebuggingToolsManagerWeak = TabManager;
+
+	const auto ToggleTabVisibility = [](TWeakPtr<FTabManager> InTabManagerWeak, FName InTabName)
+	{
+		TSharedPtr<FTabManager> InDebuggingToolsManager = InTabManagerWeak.Pin();
+		if (InDebuggingToolsManager.IsValid())
+		{
+			TSharedPtr<SDockTab> ExistingTab = InDebuggingToolsManager->FindExistingLiveTab(InTabName);
+			if (ExistingTab.IsValid())
+			{
+				ExistingTab->RequestCloseTab();
+			}
+			else
+			{
+				InDebuggingToolsManager->TryInvokeTab(InTabName);
+			}
+		}
+	};
+
+	const auto IsTabVisible = [](TWeakPtr<FTabManager> InTabManagerWeak, FName InTabName)
+	{
+		TSharedPtr<FTabManager> InDebuggingToolsManager = InTabManagerWeak.Pin();
+		if (InDebuggingToolsManager.IsValid())
+		{
+			return InDebuggingToolsManager->FindExistingLiveTab(InTabName).IsValid();
+		}
+		return false;
+	};
+
+	const auto ActionMapperToCommandList = [&](TSharedPtr<FUICommandInfo> command, const FName name)
+	{
+		CommandList->MapAction(
+			command,
+			FExecuteAction::CreateStatic(
+				ToggleTabVisibility,
+				DebuggingToolsManagerWeak,
+				name
+			),
+			FCanExecuteAction::CreateStatic(
+				[]() { return true; }
+			),
+			FIsActionChecked::CreateStatic(
+				IsTabVisible,
+				DebuggingToolsManagerWeak,
+				name
+			)
+		);
+	};
+
+	ActionMapperToCommandList(FInsightTabCommands::Get().ShowMixesTab, TextureGraphInsightEditorTabName_Mixes);
+	ActionMapperToCommandList(FInsightTabCommands::Get().ShowActionsTab, TextureGraphInsightEditorTabName_Actions);
+	ActionMapperToCommandList(FInsightTabCommands::Get().ShowJobBatchesTab, TextureGraphInsightEditorTabName_JobBatches);
+	ActionMapperToCommandList(FInsightTabCommands::Get().ShowResourcesTab, TextureGraphInsightEditorTabName_Resources);
+	ActionMapperToCommandList(FInsightTabCommands::Get().ShowDevicesTab, TextureGraphInsightEditorTabName_Devices);
+	ActionMapperToCommandList(FInsightTabCommands::Get().ShowInspectorTab, TextureGraphInsightEditorTabName_Inspector);
+
+
+	TWeakPtr<SWidget> OwningWidgetWeak = NomadTab;
+	TabContents->SetOnMouseButtonUp(
+		FPointerEventHandler::CreateStatic(
+			[]( /** The geometry of the widget*/
+				const FGeometry&,
+				/** The Mouse Event that we are processing */
+				const FPointerEvent& PointerEvent,
+				TWeakPtr<SWidget> InOwnerWeak,
+				TSharedPtr<FUICommandList> InCommandList) -> FReply
+			{
+				if (PointerEvent.GetEffectingButton() == EKeys::RightMouseButton)
+				{
+					// if the tab manager is still available then make a context window that allows users to
+					// show and hide tabs:
+					TSharedPtr<SWidget> InOwner = InOwnerWeak.Pin();
+					if (InOwner.IsValid())
+					{
+						FMenuBuilder MenuBuilder(true, InCommandList);
+
+						MenuBuilder.PushCommandList(InCommandList.ToSharedRef());
+						{
+							MenuBuilder.AddMenuEntry(FInsightTabCommands::Get().ShowMixesTab);
+							MenuBuilder.AddMenuEntry(FInsightTabCommands::Get().ShowActionsTab);
+							MenuBuilder.AddMenuEntry(FInsightTabCommands::Get().ShowJobBatchesTab);
+							MenuBuilder.AddMenuEntry(FInsightTabCommands::Get().ShowResourcesTab);
+							MenuBuilder.AddMenuEntry(FInsightTabCommands::Get().ShowDevicesTab);
+							MenuBuilder.AddMenuEntry(FInsightTabCommands::Get().ShowInspectorTab);
+						}
+						MenuBuilder.PopCommandList();
+
+						FWidgetPath WidgetPath = PointerEvent.GetEventPath() != nullptr ? *PointerEvent.GetEventPath() : FWidgetPath();
+						FSlateApplication::Get().PushMenu(InOwner.ToSharedRef(), WidgetPath, MenuBuilder.MakeWidget(), PointerEvent.GetScreenSpacePosition(), FPopupTransitionEffect(FPopupTransitionEffect::ContextMenu));
+
+						return FReply::Handled();
+					}
+				}
+
+				return FReply::Unhandled();
+			}
+			, OwningWidgetWeak
+			, CommandList
+			)
+	);
+
+	NomadTab->SetContent(
+		SNew(SBorder)
+		[
+			TabContents
+		]
+	);
+
+	return NomadTab;
+}
+
+void FTextureGraphInsightEditorModule::PluginButtonClicked()
+{
+	FGlobalTabmanager::Get()->TryInvokeTab(TextureGraphInsightEditorTabName);
+}
+
+void FTextureGraphInsightEditorModule::RegisterMenus()
+{
+	// // Owner will be used for cleanup in call to UToolMenus::UnregisterOwner
+	// FToolMenuOwnerScoped OwnerScoped(this);
+	//
+	// {
+	// 	UToolMenu* Menu = UToolMenus::Get()->ExtendMenu("LevelEditor.MainMenu.Window");
+	// 	{
+	// 		FToolMenuSection& Section = Menu->FindOrAddSection("WindowLayout");
+	// 		Section.AddMenuEntryWithCommandList(FTextureGraphInsightEditorCommands::Get().OpenPluginWindow, Commands);
+	// 	}
+	// }
+	//
+	// {
+	// 	UToolMenu* ToolbarMenu = UToolMenus::Get()->ExtendMenu("LevelEditor.LevelEditorToolBar");
+	// 	{
+	// 		FToolMenuSection& Section = ToolbarMenu->FindOrAddSection("Settings");
+	// 		{
+	// 			FToolMenuEntry& Entry = Section.AddEntry(FToolMenuEntry::InitToolBarButton(FTextureGraphInsightEditorCommands::Get().OpenPluginWindow));
+	// 			Entry.SetCommandList(Commands);
+	// 		}
+	// 	}
+	// }
+}
+
+#undef LOCTEXT_NAMESPACE
+	
+IMPLEMENT_MODULE(FTextureGraphInsightEditorModule, TextureGraphInsightEditor)
diff --git a/Engine/Plugins/Experimental/TextureGraph/Source/TextureGraphInsightEditor/Private/TextureGraphInsightEditorCommands.cpp b/Engine/Plugins/Experimental/TextureGraph/Source/TextureGraphInsightEditor/Private/TextureGraphInsightEditorCommands.cpp
new file mode 100644
--- /dev/null
+++ b/Engine/Plugins/Experimental/TextureGraph/Source/TextureGraphInsightEditor/Private/TextureGraphInsightEditorCommands.cpp
@@ -0,0 +1,12 @@
+// Copyright Epic Games, Inc. All Rights Reserved.
+
+#include "TextureGraphInsightEditorCommands.h"
+
+#define LOCTEXT_NAMESPACE "FTextureGraphInsightEditorModule"
+
+void FTextureGraphInsightEditorCommands::RegisterCommands()
+{
+	// UI_COMMAND(OpenPluginWindow, "Mixer Insight", "Bring up Mixer Insight window", EUserInterfaceActionType::Button, FInputChord());
+}
+
+#undef LOCTEXT_NAMESPACE
diff --git a/Engine/Plugins/Experimental/TextureGraph/Source/TextureGraphInsightEditor/Private/TextureGraphInsightEditorStyle.cpp b/Engine/Plugins/Experimental/TextureGraph/Source/TextureGraphInsightEditor/Private/TextureGraphInsightEditorStyle.cpp
new file mode 100644
--- /dev/null
+++ b/Engine/Plugins/Experimental/TextureGraph/Source/TextureGraphInsightEditor/Private/TextureGraphInsightEditorStyle.cpp
@@ -0,0 +1,70 @@
+// Copyright Epic Games, Inc. All Rights Reserved.
+
+#include "TextureGraphInsightEditorStyle.h"
+#include "Styling/SlateStyleRegistry.h"
+#include "Framework/Application/SlateApplication.h"
+#include "Slate/SlateGameResources.h"
+#include "Interfaces/IPluginManager.h"
+
+TSharedPtr< FSlateStyleSet > FTextureGraphInsightEditorStyle::StyleInstance = NULL;
+
+void FTextureGraphInsightEditorStyle::Initialize()
+{
+	if (!StyleInstance.IsValid())
+	{
+		StyleInstance = Create();
+		FSlateStyleRegistry::RegisterSlateStyle(*StyleInstance);
+	}
+}
+
+void FTextureGraphInsightEditorStyle::Shutdown()
+{
+	FSlateStyleRegistry::UnRegisterSlateStyle(*StyleInstance);
+	ensure(StyleInstance.IsUnique());
+	StyleInstance.Reset();
+}
+
+FName FTextureGraphInsightEditorStyle::GetStyleSetName()
+{
+	static FName StyleSetName(TEXT("TextureGraphInsightEditorStyle"));
+	return StyleSetName;
+}
+
+#define IMAGE_BRUSH( RelativePath, ... ) FSlateImageBrush( Style->RootToContentDir( RelativePath, TEXT(".png") ), __VA_ARGS__ )
+#define BOX_BRUSH( RelativePath, ... ) FSlateBoxBrush( Style->RootToContentDir( RelativePath, TEXT(".png") ), __VA_ARGS__ )
+#define BORDER_BRUSH( RelativePath, ... ) FSlateBorderBrush( Style->RootToContentDir( RelativePath, TEXT(".png") ), __VA_ARGS__ )
+#define TTF_FONT( RelativePath, ... ) FSlateFontInfo( Style->RootToContentDir( RelativePath, TEXT(".ttf") ), __VA_ARGS__ )
+#define OTF_FONT( RelativePath, ... ) FSlateFontInfo( Style->RootToContentDir( RelativePath, TEXT(".otf") ), __VA_ARGS__ )
+
+const FVector2D Icon16x16(16.0f, 16.0f);
+const FVector2D Icon20x20(20.0f, 20.0f);
+const FVector2D Icon40x40(40.0f, 40.0f);
+
+TSharedRef< FSlateStyleSet > FTextureGraphInsightEditorStyle::Create()
+{
+	TSharedRef< FSlateStyleSet > Style = MakeShareable(new FSlateStyleSet("TextureGraphInsightEditorStyle"));
+	Style->SetContentRoot(IPluginManager::Get().FindPlugin("TextureGraph")->GetBaseDir() / TEXT("Resources"));
+
+	Style->Set("TextureGraphInsightEditor.OpenPluginWindow", new IMAGE_BRUSH(TEXT("texturegraph_insight_icon_40x"), Icon40x40));
+
+	return Style;
+}
+
+#undef IMAGE_BRUSH
+#undef BOX_BRUSH
+#undef BORDER_BRUSH
+#undef TTF_FONT
+#undef OTF_FONT
+
+void FTextureGraphInsightEditorStyle::ReloadTextures()
+{
+	if (FSlateApplication::IsInitialized())
+	{
+		FSlateApplication::Get().GetRenderer()->ReloadTextureResources();
+	}
+}
+
+const ISlateStyle& FTextureGraphInsightEditorStyle::Get()
+{
+	return *StyleInstance;
+}
diff --git a/Engine/Plugins/Experimental/TextureGraph/Source/TextureGraphInsightEditor/Public/TextureGraphInsightEditor.h b/Engine/Plugins/Experimental/TextureGraph/Source/TextureGraphInsightEditor/Public/TextureGraphInsightEditor.h
new file mode 100644
--- /dev/null
+++ b/Engine/Plugins/Experimental/TextureGraph/Source/TextureGraphInsightEditor/Public/TextureGraphInsightEditor.h
@@ -0,0 +1,33 @@
+// Copyright Epic Games, Inc. All Rights Reserved.
+
+#pragma once
+
+#include "CoreMinimal.h"
+#include "Modules/ModuleManager.h"
+#include "Containers/Ticker.h"
+#include "Framework/Docking/TabManager.h"
+
+class FTextureGraphInsightEditorModule : public IModuleInterface
+{
+	FTickerDelegate _tickDelegate;
+	FTSTicker::FDelegateHandle _tickDelegateHandle;
+public:
+
+	/** IModuleInterface implementation */
+	virtual void	StartupModule() override;
+	virtual void	ShutdownModule() override;
+
+	bool Tick(float DeltaTime);
+	/** This function will be bound to Command (by default it will bring up plugin window) */
+	void PluginButtonClicked();
+	
+private:
+	void RegisterMenus();
+	TSharedRef<class SDockTab> OnSpawnPluginTab(const class FSpawnTabArgs& SpawnTabArgs);
+
+private:
+	TSharedPtr<class FUICommandList>	Commands;
+	TSharedPtr<FTabManager > TabManager;
+	TSharedPtr<FTabManager::FLayout > Layout;
+
+};
diff --git a/Engine/Plugins/Experimental/TextureGraph/Source/TextureGraphInsightEditor/Public/TextureGraphInsightEditorCommands.h b/Engine/Plugins/Experimental/TextureGraph/Source/TextureGraphInsightEditor/Public/TextureGraphInsightEditorCommands.h
new file mode 100644
--- /dev/null
+++ b/Engine/Plugins/Experimental/TextureGraph/Source/TextureGraphInsightEditor/Public/TextureGraphInsightEditorCommands.h
@@ -0,0 +1,23 @@
+// Copyright Epic Games, Inc. All Rights Reserved.
+
+#pragma once
+
+#include "CoreMinimal.h"
+#include "Framework/Commands/Commands.h"
+#include "TextureGraphInsightEditorStyle.h"
+
+class FTextureGraphInsightEditorCommands : public TCommands<FTextureGraphInsightEditorCommands>
+{
+public:
+
+	FTextureGraphInsightEditorCommands()
+		: TCommands<FTextureGraphInsightEditorCommands>(TEXT("TextureGraphInsightEditor"), NSLOCTEXT("Contexts", "TextureGraphInsightEditor", "TextureGraphInsightEditor Plugin"), NAME_None, FTextureGraphInsightEditorStyle::GetStyleSetName())
+	{
+	}
+
+	// TCommands<> interface
+	virtual void RegisterCommands() override;
+
+// public:
+	// TSharedPtr< FUICommandInfo > OpenPluginWindow;
+};
diff --git a/Engine/Plugins/Experimental/TextureGraph/Source/TextureGraphInsightEditor/Public/TextureGraphInsightEditorStyle.h b/Engine/Plugins/Experimental/TextureGraph/Source/TextureGraphInsightEditor/Public/TextureGraphInsightEditorStyle.h
new file mode 100644
--- /dev/null
+++ b/Engine/Plugins/Experimental/TextureGraph/Source/TextureGraphInsightEditor/Public/TextureGraphInsightEditorStyle.h
@@ -0,0 +1,32 @@
+// Copyright Epic Games, Inc. All Rights Reserved.
+
+#pragma once
+
+#include "CoreMinimal.h"
+#include "Styling/SlateStyle.h"
+
+/**  */
+class FTextureGraphInsightEditorStyle
+{
+public:
+
+	static void Initialize();
+
+	static void Shutdown();
+
+	/** reloads textures used by slate renderer */
+	static void ReloadTextures();
+
+	/** @return The Slate style set for the Shooter game */
+	static const ISlateStyle& Get();
+
+	static FName GetStyleSetName();
+
+private:
+
+	static TSharedRef< class FSlateStyleSet > Create();
+
+private:
+
+	static TSharedPtr< class FSlateStyleSet > StyleInstance;
+};
diff --git a/Engine/Plugins/Experimental/TextureGraph/Source/TextureGraphInsightEditor/TextureGraphInsightEditor.Build.cs b/Engine/Plugins/Experimental/TextureGraph/Source/TextureGraphInsightEditor/TextureGraphInsightEditor.Build.cs
new file mode 100644
--- /dev/null
+++ b/Engine/Plugins/Experimental/TextureGraph/Source/TextureGraphInsightEditor/TextureGraphInsightEditor.Build.cs
@@ -0,0 +1,66 @@
+// Copyright Epic Games, Inc. All Rights Reserved.
+
+using UnrealBuildTool;
+
+public class TextureGraphInsightEditor : ModuleRules
+{
+	public TextureGraphInsightEditor(ReadOnlyTargetRules Target) : base(Target)
+	{
+		PCHUsage = ModuleRules.PCHUsageMode.UseExplicitOrSharedPCHs;
+		
+		PublicIncludePaths.AddRange(
+			new string[] {
+				// ... add public include paths required here ...
+			}
+			);
+				
+		
+		PrivateIncludePaths.AddRange(
+			new string[] {
+				// ... add other private include paths required here ...
+			}
+			);
+			
+		
+		PublicDependencyModuleNames.AddRange(
+			new string[]
+			{
+				"Core",
+				// ... add other public dependencies that you statically link with here ...
+			}
+			);
+			
+		
+		PrivateDependencyModuleNames.AddRange(
+			new string[]
+			{
+				"Projects",
+				"InputCore",
+				"UnrealEd",
+				"ToolMenus",
+				"CoreUObject",
+				"Engine",
+				"Slate",
+				"SlateCore",
+                "Blutility",
+                "UMG",
+                "UMGEditor",
+                "AdvancedPreviewScene",
+				"StaticMeshEditor",
+				// ... add private dependencies that you statically link with here ...	
+				"Continuable",
+				"DeveloperSettings",
+                "TextureGraphEngine",
+                "TextureGraphInsight",
+            }
+			);
+		
+		
+		DynamicallyLoadedModuleNames.AddRange(
+			new string[]
+			{
+				// ... add any modules that your module loads dynamically here ...
+			}
+			);
+	}
+}
diff --git a/Engine/Plugins/Experimental/TextureGraph/TextureGraph.uplugin b/Engine/Plugins/Experimental/TextureGraph/TextureGraph.uplugin
--- a/Engine/Plugins/Experimental/TextureGraph/TextureGraph.uplugin
+++ b/Engine/Plugins/Experimental/TextureGraph/TextureGraph.uplugin
@@ -23,7 +23,7 @@
 			]
 		},
 		{
-			"Name": "MixerInsight",
+			"Name": "TextureGraphInsight",
 			"Type": "Editor",
 			"LoadingPhase": "Default",
 			"PlatformAllowList": [
@@ -31,7 +31,7 @@
 			]
 		},
 		{
-			"Name": "MixerInsightEditor",
+			"Name": "TextureGraphInsightEditor",
 			"Type": "Editor",
 			"LoadingPhase": "Default",
 			"PlatformAllowList": [
diff --git a/Engine/Source/UnrealEditorVM.Target.cs b/Engine/Source/UnrealEditorVM.Target.cs
--- a/Engine/Source/UnrealEditorVM.Target.cs
+++ b/Engine/Source/UnrealEditorVM.Target.cs
@@ -12,6 +12,7 @@ public class UnrealEditorVMTarget : UnrealEditorTarget
 		bUseAutoRTFMCompiler = true;
 		BuildEnvironment = TargetBuildEnvironment.Unique;
 		bLegalToDistributeBinary = true;
+		Type = TargetType.Program;
 
 		DisablePlugins.AddRange(
 			new string[] {
