Nodes producing a grayscale mask should output RGBA8 by default
THis way the value is broadcasted on the 4 components of the texels when fetched for further operations

Updated the tests which have changed as a result
instead of producing red only images, they are black and white images
#tests 65b3ed0844366aef6ec25e62

diff --git a/Engine/Plugins/Experimental/TextureGraph/Source/TextureGraph/Private/Expressions/Procedural/TG_Expression_Noise.cpp b/Engine/Plugins/Experimental/TextureGraph/Source/TextureGraph/Private/Expressions/Procedural/TG_Expression_Noise.cpp
--- a/Engine/Plugins/Experimental/TextureGraph/Source/TextureGraph/Private/Expressions/Procedural/TG_Expression_Noise.cpp
+++ b/Engine/Plugins/Experimental/TextureGraph/Source/TextureGraph/Private/Expressions/Procedural/TG_Expression_Noise.cpp
@@ -34,12 +34,12 @@ void UTG_Expression_Noise::Evaluate(FTG_EvaluationContext* InContext)
 	const FString Name = TEXT("Noise"); /// FString::Printf(TEXT("[%s].[%d].[%llu] Noise"), GetDefaultName(), InContext->TargetId, InContext->Cycle->Batch()->BatchId());
 
 	BufferDescriptor Desc = Output.GetBufferDescriptor();
-
+	
+	// Default format to RGBA8
 	if (Desc.Format == BufferFormat::Auto)
 		Desc.Format = BufferFormat::Byte;
-
 	if (Desc.ItemsPerPoint == 0)
-		Desc.ItemsPerPoint = 1;
+		Desc.ItemsPerPoint = 4;
 
 	Desc.DefaultValue = FLinearColor(1.0f, 1.0f, 1.0f, 1.0f);
 
diff --git a/Engine/Plugins/Experimental/TextureGraph/Source/TextureGraph/Private/Expressions/Procedural/TG_Expression_Pattern.cpp b/Engine/Plugins/Experimental/TextureGraph/Source/TextureGraph/Private/Expressions/Procedural/TG_Expression_Pattern.cpp
--- a/Engine/Plugins/Experimental/TextureGraph/Source/TextureGraph/Private/Expressions/Procedural/TG_Expression_Pattern.cpp
+++ b/Engine/Plugins/Experimental/TextureGraph/Source/TextureGraph/Private/Expressions/Procedural/TG_Expression_Pattern.cpp
@@ -57,12 +57,11 @@ void UTG_Expression_Pattern::Evaluate(FTG_EvaluationContext* InContext)
 	const FString Name = FString::Printf(TEXT("[%s].[%d].[%llu] Pattern"), *GetDefaultName().ToString(), InContext->TargetId, InContext->Cycle->GetBatch()->GetBatchId());
 
 	BufferDescriptor Desc = Output.GetBufferDescriptor();
-
+	// Default format to RGBA8
 	if (Desc.Format == BufferFormat::Auto)
 		Desc.Format = BufferFormat::Byte;
-
 	if (Desc.ItemsPerPoint == 0)
-		Desc.ItemsPerPoint = 1;
+		Desc.ItemsPerPoint = 4;
 
 	Desc.DefaultValue = FLinearColor(1.0f, 1.0f, 1.0f, 1.0f);
 
diff --git a/Engine/Plugins/Experimental/TextureGraph/Source/TextureGraph/Private/Expressions/Procedural/TG_Expression_Shape.cpp b/Engine/Plugins/Experimental/TextureGraph/Source/TextureGraph/Private/Expressions/Procedural/TG_Expression_Shape.cpp
--- a/Engine/Plugins/Experimental/TextureGraph/Source/TextureGraph/Private/Expressions/Procedural/TG_Expression_Shape.cpp
+++ b/Engine/Plugins/Experimental/TextureGraph/Source/TextureGraph/Private/Expressions/Procedural/TG_Expression_Shape.cpp
@@ -13,7 +13,7 @@ void UTG_Expression_Shape::Evaluate(FTG_EvaluationContext* InContext)
 	/// If it's auto then make it grey scale as we don't need more than a single channel output
 	if (DesiredDescriptor.TextureFormat == ETG_TextureFormat::Auto)
 	{
-		DesiredDescriptor.TextureFormat = ETG_TextureFormat::G8;
+		DesiredDescriptor.TextureFormat = ETG_TextureFormat::BGRA8;
 	}
 
 	T_ShapeMask::FParams Params{
diff --git a/Engine/Plugins/Experimental/TextureGraph/Source/TextureGraphEngine/Transform/Expressions/T_MaterialIDMask.cpp b/Engine/Plugins/Experimental/TextureGraph/Source/TextureGraphEngine/Transform/Expressions/T_MaterialIDMask.cpp
--- a/Engine/Plugins/Experimental/TextureGraph/Source/TextureGraphEngine/Transform/Expressions/T_MaterialIDMask.cpp
+++ b/Engine/Plugins/Experimental/TextureGraph/Source/TextureGraphEngine/Transform/Expressions/T_MaterialIDMask.cpp
@@ -53,11 +53,11 @@ TiledBlobPtr T_MaterialIDMask::Create(MixUpdateCyclePtr InCycle, TiledBlobPtr In
 
 	BufferDescriptor Desc = DesiredDesc;
 
+	// Default format to RGBA8
 	if (Desc.Format == BufferFormat::Auto)
 		Desc.Format = BufferFormat::Byte;
-
 	if (Desc.ItemsPerPoint == 0)
-		Desc.ItemsPerPoint = 1;
+		Desc.ItemsPerPoint = 4;
 
 	Desc.DefaultValue = FLinearColor(1.0f, 1.0f, 1.0f, 1.0f);
 
