Fixing build warning for lack of default value

diff --git a/Engine/Plugins/Experimental/TextureGraph/Source/TextureGraph/Public/TG_OutputSettings.h b/Engine/Plugins/Experimental/TextureGraph/Source/TextureGraph/Public/TG_OutputSettings.h
--- a/Engine/Plugins/Experimental/TextureGraph/Source/TextureGraph/Public/TG_OutputSettings.h
+++ b/Engine/Plugins/Experimental/TextureGraph/Source/TextureGraph/Public/TG_OutputSettings.h
@@ -38,10 +38,10 @@ public:
 		ETG_TexturePresetType TexturePresetType = ETG_TexturePresetType::None;
 
 	UPROPERTY(BlueprintReadWrite, EditAnywhere, Category = "Advanced", DisplayName = "LOD Texture Group", Meta = (NoResetToDefault, EditCondition = "TexturePresetType == ETG_TexturePresetType::None"))
-		TEnumAsByte<enum TextureGroup> LODGroup;
+		TEnumAsByte<enum TextureGroup> LODGroup = TextureGroup::TEXTUREGROUP_World;
 
 	UPROPERTY(BlueprintReadWrite, EditAnywhere, Category = "Advanced", DisplayName = "Compression", Meta = (NoResetToDefault, EditCondition = "TexturePresetType == ETG_TexturePresetType::None") )
-		TEnumAsByte <enum TextureCompressionSettings> Compression;
+		TEnumAsByte <enum TextureCompressionSettings> Compression = TextureCompressionSettings::TC_Default;
 
 	UPROPERTY(BlueprintReadWrite, EditAnywhere, Category = "Advanced", DisplayName = "sRGB", Meta = (NoResetToDefault, EditCondition = "TexturePresetType == ETG_TexturePresetType::None"))
 		bool bSRGB = false;
