OnFinalised for blobs still trying to access some invalid slate objects

diff --git a/Engine/Plugins/Experimental/TextureGraph/Source/TextureGraphEditor/Private/Customizations/STG_TextureHistogram.cpp b/Engine/Plugins/Experimental/TextureGraph/Source/TextureGraphEditor/Private/Customizations/STG_TextureHistogram.cpp
--- a/Engine/Plugins/Experimental/TextureGraph/Source/TextureGraphEditor/Private/Customizations/STG_TextureHistogram.cpp
+++ b/Engine/Plugins/Experimental/TextureGraph/Source/TextureGraphEditor/Private/Customizations/STG_TextureHistogram.cpp
@@ -37,11 +37,19 @@ void STG_TextureHistogram::SetTexture(FTG_Texture& Source, UTextureGraph* InText
 		Source.RasterBlob->OnFinalise()
 		.then([this, InTextureGraph, Source]() mutable
 		{
-			T_TextureHistogram::CreateOnService(InTextureGraph, std::static_pointer_cast<TiledBlob>(Source.RasterBlob), 0);
-			return  Source->GetHistogram()->OnFinalise();
+			// OnFinalise can sometimes occur after the editor is closed and thus can potentially
+			// deallocate all corresponding slate objects
+			if (DoesSharedInstanceExist())
+			{
+				T_TextureHistogram::CreateOnService(InTextureGraph, std::static_pointer_cast<TiledBlob>(Source.RasterBlob), 0);
+				return  Source->GetHistogram()->OnFinalise();
+			}
+			return (AsyncBlobResultPtr)(cti::make_ready_continuable<const Blob*>(nullptr));
 		})
 		.then([this, Source]() mutable
 		{
+			// OnFinalise can sometimes occur after the editor is closed and thus can potentially
+			// deallocate all corresponding slate objects
 			if (DoesSharedInstanceExist())
             {
 				TiledBlobPtr FinalizedHistogram = std::static_pointer_cast<TiledBlob>(Source.RasterBlob->GetHistogram());
diff --git a/Engine/Plugins/Experimental/TextureGraph/Source/TextureGraphEditor/Private/EdGraph/STG_EditorGraphNode.cpp b/Engine/Plugins/Experimental/TextureGraph/Source/TextureGraphEditor/Private/EdGraph/STG_EditorGraphNode.cpp
--- a/Engine/Plugins/Experimental/TextureGraph/Source/TextureGraphEditor/Private/EdGraph/STG_EditorGraphNode.cpp
+++ b/Engine/Plugins/Experimental/TextureGraph/Source/TextureGraphEditor/Private/EdGraph/STG_EditorGraphNode.cpp
@@ -885,6 +885,8 @@ TSharedRef<SWidget> STG_EditorGraphNode::CreateTitleRightWidget()
 			{
 				CachedThumb->OnFinalise().then([CachedThumb, ThumbnailWidget]()
 				{
+					// NOTE: If later, "this" were to be captured here, we should check DoesSharedInstanceExist()
+					// as there is a chance this might be invoked when the slate widgets have already been destroyed
 					if (ThumbnailWidget.IsValid())
 					{
 						ThumbnailWidget->UpdateBlob(CachedThumb);
@@ -1010,15 +1012,26 @@ void STG_EditorGraphNode::UpdateThumbnail(const FTG_EvaluationContext* InContext
 				ThumbBlob->OnFinalise()
 					.then([=, this](const Blob* FinalisedBlob) mutable
 						{
-							return ThumbBlob->CombineTiles(true, false);
+							// OnFinalise can sometimes occur after the editor is closed and thus can potentially
+							// deallocate all corresponding slate objects
+							if (DoesSharedInstanceExist())
+							{
+								return ThumbBlob->CombineTiles(true, false);
+							}
+							return static_cast<AsyncBufferResultPtr>(cti::make_ready_continuable<BufferResultPtr>(std::make_shared<BufferResult>()));
 						})
 					.then([=, this](BufferResultPtr) mutable
 						{
-							auto ThumbWidget = FindOrCreateThumbWidget(Pin->GetId());
-
-							if (ThumbWidget.IsValid())
+							// OnFinalise can sometimes occur after the editor is closed and thus can potentially
+							// deallocate all corresponding slate objects
+							if (DoesSharedInstanceExist())
 							{
-								ThumbWidget->UpdateBlob(ThumbBlob);
+								auto ThumbWidget = FindOrCreateThumbWidget(Pin->GetId());
+
+								if (ThumbWidget.IsValid())
+								{
+									ThumbWidget->UpdateBlob(ThumbBlob);
+								}
 							}
 						});
 			}
diff --git a/Engine/Plugins/Experimental/TextureGraph/Source/TextureGraphEditor/Private/SBlobTileView.cpp b/Engine/Plugins/Experimental/TextureGraph/Source/TextureGraphEditor/Private/SBlobTileView.cpp
--- a/Engine/Plugins/Experimental/TextureGraph/Source/TextureGraphEditor/Private/SBlobTileView.cpp
+++ b/Engine/Plugins/Experimental/TextureGraph/Source/TextureGraphEditor/Private/SBlobTileView.cpp
@@ -52,8 +52,13 @@ void SBlobTileView::Construct(const FArguments& InArgs)
 		{
 			Blob->OnFinalise().then([this, InArgs, Blob, GridPanel]()
 			{
-				InArgs._OnFinalizeBlob.ExecuteIfBound();
-				CreateTiles(Blob, GridPanel, InArgs._padding);
+				// OnFinalise can sometimes occur after the editor is closed and thus can potentially
+				// deallocate all corresponding slate objects
+				if (DoesSharedInstanceExist())
+				{		
+					InArgs._OnFinalizeBlob.ExecuteIfBound();
+					CreateTiles(Blob, GridPanel, InArgs._padding);
+				}
 			});
 		}
 		else
diff --git a/Engine/Plugins/Experimental/TextureGraph/Source/TextureGraphEditor/Private/STG_OutputSelectionDlg.cpp b/Engine/Plugins/Experimental/TextureGraph/Source/TextureGraphEditor/Private/STG_OutputSelectionDlg.cpp
--- a/Engine/Plugins/Experimental/TextureGraph/Source/TextureGraphEditor/Private/STG_OutputSelectionDlg.cpp
+++ b/Engine/Plugins/Experimental/TextureGraph/Source/TextureGraphEditor/Private/STG_OutputSelectionDlg.cpp
@@ -97,6 +97,8 @@ void STG_OutputSelectionDlg::AddExportItems()
 				ThumbBlob->OnFinalise()
 					.then([ThumbBlob, NodeThumbnail]
 					{
+						// NOTE: If later, "this" were to be captured here, we should check DoesSharedInstanceExist()
+						// as there is a chance this might be invoked when the slate widgets have already been destroyed
 						if (NodeThumbnail.IsValid())
 						{
 							NodeThumbnail->UpdateBlob(ThumbBlob);
diff --git a/Engine/Plugins/Experimental/TextureGraph/Source/TextureGraphEditor/Private/STG_SelectionPreview.cpp b/Engine/Plugins/Experimental/TextureGraph/Source/TextureGraphEditor/Private/STG_SelectionPreview.cpp
--- a/Engine/Plugins/Experimental/TextureGraph/Source/TextureGraphEditor/Private/STG_SelectionPreview.cpp
+++ b/Engine/Plugins/Experimental/TextureGraph/Source/TextureGraphEditor/Private/STG_SelectionPreview.cpp
@@ -736,17 +736,27 @@ void STG_SelectionPreview::OnSelectionChanged(UTG_EdGraphNode* InNode)
 			if (Blob->IsTiled())
 			{
 				Blob->OnFinalise()
-					.then([=]()
+					.then([=, this]()
 					{
-						TiledBlobPtr BlobTiled = std::static_pointer_cast<TiledBlob>(Blob);
-						return BlobTiled->CombineTiles(false, false);
+						// OnFinalise can sometimes occur after the editor is closed and thus can potentially
+						// deallocate all corresponding slate objects
+						if (DoesSharedInstanceExist())
+						{
+							TiledBlobPtr BlobTiled = std::static_pointer_cast<TiledBlob>(Blob);
+							return BlobTiled->CombineTiles(false, false);
+						}
+						return static_cast<AsyncBufferResultPtr>(cti::make_ready_continuable<BufferResultPtr>(std::make_shared<BufferResult>()));
+						
 					})
 					.then([this, Blob](BufferResultPtr BufferPtr)
 					{
-						SetTexture(Blob);
+						if (DoesSharedInstanceExist())
+						{
+							SetTexture(Blob);
 
-						// Set Texture Properties
-						SetTextureProperties();
+							// Set Texture Properties
+							SetTextureProperties();
+						}
 					});
 			}
 			else
@@ -754,10 +764,13 @@ void STG_SelectionPreview::OnSelectionChanged(UTG_EdGraphNode* InNode)
 				Blob->OnFinalise()
 					.then([this, Blob]()
 					{
-						// Set Texture
-						SetTexture(Blob);
-						// Set Texture Properties
-						SetTextureProperties();
+						if (DoesSharedInstanceExist())
+						{
+							// Set Texture
+							SetTexture(Blob);
+							// Set Texture Properties
+							SetTextureProperties();
+						}
 					});
 			}
 		}
diff --git a/Engine/Plugins/Experimental/TextureGraph/Source/TextureGraphEditor/Private/STG_TextureDetails.cpp b/Engine/Plugins/Experimental/TextureGraph/Source/TextureGraphEditor/Private/STG_TextureDetails.cpp
--- a/Engine/Plugins/Experimental/TextureGraph/Source/TextureGraphEditor/Private/STG_TextureDetails.cpp
+++ b/Engine/Plugins/Experimental/TextureGraph/Source/TextureGraphEditor/Private/STG_TextureDetails.cpp
@@ -154,7 +154,9 @@ void STG_TextureDetails::CalculateHistogram(BlobPtr InBlob, UTextureGraph* InTex
 	InBlob->OnFinalise()
 		.then([this, InTextureGraph, InBlob]() mutable
 		{
-			if (!InBlob->IsTransient())
+			// OnFinalise can sometimes occur after the editor is closed and thus can potentially
+			// deallocate all corresponding slate objects
+			if (DoesSharedInstanceExist() && !InBlob->IsTransient())
 			{
 				T_TextureHistogram::CreateOnService(InTextureGraph, std::static_pointer_cast<TiledBlob>(InBlob), 0);
 				return InBlob->GetHistogram()->OnFinalise();
