Brightness contrast issue fixes

diff --git a/Engine/Plugins/Experimental/TextureGraph/Source/TextureGraphEditor/Private/STG_TextureDetails.cpp b/Engine/Plugins/Experimental/TextureGraph/Source/TextureGraphEditor/Private/STG_TextureDetails.cpp
--- a/Engine/Plugins/Experimental/TextureGraph/Source/TextureGraphEditor/Private/STG_TextureDetails.cpp
+++ b/Engine/Plugins/Experimental/TextureGraph/Source/TextureGraphEditor/Private/STG_TextureDetails.cpp
@@ -163,7 +163,7 @@ void STG_TextureDetails::CalculateHistogram(BlobPtr InBlob, UTextureGraph* InTex
 			}
 			return (AsyncBlobResultPtr)(cti::make_ready_continuable<const Blob*>(nullptr));
 		})
-		.then([this, InBlob]() mutable
+		.then([this, InBlob](const Blob* Histogram) mutable
 		{
 			if (DoesSharedInstanceExist() && !InBlob->IsTransient())
 			{
diff --git a/Engine/Plugins/Experimental/TextureGraph/Source/TextureGraphEngine/Data/Blob.cpp b/Engine/Plugins/Experimental/TextureGraph/Source/TextureGraphEngine/Data/Blob.cpp
--- a/Engine/Plugins/Experimental/TextureGraph/Source/TextureGraphEngine/Data/Blob.cpp
+++ b/Engine/Plugins/Experimental/TextureGraph/Source/TextureGraphEngine/Data/Blob.cpp
@@ -305,10 +305,9 @@ AsyncBlobResultPtr Blob::OnFinalise() const
 
 	return cti::make_continuable<const Blob*>([this](auto&& Promise) 
 	{
-		TSharedPtr<cti::promise<const Blob*>> PromisePtr = MakeShared<cti::promise<const Blob*>>(std::move(Promise));
-		OnFinaliseInternal([this, PromisePtr](const Blob* BlobObj) mutable
+		OnFinaliseInternal([this, FWD_PROMISE(Promise)](const Blob* BlobObj) mutable
 		{
-			PromisePtr->set_value(BlobObj);
+			Promise.set_value(BlobObj);
 		});
 	});
 }
@@ -494,24 +493,46 @@ void Blob::AddLinkedBlob(BlobPtr LinkedBlob)
 {
 	if (IsFinalised())
 	{
-		*LinkedBlob = *this;
+		LinkedBlob->FinaliseFrom(this);
 		return;
 	}
 	LinkedBlobs.push_back(LinkedBlob);
 }
 
-void Blob::FinaliseFrom(const Blob* RHS)
+void Blob::SyncWith(Blob* RHS)
+{
+	/// Synchronise the histograms
+	if (!Histogram && RHS->Histogram)
+		Histogram = RHS->Histogram;
+	else if (Histogram && !RHS->Histogram)
+		RHS->Histogram = Histogram;
+	else if (Histogram && RHS->Histogram)
+	{
+		/// Sync based on whatever has been finalised
+		if (RHS->Histogram->IsFinalised() && !Histogram->IsFinalised())
+			Histogram->FinaliseFrom(RHS->Histogram.get());
+		else if (Histogram->IsFinalised() && !RHS->Histogram->IsFinalised())
+			RHS->Histogram->FinaliseFrom(Histogram.get());
+
+		/// If both have been finalised then there's nothing we should do.
+		/// We have two copies within the system and in due course, one
+		/// of these will get deleted
+	}
+}
+
+void Blob::FinaliseFrom(Blob* RHS)
 {
 	Buffer = RHS->Buffer;
 	LODLevels = RHS->LODLevels;
 	MinMax = RHS->MinMax;
 	MinValue = RHS->MinValue;
 	MaxValue = RHS->MaxValue;
-	Histogram = RHS->Histogram;
+
 	LODParent = RHS->LODParent;
 	LODSource = RHS->LODSource;
 	bIsLODLevel = RHS->bIsLODLevel;
 
+	SyncWith(RHS);
 	FinaliseNow(true, nullptr);
 }
 
diff --git a/Engine/Plugins/Experimental/TextureGraph/Source/TextureGraphEngine/Data/Blob.h b/Engine/Plugins/Experimental/TextureGraph/Source/TextureGraphEngine/Data/Blob.h
--- a/Engine/Plugins/Experimental/TextureGraph/Source/TextureGraphEngine/Data/Blob.h
+++ b/Engine/Plugins/Experimental/TextureGraph/Source/TextureGraphEngine/Data/Blob.h
@@ -13,7 +13,7 @@ class BlobTransform;
 class Blobber;
 struct ResourceBindInfo;
 
-typedef TFunction<void(const Blob*)> BlobReadyCallback;
+typedef TUniqueFunction<void(const Blob*)> BlobReadyCallback;
 typedef cti::continuable<const Blob*> AsyncBlobResultPtr;
 
 typedef std::shared_ptr<Blob>	BlobPtr;
@@ -83,7 +83,8 @@ protected:
 	virtual void					UpdateLinkedBlobs(bool bDoFinalise);
 	virtual void					AddLinkedBlob(BlobPtr LinkedBlob);
 
-	virtual void					FinaliseFrom(const Blob* RHS);
+	virtual void					FinaliseFrom(Blob* RHS);
+	virtual void					SyncWith(Blob* RHS);
 
 public:
 									Blob(DeviceBufferRef InBuffer);
diff --git a/Engine/Plugins/Experimental/TextureGraph/Source/TextureGraphEngine/Data/Blobber.cpp b/Engine/Plugins/Experimental/TextureGraph/Source/TextureGraphEngine/Data/Blobber.cpp
--- a/Engine/Plugins/Experimental/TextureGraph/Source/TextureGraphEngine/Data/Blobber.cpp
+++ b/Engine/Plugins/Experimental/TextureGraph/Source/TextureGraphEngine/Data/Blobber.cpp
@@ -219,7 +219,7 @@ void Blobber::UpdateBlobCache()
 	FScopeLock Lock(&BlobLookupLock);
 	auto Iter = BlobCache.GetCache().begin();
 	bool bStop = false;
-	CHashPtrVec HashesToRemove;
+	HashTypeVec HashesToRemove;
 	int32 IterCount = 0;
 	double PrevTimestamp = -1;
 
@@ -232,7 +232,14 @@ void Blobber::UpdateBlobCache()
 		if (Entry->BlobObj.use_count() == 1)
 		{
 			UE_LOG(LogBlob, Log, TEXT("Removing permanent blob %s (Hash: %llu)"), *Entry->BlobObj->DisplayName(), Entry->BlobObj->Hash()->Value());
-			HashesToRemove.push_back(Entry->BlobObj->Hash());
+			HashesToRemove.push_back(Iter.Key());
+
+			HashTypeVec IntermediateHashes = Entry->BlobObj->Hash()->GetIntermediateHashes();
+			if (!IntermediateHashes.empty())
+				HashesToRemove.insert(HashesToRemove.end(), IntermediateHashes.begin(), IntermediateHashes.end());
+
+			if (Entry->BlobObj->Hash()->Value() != Iter.Key())
+				HashesToRemove.push_back(Entry->BlobObj->Hash()->Value());
 		}
 		else
 		{
@@ -255,8 +262,12 @@ void Blobber::UpdateBlobCache()
 		UE_LOG(LogBlob, Log, TEXT("Removing num items from the cache: %llu"), HashesToRemove.size());
 		for (size_t RemoveIndex = 0; RemoveIndex < HashesToRemove.size(); RemoveIndex++)
 		{
-			CHashPtr HashToRemove = HashesToRemove[RemoveIndex];
-			BlobCache.Remove(HashToRemove->Value());
+			HashType HashToRemove = HashesToRemove[RemoveIndex];
+			BlobCache.Remove(HashToRemove);
+			
+			auto MappingIter = HashMappings.find(HashToRemove);
+			if (MappingIter != HashMappings.end())
+				HashMappings.erase(MappingIter);
 		}
 	}
 }
@@ -353,13 +364,13 @@ void Blobber::PrintStats()
 	ThumbnailMemUsage = ThumbnailMemUsage / MBConv;
 	LargeImageMemUsage = LargeImageMemUsage / MBConv;
 	
-	UE_LOG(LogBlob, Log, TEXT("===== BEGIN Blobber STATS ====="));
+	UE_LOG(LogBlob, VeryVerbose, TEXT("===== BEGIN Blobber STATS ====="));
 	
-	UE_LOG(LogBlob, Log, TEXT("Total Objects    : %llu (Transient: %llu, Tiled: %llu)"), NumObjects, NumTransient, NumTiled);
-	UE_LOG(LogBlob, Log, TEXT("Num Large Images	: %llu (Thumbnails: %llu)"), NumLargeImages, NumThumbnails);
-	UE_LOG(LogBlob, Log, TEXT("Large Image Mem  : %0.2f MB (Thumbnail Mem: %0.2f MB)"), LargeImageMemUsage, ThumbnailMemUsage);
+	UE_LOG(LogBlob, VeryVerbose, TEXT("Total Objects    : %llu (Transient: %llu, Tiled: %llu)"), NumObjects, NumTransient, NumTiled);
+	UE_LOG(LogBlob, VeryVerbose, TEXT("Num Large Images	: %llu (Thumbnails: %llu)"), NumLargeImages, NumThumbnails);
+	UE_LOG(LogBlob, VeryVerbose, TEXT("Large Image Mem  : %0.2f MB (Thumbnail Mem: %0.2f MB)"), LargeImageMemUsage, ThumbnailMemUsage);
 	
-	UE_LOG(LogBlob, Log, TEXT("===== END Blobber STATS   ====="));
+	UE_LOG(LogBlob, VeryVerbose, TEXT("===== END Blobber STATS   ====="));
 }
 
 BlobRef Blobber::Create(Device* Dev, RawBufferPtr Raw)
@@ -638,6 +649,11 @@ void Blobber::UpdateHash(HashType OldHash, CHashPtr Hash)
 		else
 			BlobObj->GetBufferRef()->SetHash(Hash);
 	}
+	else
+	{
+		TiledBlobPtr TiledBlobObj = std::static_pointer_cast<TiledBlob>(BlobObj);
+		TiledBlobObj->HashValue = Hash;
+	}
 
 	UpdateBlobHash(OldHash, BlobObj);
 }
@@ -722,12 +738,13 @@ TiledBlobRef Blobber::AddTiledResult(CHashPtr LHash, TiledBlobPtr Result, BlobCa
 {
 	BlobRef Ret = AddResult(LHash, std::static_pointer_cast<Blob>(Result), Options);
 	check(Ret->IsTiled());
+
 	TiledBlobPtr CachedPtr = std::static_pointer_cast<TiledBlob>(Ret.get());
+	check(CachedPtr->IsTiled());
 
 	if (CachedPtr && Result != CachedPtr)
 	{
-		if (typeid(*CachedPtr) == typeid(*Result))
-			CachedPtr->AddLinkedBlob(Result);
+		CachedPtr->AddLinkedBlob(Result);
 	}
 
 	/// If we got a different pointer back from the blobber then that means that this result was already cached into 
@@ -824,6 +841,10 @@ void Blobber::AddBlobEntry(HashType Hash, BlobPtr BlobObj, BlobCacheOptions Opti
 
 	/// Sanity check that we're not adding this multiple times
 	const BlobCacheEntryPtr* ExistingEntry = BlobCache.Find(Hash, false);
+
+	if (ExistingEntry && (*ExistingEntry)->BlobObj == BlobObj)
+		return;
+
 	check(!ExistingEntry || (*ExistingEntry)->BlobObj != BlobObj);
 	
 	BlobCacheEntryPtr Entry = std::make_shared<BlobCacheEntry>();
diff --git a/Engine/Plugins/Experimental/TextureGraph/Source/TextureGraphEngine/Data/TiledBlob.cpp b/Engine/Plugins/Experimental/TextureGraph/Source/TextureGraphEngine/Data/TiledBlob.cpp
--- a/Engine/Plugins/Experimental/TextureGraph/Source/TextureGraphEngine/Data/TiledBlob.cpp
+++ b/Engine/Plugins/Experimental/TextureGraph/Source/TextureGraphEngine/Data/TiledBlob.cpp
@@ -27,7 +27,7 @@ TiledBlob::TiledBlob(const BufferDescriptor& InDesc, const BlobPtrTiles& Tiles)
 		Desc.Height = std::max(Desc.Height, static_cast<uint32>(Tiles.Cols()));
 	}
 	
-	CalcHash_Now();
+	CalcHashNow();
 	Buffer.reset();
 
 	bIsFinalised = true;
@@ -123,13 +123,13 @@ TiledBlobPtr TiledBlob::AsTiledBlob(BlobPtr BlobObj)
 	return std::make_shared<TiledBlob>(BlobObj);
 }
 
-void TiledBlob::FinaliseFrom(const Blob* RHS)
+void TiledBlob::FinaliseFrom(Blob* RHS)
 {
 	check(RHS->IsTiled());
-	const TiledBlob* RHSTiled = static_cast<const TiledBlob*>(RHS);
+	TiledBlob* RHSTiled = static_cast<TiledBlob*>(RHS);
 
 	Tiles = RHSTiled->Tiles;
-	HashValue = RHSTiled->HashValue;
+	SetHash(RHSTiled->HashValue);
 
 	bool bIsTransient = Desc.bIsTransient;
 	Desc = RHSTiled->Desc;
@@ -171,7 +171,7 @@ bool TiledBlob::HasBlobAsTile(Blob* BlobObj) const
 void TiledBlob::SetTiles(const BlobPtrTiles& InTiles)
 {
 	Tiles = InTiles;
-	CalcHash_Now();
+	CalcHashNow();
 }
 
 void TiledBlob::SetTile(int32 X, int32 Y, BlobRef Tile)
@@ -199,7 +199,7 @@ TiledBlob& TiledBlob::operator = (const TiledBlob& RHS)
 	ReplayCount = RHS.ReplayCount;
 
 	Tiles = RHS.Tiles;
-	HashValue = RHS.HashValue;
+	SetHash(RHS.HashValue);
 
 	/// Ensure that 1x1 tiled BlobObj hashes are calculated correctly
 	check((Tiles.Rows() > 1 && Tiles.Cols() > 1) ||
@@ -282,7 +282,7 @@ void TiledBlob::CopyResolveLateBound(BlobPtr RHS_)
 
 		Desc = RHS_->GetDescriptor();
 
-		CalcHash_Now();
+		CalcHashNow();
 
 		bReady = true;
 		bTiledTarget = false;
@@ -387,7 +387,7 @@ AsyncBufferResultPtr TiledBlob::CombineTiles(bool bTouch, bool bIsArray, uint64
 
 	/// If we don't have a hash then calculate it now
 	if (!HashValue)
-		CalcHash_Now();
+		CalcHashNow();
 
 	/// We've already calculated the hash from all the child blobs
 	UE_LOG(LogData, VeryVerbose, TEXT("*** Combine TiledBlob: %s [Hash: %llu] [T0: %s] ***"), *Desc.Name, HashValue->Value(), *Tile0->Name());
@@ -458,24 +458,33 @@ AscynCHashPtr TiledBlob::CalcHash()
 	{
 		return cti::when_all(Promises.begin(), Promises.end()).then([this](std::vector<CHashPtr>) mutable
 		{
-			CalcHash_Now();
+			CalcHashNow();
 			check(HashValue->IsFinal());
 			return HashValue;
 		});
 	}
 
-	CalcHash_Now();
+	CalcHashNow();
 
 	return cti::make_ready_continuable(HashValue);
 }
 
 void TiledBlob::SetHash(CHashPtr InHashValue)
 {
-	check(!HashValue->IsFinal() || *HashValue.get() == *InHashValue.get());
-	HashValue = InHashValue;
+	check(InHashValue);
+	check(IsInGameThread() || InHashValue->IsFinal());
+
+	if (HashValue && *HashValue == *InHashValue)
+		return;
+
+	/// Update the hash in the blobber
+	if (HashValue)
+		HashValue = CHash::UpdateHash(InHashValue, HashValue);
+	else
+		HashValue = InHashValue;
 }
 
-void TiledBlob::CalcHash_Now() const
+void TiledBlob::CalcHashNow() const
 {
 	if (!Tiles.Rows() || !Tiles.Cols() || !Tiles[0][0])
 		return;
@@ -487,12 +496,12 @@ void TiledBlob::CalcHash_Now() const
 			return;
 
 		Desc = Tiles[0][0]->GetDescriptor();
-		HashValue = Tiles[0][0]->Hash();
+		const_cast<TiledBlob*>(this)->SetHash(Tiles[0][0]->Hash());
 
 		return;
 	}
 
-	HashValue = 0;
+	HashValue = nullptr;
 
 	const size_t NumTilesX = Tiles.Rows();
 	const size_t NumTilesY = Tiles.Cols();
@@ -532,13 +541,13 @@ void TiledBlob::CalcHash_Now() const
 		}
 	}
 
-	HashValue = CHash::ConstructFromSources(TileHashes);
+	const_cast<TiledBlob*>(this)->SetHash(CHash::ConstructFromSources(TileHashes));
 }
 
 CHashPtr TiledBlob::Hash() const
 {
 	if (!HashValue)
-		CalcHash_Now();
+		CalcHashNow();
 
 	return HashValue;
 }
@@ -635,7 +644,7 @@ AsyncBufferResultPtr TiledBlob::Bind(const BlobTransform* Transform, const Resou
 
 			///// If we don't have a hash then calculate it now
 			//if (!_hash)
-			//	CalcHash_Now();
+			//	CalcHashNow();
 			//
 			///// We've already calculated the hash from all the child blobs
 			//UE_LOG(LogData, Log, TEXT("*** Combine TiledBlob: %s [Hash: %llu] [T0: %s] ***"), *_desc.name, _hash->Value(), *Tiles[0][0]->Name());
@@ -1045,34 +1054,6 @@ void TiledBlob_Promise::AddLinkedBlob(BlobPtr LinkedBlob)
 	TiledBlob::AddLinkedBlob(LinkedBlob);
 }
 
-//void TiledBlob_Promise::UpdateLinkedBlobs(bool bDoFinalise)
-//{
-//	check(IsInGameThread());
-//
-//	for (BlobPtrW LinkedBlobW : LinkedBlobs)
-//	{
-//		BlobPtr LinkedBlob = LinkedBlobW.lock();
-//
-//		if (LinkedBlob)
-//		{
-//			check(LinkedBlob->IsTiled());
-//			TiledBlobPtr TiledLinkedBlob = std::static_pointer_cast<TiledBlob>(LinkedBlob);
-//			check(TiledLinkedBlob->IsPromise());
-//
-//			TiledBlob_PromisePtr TiledPromiseLinkedBlob = std::static_pointer_cast<TiledBlob_Promise>(TiledLinkedBlob);
-//			*TiledPromiseLinkedBlob = *this;
-//
-//			/// We don't want to keep recursive links
-//			TiledPromiseLinkedBlob->LinkedBlobs.clear();
-//
-//			if (bDoFinalise)
-//			{
-//				TiledPromiseLinkedBlob->FinaliseNow(false, nullptr);
-//			}
-//		}
-//	}
-//}
-
 void TiledBlob_Promise::FinaliseNow(bool bNoCalcHash, CHashPtr FixedHash)
 {
 	check(IsInGameThread());
@@ -1101,14 +1082,14 @@ void TiledBlob_Promise::FinaliseFrom(TiledBlobPtr RHS)
 	FinaliseFrom(std::static_pointer_cast<Blob>(RHS).get());
 }
 
-void TiledBlob_Promise::FinaliseFrom(const Blob* RHS)
+void TiledBlob_Promise::FinaliseFrom(Blob* RHS)
 {
 	check(RHS->IsTiled());
-	const TiledBlob* RHSTiled = static_cast<const TiledBlob*>(RHS);
+	TiledBlob* RHSTiled = static_cast<TiledBlob*>(RHS);
 
 	if (RHSTiled->IsPromise())
 	{
-		const TiledBlob_Promise* RHSTiledPromise = static_cast<const TiledBlob_Promise*>(RHS);
+		TiledBlob_Promise* RHSTiledPromise = static_cast<TiledBlob_Promise*>(RHS);
 		bMakeSingleBlob = RHSTiledPromise->bMakeSingleBlob;
 	}
 
@@ -1119,6 +1100,9 @@ AsyncBufferResultPtr TiledBlob_Promise::Finalise(bool bNoCalcHash, CHashPtr Fixe
 {
 	check(IsInGameThread());
 
+	if (bIsFinalised)
+		return cti::make_ready_continuable<BufferResultPtr>(std::make_shared<BufferResult>());
+
 	/// cannot re-finalise something that's already been finalised
 	check(!bIsFinalised);
 	bIsFinalised = true;
@@ -1200,7 +1184,7 @@ void TiledBlob_Promise::OnFinaliseInternal(BlobReadyCallback Callback) const
 	}
 
 	/// Otherwise we just queue the callbacks
-	Callbacks.push_back(Callback);
+	Callbacks.push_back(std::move(Callback));
 }
 
 void TiledBlob_Promise::SetTile(int32 TileX, int32 TileY, BlobRef InTile)
@@ -1219,7 +1203,7 @@ void TiledBlob_Promise::SetTile(int32 TileX, int32 TileY, BlobRef InTile)
 		Buffer = Tile->GetBufferRef();
 
 	/// Reset the hash
-	HashValue = 0;
+	HashValue = nullptr;
 }
 
 AsyncBufferResultPtr TiledBlob_Promise::MakeSingleBlob()
@@ -1285,9 +1269,13 @@ void TiledBlob_Promise::CopyResolveLateBound(BlobPtr RHS_)
 
 void TiledBlob_Promise::NotifyCallbacks()
 {
+	check(IsInGameThread());
+
 	/// Trigger the callbacks - Should always be the last step of finalize
-	for (auto Callback : Callbacks)
+	for (BlobReadyCallback& Callback : Callbacks)
 		Callback(this);
+
+	Callbacks.clear();
 }
 
 void TiledBlob_Promise::ResetForReplay()
diff --git a/Engine/Plugins/Experimental/TextureGraph/Source/TextureGraphEngine/Data/TiledBlob.h b/Engine/Plugins/Experimental/TextureGraph/Source/TextureGraphEngine/Data/TiledBlob.h
--- a/Engine/Plugins/Experimental/TextureGraph/Source/TextureGraphEngine/Data/TiledBlob.h
+++ b/Engine/Plugins/Experimental/TextureGraph/Source/TextureGraphEngine/Data/TiledBlob.h
@@ -27,7 +27,7 @@ protected:
 	JobPtrW							JobObj;					/// The job that is potentially generating this tiled BlobObj
 	BlobPtr							SingleBlob;			/// Single BlobObj pointer that we keep to prevent going out of Ref
 
-	void							CalcHash_Now() const;
+	void							CalcHashNow() const;
 	virtual void					ResetBuffer() override;
 
 	virtual void					TouchTiles(uint64 BatchId);
@@ -37,7 +37,7 @@ protected:
 	virtual void					SetHash(CHashPtr Hash) override;
 
 	virtual void					AddLinkedBlob(BlobPtr LinkedBlob) override;
-	virtual void					FinaliseFrom(const Blob* RHS) override;
+	virtual void					FinaliseFrom(Blob* RHS) override;
 
 public:
 	static AsyncBufferResultPtr		TileBuffer(DeviceBufferRef Buffer, BlobPtrTiles& Tiles);
@@ -163,7 +163,7 @@ protected:
 	virtual void					NotifyCallbacks();
 
 	virtual void					AddLinkedBlob(BlobPtr LinkedBlob) override;
-	virtual void					FinaliseFrom(const Blob* RHS) override;
+	virtual void					FinaliseFrom(Blob* RHS) override;
 
 public:
 									TiledBlob_Promise(TiledBlobPtr Source);
diff --git a/Engine/Plugins/Experimental/TextureGraph/Source/TextureGraphEngine/Device/DeviceNativeTask.cpp b/Engine/Plugins/Experimental/TextureGraph/Source/TextureGraphEngine/Device/DeviceNativeTask.cpp
--- a/Engine/Plugins/Experimental/TextureGraph/Source/TextureGraphEngine/Device/DeviceNativeTask.cpp
+++ b/Engine/Plugins/Experimental/TextureGraph/Source/TextureGraphEngine/Device/DeviceNativeTask.cpp
@@ -279,22 +279,25 @@ bool DeviceNativeTask::WaitFor(int32 Seconds)
 	for (auto& PrevTask : Prev)
 	{
 		/// Culled jobs should never be added even though their promise has been resolved
-		check(!PrevTask->IsCulled());
+		if (!PrevTask->IsCulled())
+		{
+			check(!PrevTask->IsCulled());
 
-		/// Must be higher TaskPriority than this particular task
-		check(PrevTask->IsHigherPriorityThan(*this));
+			/// Must be higher TaskPriority than this particular task
+			check(PrevTask->IsHigherPriorityThan(*this));
 
-		UE_LOG(LogDevice, VeryVerbose, TEXT("[%u] Begin Wait::%s"), PrevTask->GetTaskId(), *PrevTask->GetName());
-		std::future_status Status = PrevTask->GetFuture().wait_for(std::chrono::seconds(Seconds));
+			UE_LOG(LogDevice, VeryVerbose, TEXT("[%u] Begin Wait::%s"), PrevTask->GetTaskId(), *PrevTask->GetName());
+			std::future_status Status = PrevTask->GetFuture().wait_for(std::chrono::seconds(Seconds));
 
-		if (Status == std::future_status::timeout)
-		{
-			FString WaitChain = DebugWaitChain();
-			UE_LOG(LogDevice, VeryVerbose, TEXT("[%u] Timed out wait for task::%s. Chain: %s"), PrevTask->GetTaskId(), *PrevTask->GetName(), *WaitChain);
-			return false;
-		}
+			if (Status == std::future_status::timeout)
+			{
+				FString WaitChain = DebugWaitChain();
+				UE_LOG(LogDevice, VeryVerbose, TEXT("[%u] Timed out wait for task::%s. Chain: %s"), PrevTask->GetTaskId(), *PrevTask->GetName(), *WaitChain);
+				return false;
+			}
 
-		UE_LOG(LogDevice, VeryVerbose, TEXT("[%u] End Wait::%s"), PrevTask->GetTaskId(), *PrevTask->GetName());
+			UE_LOG(LogDevice, VeryVerbose, TEXT("[%u] End Wait::%s"), PrevTask->GetTaskId(), *PrevTask->GetName());
+		}
 	}
 
 	return true;
diff --git a/Engine/Plugins/Experimental/TextureGraph/Source/TextureGraphEngine/Helper/DataUtil.cpp b/Engine/Plugins/Experimental/TextureGraph/Source/TextureGraphEngine/Helper/DataUtil.cpp
--- a/Engine/Plugins/Experimental/TextureGraph/Source/TextureGraphEngine/Helper/DataUtil.cpp
+++ b/Engine/Plugins/Experimental/TextureGraph/Source/TextureGraphEngine/Helper/DataUtil.cpp
@@ -123,8 +123,8 @@ HashType DataUtil::Hash(const CHashPtrVec& InSubHashes, HashType InitialValue /*
 {
 	HashTypeVec SubHashes(InSubHashes.size());
 
-	for (size_t hi = 0; hi < InSubHashes.size(); hi++)
-		SubHashes[hi] = InSubHashes[hi]->Value();
+	for (size_t HashIndex = 0; HashIndex < InSubHashes.size(); HashIndex++)
+		SubHashes[HashIndex] = InSubHashes[HashIndex]->Value();
 
 	return Hash(SubHashes, InitialValue, Prime);
 }
@@ -201,13 +201,13 @@ void CHash::CheckLinkCycles(std::unordered_set<CHashPtr>& Chain)
 {
 	for (CHashPtrW link_ : Linked)
 	{
-		CHashPtr link = link_.lock();
+		CHashPtr Link = link_.lock();
 
-		if (link)
+		if (Link)
 		{
-			check(Chain.find(link) == Chain.end());
-			Chain.insert(link);
-			link->CheckLinkCycles(Chain);
+			check(Chain.find(Link) == Chain.end());
+			Chain.insert(Link);
+			Link->CheckLinkCycles(Chain);
 		}
 	}
 }
@@ -241,7 +241,9 @@ bool CHash::TryFinalise(HashType FinalHash /* = DataUtil::s_nullHash */, bool Up
 		check(HashValue == DataUtil::GNullHash);
 
 		/// Save the old Value of the temp HashValue
-		auto OldHash = TempHashValue->HashValue;
+		HashType OldHash = TempHashValue->HashValue;
+		IntermediateHashes = TempHashValue->GetIntermediateHashes();
+
 		bool bDidUpdate = TempHashValue->TryFinalise(FinalHash, false);
 
 		if (bDidUpdate)
@@ -256,8 +258,13 @@ bool CHash::TryFinalise(HashType FinalHash /* = DataUtil::s_nullHash */, bool Up
 
 			if (HashValue != DataUtil::GNullHash && OldHash != HashValue && UpdateBlobber)
 			{
+				CHashPtr ThisHash = shared_from_this();
+
+				IntermediateHashes.push_back(OldHash);
+
 				/// Now we try to update the mapping in blobber
-				TextureGraphEngine::GetBlobber()->UpdateHash(OldHash, shared_from_this());
+				for (HashType IntermediateHash : IntermediateHashes)
+					TextureGraphEngine::GetBlobber()->UpdateHash(IntermediateHash, ThisHash);
 			}
 
 			return true;
@@ -269,19 +276,21 @@ bool CHash::TryFinalise(HashType FinalHash /* = DataUtil::s_nullHash */, bool Up
 	/// Cannot finalise a HashValue that's already been finalised!
 	bool bShouldUpdate = false;
 	bool bDidUpdate = false;
+	HashType CurrentHash = HashValue;
+	bool bIsTemp = IsTemp();
 
 	if (!HashSources.empty())
 	{
 		HashTypeVec Sources(HashSources.size());
 		bool bIsFinalHash = true;
 
-		for (size_t hi = 0; hi < HashSources.size(); hi++)
+		for (size_t HashIndex = 0; HashIndex < HashSources.size(); HashIndex++)
 		{
 			/// Check whether the source HashValue has been finalised or not
-			bIsFinalHash &= HashSources[hi]->IsFinal();
-			Sources[hi] = HashSources[hi]->Value();
+			bIsFinalHash &= HashSources[HashIndex]->IsFinal();
+			Sources[HashIndex] = HashSources[HashIndex]->Value();
 
-			if (HashSources[hi]->GetTimestamp() > Timestamp)
+			if (HashSources[HashIndex]->GetTimestamp() > Timestamp)
 				bShouldUpdate = true;
 		}
 
@@ -322,7 +331,12 @@ bool CHash::TryFinalise(HashType FinalHash /* = DataUtil::s_nullHash */, bool Up
 	}
 
 	if (bDidUpdate)
+	{
+		if (!bIsTemp && CurrentHash != HashValue)
+			IntermediateHashes.push_back(CurrentHash);
+
 		UpdateLinks();
+	}
 
 	return bDidUpdate;
 }
@@ -343,12 +357,12 @@ void CHash::UpdateLinks()
 
 	for (CHashPtrW& Iter : Linked)
 	{ 
-		CHashPtr link = Iter.lock();
+		CHashPtr Link = Iter.lock();
 
-		if (link)
+		if (Link)
 		{
 			/// Call the linked to update itself
-			link->HandleLinkUpdated(ThisHash);
+			Link->HandleLinkUpdated(ThisHash);
 		}
 	}
 
@@ -357,13 +371,26 @@ void CHash::UpdateLinks()
 	Linked.clear();
 }
 
+HashTypeVec CHash::GetIntermediateHashes() const
+{
+	if (!TempHashValue)
+		return IntermediateHashes;
+
+	HashTypeVec CombinedIntermediateHashes = IntermediateHashes;
+	CombinedIntermediateHashes.insert(CombinedIntermediateHashes.end(), TempHashValue->IntermediateHashes.begin(), TempHashValue->IntermediateHashes.end());
+
+	return CombinedIntermediateHashes;
+}
+
 CHashPtr CHash::UpdateHash(CHashPtr NewHash, CHashPtr PrevHash)
 {
 	if (PrevHash)
 	{
 		check(NewHash != PrevHash);
 
-		auto PrevHashValue = PrevHash->Value();
+		HashType PrevHashValue = PrevHash->Value();
+		NewHash->IntermediateHashes = PrevHash->GetIntermediateHashes();
+		NewHash->IntermediateHashes.push_back(PrevHashValue);
 
 		/// Copy some Data over from the previous HashValue
 		NewHash->Linked = PrevHash->Linked;
@@ -382,7 +409,9 @@ CHashPtr CHash::UpdateHash(CHashPtr NewHash, CHashPtr PrevHash)
 
 		if (PrevHashValue != PrevHash->Value() && !TextureGraphEngine::IsDestroying() && TextureGraphEngine::GetBlobber())
 		{
-			TextureGraphEngine::GetBlobber()->UpdateHash(PrevHashValue, NewHash);
+			for (HashType IntermediateHash : NewHash->IntermediateHashes)
+				TextureGraphEngine::GetBlobber()->UpdateHash(IntermediateHash, NewHash);
+
 			NewHash->UpdateLinks();
 		}
 	}
@@ -393,7 +422,7 @@ CHashPtr CHash::UpdateHash(CHashPtr NewHash, CHashPtr PrevHash)
 void CHash::HandleLinkUpdated(CHashPtr LinkUpdated)
 {
 	/// If this has already been finalised then we don't need to do anything. This can happen
-	/// if the BlobHasher service got to a blob before this link got updated
+	/// if the BlobHasher service got to a blob before this Link got updated
 	if (IsFinal())
 		return;
 
diff --git a/Engine/Plugins/Experimental/TextureGraph/Source/TextureGraphEngine/Helper/DataUtil.h b/Engine/Plugins/Experimental/TextureGraph/Source/TextureGraphEngine/Helper/DataUtil.h
--- a/Engine/Plugins/Experimental/TextureGraph/Source/TextureGraphEngine/Helper/DataUtil.h
+++ b/Engine/Plugins/Experimental/TextureGraph/Source/TextureGraphEngine/Helper/DataUtil.h
@@ -160,8 +160,11 @@ private:
 	/// otherwise we'll have no way of connecting complex-temp hashes to the correct hash and 
 	/// determine whether a job has already been done before or not
 	CHashPtrWVec					Linked;						/// The HashValues that have this HashValue embedded in it. This is kept
-																	/// temporarily if THIS HASH is a temp HashValue. When THIS HASH gets
-																	/// updated, we can update the linked parent HashValues
+																/// temporarily if THIS HASH is a temp HashValue. When THIS HASH gets
+																/// updated, we can update the linked parent HashValues
+
+	HashTypeVec						IntermediateHashes;			/// Intermediate hashes that might have been evaluated before this
+																/// hash became finalised
 
 	FORCEINLINE bool				HasTempDependency_Internal() const { return (IsTemp() && !HashSources.empty() && TempHashValue->IsTemp()); }
 	void							AddLink(CHashPtrW Link);
@@ -178,6 +181,7 @@ public:
 
 	bool							TryFinalise(HashType FinalHash = DataUtil::GNullHash, bool UpdateBlobber = true);
 	bool							operator == (const CHash& RHS) const;
+	HashTypeVec						GetIntermediateHashes() const;
 
 	static CHashPtr					ConstructFromSources(CHashPtrVec Sources);
 	static CHashPtr					UpdateHash(CHashPtr NewHash, CHashPtr PrevHash);
diff --git a/Engine/Plugins/Experimental/TextureGraph/Source/TextureGraphEngine/Job/BlobHasherService.cpp b/Engine/Plugins/Experimental/TextureGraph/Source/TextureGraphEngine/Job/BlobHasherService.cpp
--- a/Engine/Plugins/Experimental/TextureGraph/Source/TextureGraphEngine/Job/BlobHasherService.cpp
+++ b/Engine/Plugins/Experimental/TextureGraph/Source/TextureGraphEngine/Job/BlobHasherService.cpp
@@ -79,7 +79,7 @@ AsyncJobResultPtr BlobHasherService::Tick()
 					HashType PrevHash = PrevHashes[BlobIndex];
 
 					check(CurrentHash->IsFinal());
-					UE_LOG(LogIdle_Svc, Log, TEXT("Hashed BlobObj: %s => %llu [Total Blobs: %llu]"), *BlobObj->Name(), CurrentHash->Value(), TotalCount);
+					UE_LOG(LogIdle_Svc, VeryVerbose, TEXT("Hashed BlobObj: %s => %llu [Total Blobs: %llu]"), *BlobObj->Name(), CurrentHash->Value(), TotalCount);
 
 					if (!TextureGraphEngine::IsDestroying())
 						TextureGraphEngine::GetBlobber()->UpdateBlobHash(PrevHash, BlobObj);
diff --git a/Engine/Plugins/Experimental/TextureGraph/Source/TextureGraphEngine/Transform/BlobTransform.cpp b/Engine/Plugins/Experimental/TextureGraph/Source/TextureGraphEngine/Transform/BlobTransform.cpp
--- a/Engine/Plugins/Experimental/TextureGraph/Source/TextureGraphEngine/Transform/BlobTransform.cpp
+++ b/Engine/Plugins/Experimental/TextureGraph/Source/TextureGraphEngine/Transform/BlobTransform.cpp
@@ -8,7 +8,9 @@
 
 DEFINE_LOG_CATEGORY(LogBlobTransform);
 
-BlobTransform::BlobTransform(FString Name_) : Name(!Name_.IsEmpty() ? Name_ : FString("<Unknown>"))
+BlobTransform::BlobTransform(FString Name_) 
+	: Name(!Name_.IsEmpty() ? Name_ : FString("<Unknown>"))
+	, HashValue(std::make_shared<CHash>(DataUtil::Hash_GenericString_Name(Name), true))
 {
 }
 
