Fix issues with duplicate blobs not getting bIsFinalised flag set before they send the notification callbacks.

diff --git a/Engine/Plugins/Experimental/TextureGraph/Source/TextureGraph/Private/Expressions/Output/TG_Expression_Output.cpp b/Engine/Plugins/Experimental/TextureGraph/Source/TextureGraph/Private/Expressions/Output/TG_Expression_Output.cpp
--- a/Engine/Plugins/Experimental/TextureGraph/Source/TextureGraph/Private/Expressions/Output/TG_Expression_Output.cpp
+++ b/Engine/Plugins/Experimental/TextureGraph/Source/TextureGraph/Private/Expressions/Output/TG_Expression_Output.cpp
@@ -32,19 +32,19 @@ void UTG_Expression_Output::Evaluate(FTG_EvaluationContext* InContext)
 		{
 			UpdateBufferDescriptorValues();
 
-			const BufferDescriptor OutputDesc = Output.EditTexture().GetBufferDescriptor();
-			const BufferDescriptor InputDesc = Source.GetTexture()->GetDescriptor();
+			BufferDescriptor OutputDesc = Output.EditTexture().GetBufferDescriptor();
+			const BufferDescriptor& InputDesc = Source.GetTexture()->GetDescriptor();
 
 			/// If the descriptors are not the same
 			if (OutputDesc != InputDesc)
 			{
-				RenderMaterial_FXPtr RenderMaterial = TextureGraphEngine::GetMaterialManager()->CreateMaterialOfType_FX<Fx_FullScreenCopy>(TEXT("Tex::FullScreenCopy"));
+				RenderMaterial_FXPtr RenderMaterial = TextureGraphEngine::GetMaterialManager()->CreateMaterialOfType_FX<Fx_FullScreenCopy>(TEXT("OutputCopy"));
 
 				check(RenderMaterial);
 
 				JobUPtr RenderJob = std::make_unique<Job>(InContext->Cycle->GetMix(), InContext->TargetId, std::static_pointer_cast<BlobTransform>(RenderMaterial));
 
-				auto Format = (int32)OutputDesc.Format;
+				int32 Format = (int32)OutputDesc.Format;
 
 				RenderJob
 					->AddArg(ARG_BLOB(Source.GetTexture().RasterBlob, "SourceTexture"))
@@ -52,6 +52,8 @@ void UTG_Expression_Output::Evaluate(FTG_EvaluationContext* InContext)
 					->AddArg(WithUnbounded(ARG_INT(Format, "Format")));
 
 				const FString Name = TEXT("Output");
+				OutputDesc.Name = Name;
+
 				Output.EditTexture() = RenderJob->InitResult(Name, &OutputDesc);
 				InContext->Cycle->AddJob(InContext->TargetId, std::move(RenderJob));
 			}
diff --git a/Engine/Plugins/Experimental/TextureGraph/Source/TextureGraphEngine/Data/TiledBlob.cpp b/Engine/Plugins/Experimental/TextureGraph/Source/TextureGraphEngine/Data/TiledBlob.cpp
--- a/Engine/Plugins/Experimental/TextureGraph/Source/TextureGraphEngine/Data/TiledBlob.cpp
+++ b/Engine/Plugins/Experimental/TextureGraph/Source/TextureGraphEngine/Data/TiledBlob.cpp
@@ -1071,6 +1071,9 @@ void TiledBlob_Promise::FinaliseNow(bool bNoCalcHash, CHashPtr FixedHash)
 		}
 	}
 
+	bIsFinalised = true;
+	FinaliseTS = FDateTime::Now();
+
 	/// Trigger the callbacks - Should always be the last step of finalize
 	NotifyCallbacks();
 
diff --git a/Engine/Plugins/Experimental/TextureGraph/Source/TextureGraphEngine/Transform/Layer/T_Thumbnail.cpp b/Engine/Plugins/Experimental/TextureGraph/Source/TextureGraphEngine/Transform/Layer/T_Thumbnail.cpp
--- a/Engine/Plugins/Experimental/TextureGraph/Source/TextureGraphEngine/Transform/Layer/T_Thumbnail.cpp
+++ b/Engine/Plugins/Experimental/TextureGraph/Source/TextureGraphEngine/Transform/Layer/T_Thumbnail.cpp
@@ -29,7 +29,8 @@ TiledBlobRef T_Thumbnail::Bind(UMixInterface* Mix, UObject* Model, TiledBlobPtr
 
 	UE_LOG(LogJob, VeryVerbose, TEXT("T_Thumbnail::Bind [%s]"), *InBlobToBind->Name());
 	//RenderMaterial_ThumbPtr MatThumb = TextureGraphEngine::GetMaterialManager()->CreateMaterial_Thumbnail(TEXT("T_Thumbnail"), TEXT("Util/CopyUnlit"));
-	RenderMaterial_FXPtr RenderMaterial = TextureGraphEngine::GetMaterialManager()->CreateMaterialOfType_FX<Fx_FullScreenCopy>(TEXT("Tex::FullScreenCopy"));
+	FString ThmJobName = FString::Printf(TEXT("Thm [%s]"), *InBlobToBind->GetDescriptor().Name);
+	RenderMaterial_FXPtr RenderMaterial = TextureGraphEngine::GetMaterialManager()->CreateMaterialOfType_FX<Fx_FullScreenCopy>(ThmJobName);
 
 	JobUPtr JobPtr = std::make_unique<Job>(Mix, InTargetId, std::static_pointer_cast<BlobTransform>(RenderMaterial));
 	JobPtr->AddArg(ARG_BLOB(InBlobToBind, "SourceTexture"));
