Fixed issue: FName's 1023 max length exceeded assertion

diff --git a/Engine/Plugins/Experimental/TextureGraph/Source/TextureGraph/Private/Expressions/Channel/TG_Expression_ChannelCombiner.cpp b/Engine/Plugins/Experimental/TextureGraph/Source/TextureGraph/Private/Expressions/Channel/TG_Expression_ChannelCombiner.cpp
--- a/Engine/Plugins/Experimental/TextureGraph/Source/TextureGraph/Private/Expressions/Channel/TG_Expression_ChannelCombiner.cpp
+++ b/Engine/Plugins/Experimental/TextureGraph/Source/TextureGraph/Private/Expressions/Channel/TG_Expression_ChannelCombiner.cpp
@@ -50,7 +50,7 @@ void UTG_Expression_ChannelCombiner::Evaluate(FTG_EvaluationContext* InContext)
 		->AddArg(ARG_BLOB(Alpha, "SourceAlpha"))
 		;
 
-	const FString Name = TEXT("ChannelCombiner"); // FString::Printf(TEXT("Grayscale.[%s].[%d].[%llu]"), *gmask->ID(), InContext->TargetId, InContext->Cycle->Batch()->BatchId());
+	const FString Name = FString::Printf(TEXT("ChannelCombiner.[%llu]"), InContext->Cycle->GetBatch()->GetBatchId());
 
 	BufferDescriptor Desc = Output.GetBufferDescriptor();
 	Desc.ItemsPerPoint = 4;
diff --git a/Engine/Plugins/Experimental/TextureGraph/Source/TextureGraphEngine/Job/Job.cpp b/Engine/Plugins/Experimental/TextureGraph/Source/TextureGraphEngine/Job/Job.cpp
--- a/Engine/Plugins/Experimental/TextureGraph/Source/TextureGraphEngine/Job/Job.cpp
+++ b/Engine/Plugins/Experimental/TextureGraph/Source/TextureGraphEngine/Job/Job.cpp
@@ -451,18 +451,8 @@ TiledBlobPtr Job::InitResult(FString InNewName, const BufferDescriptor* InDesire
 	DebugJobName = InNewName;
 	FString NewName = Transform->GetName();
 
-	if (NewName.IsEmpty())
+	if (!InNewName.IsEmpty())
 		NewName = InNewName;
-	else
-	{
-		for (size_t ArgIndex = 0; ArgIndex < Args.Num(); ArgIndex++)
-		{
-			const JobArgPtr& Arg = Args[ArgIndex];
-			const BufferDescriptor* ArgDesc = Arg->GetDescriptor();
-			if (ArgDesc)
-				NewName += "-" + ArgDesc->Name;
-		}
-	}
 
 	ResultDesc.bIsTransient = IsDiscard();
 	ResultDesc.Name = NewName;
diff --git a/Engine/Plugins/Experimental/TextureGraph/Source/TextureGraphEngine/Transform/Expressions/T_Transform.cpp b/Engine/Plugins/Experimental/TextureGraph/Source/TextureGraphEngine/Transform/Expressions/T_Transform.cpp
--- a/Engine/Plugins/Experimental/TextureGraph/Source/TextureGraphEngine/Transform/Expressions/T_Transform.cpp
+++ b/Engine/Plugins/Experimental/TextureGraph/Source/TextureGraphEngine/Transform/Expressions/T_Transform.cpp
@@ -64,7 +64,7 @@ TiledBlobPtr T_Transform::Create(MixUpdateCyclePtr Cycle, BufferDescriptor Desir
 		->AddArg(ARG_FLOAT(ColorParam.ShowDebugGrid, "BlendDebugGrid"))
 		;
 
-	FString Name = FString::Printf(TEXT("TransformBlit.[%llu]"), Cycle->GetBatch()->GetBatchId());
+	FString Name = FString::Printf(TEXT("T_Transform.[%llu]"), Cycle->GetBatch()->GetBatchId());
 
 	TiledBlobPtr Result = RenderJob->InitResult(Name, &DesiredOutputDesc);
 	Cycle->AddJob(TargetId, std::move(RenderJob));
