Remove batch id from job/result name

diff --git a/Engine/Plugins/Experimental/TextureGraph/Source/TextureGraph/Private/Expressions/Channel/TG_Expression_ChannelCombiner.cpp b/Engine/Plugins/Experimental/TextureGraph/Source/TextureGraph/Private/Expressions/Channel/TG_Expression_ChannelCombiner.cpp
--- a/Engine/Plugins/Experimental/TextureGraph/Source/TextureGraph/Private/Expressions/Channel/TG_Expression_ChannelCombiner.cpp
+++ b/Engine/Plugins/Experimental/TextureGraph/Source/TextureGraph/Private/Expressions/Channel/TG_Expression_ChannelCombiner.cpp
@@ -50,7 +50,7 @@ void UTG_Expression_ChannelCombiner::Evaluate(FTG_EvaluationContext* InContext)
 		->AddArg(ARG_BLOB(Alpha, "SourceAlpha"))
 		;
 
-	const FString Name = FString::Printf(TEXT("ChannelCombiner.[%llu]"), InContext->Cycle->GetBatch()->GetBatchId());
+	const FString Name = TEXT("ChannelCombiner");
 
 	BufferDescriptor Desc = Output.GetBufferDescriptor();
 	Desc.ItemsPerPoint = 4;
diff --git a/Engine/Plugins/Experimental/TextureGraph/Source/TextureGraphEngine/Transform/Expressions/T_Transform.cpp b/Engine/Plugins/Experimental/TextureGraph/Source/TextureGraphEngine/Transform/Expressions/T_Transform.cpp
--- a/Engine/Plugins/Experimental/TextureGraph/Source/TextureGraphEngine/Transform/Expressions/T_Transform.cpp
+++ b/Engine/Plugins/Experimental/TextureGraph/Source/TextureGraphEngine/Transform/Expressions/T_Transform.cpp
@@ -64,7 +64,7 @@ TiledBlobPtr T_Transform::Create(MixUpdateCyclePtr Cycle, BufferDescriptor Desir
 		->AddArg(ARG_FLOAT(ColorParam.ShowDebugGrid, "BlendDebugGrid"))
 		;
 
-	FString Name = FString::Printf(TEXT("T_Transform.[%llu]"), Cycle->GetBatch()->GetBatchId());
+	FString Name = TEXT("T_Transform");
 
 	TiledBlobPtr Result = RenderJob->InitResult(Name, &DesiredOutputDesc);
 	Cycle->AddJob(TargetId, std::move(RenderJob));
