Texture Graph Fixes:
Fix for Width and Height of output not changing from node
Fix for Node info in Node Preview when color or other non texture type is selected
Fix for Output node getting collapsed when adding a new node

diff --git a/Engine/Plugins/Experimental/TextureGraph/Source/TextureGraphEditor/Private/Pins/STG_GraphPinOutputSettings.cpp b/Engine/Plugins/Experimental/TextureGraph/Source/TextureGraphEditor/Private/Pins/STG_GraphPinOutputSettings.cpp
--- a/Engine/Plugins/Experimental/TextureGraph/Source/TextureGraphEditor/Private/Pins/STG_GraphPinOutputSettings.cpp
+++ b/Engine/Plugins/Experimental/TextureGraph/Source/TextureGraphEditor/Private/Pins/STG_GraphPinOutputSettings.cpp
@@ -27,7 +27,10 @@ void STG_GraphPinOutputSettings::Construct(const FArguments& InArgs, UEdGraphPin
 		}
 
 		GetPinObj()->bAdvancedView = ShowChildProperties();
-		GetPinObj()->GetOwningNode()->AdvancedPinDisplay = GetPinObj()->bAdvancedView ? ENodeAdvancedPins::Hidden : ENodeAdvancedPins::NoPins;
+		if (GetPinObj()->GetOwningNode()->AdvancedPinDisplay != ENodeAdvancedPins::Shown)
+		{
+			GetPinObj()->GetOwningNode()->AdvancedPinDisplay = GetPinObj()->bAdvancedView ? ENodeAdvancedPins::Hidden : ENodeAdvancedPins::NoPins;
+		}
 	}
 
 	CachedImg_Pin_BackgroundHovered = CachedImg_Pin_Background;
@@ -80,7 +83,8 @@ bool STG_GraphPinOutputSettings::CollapsibleChildProperties() const
 
 EVisibility STG_GraphPinOutputSettings::ShowLabel() const
 {
-	bool bHide = ShowChildProperties();
+	bool bOutput = GetDirection() == EEdGraphPinDirection::EGPD_Output;
+	bool bHide = ShowChildProperties() && !bOutput;
 
 	if ((GraphPinObj->GetOwningNode()->AdvancedPinDisplay == ENodeAdvancedPins::Type::Hidden && GraphPinObj->LinkedTo.Num() > 0))
 	{
diff --git a/Engine/Plugins/Experimental/TextureGraph/Source/TextureGraphEditor/Private/Pins/STG_GraphPinOutputSettingsWidget.cpp b/Engine/Plugins/Experimental/TextureGraph/Source/TextureGraphEditor/Private/Pins/STG_GraphPinOutputSettingsWidget.cpp
--- a/Engine/Plugins/Experimental/TextureGraph/Source/TextureGraphEditor/Private/Pins/STG_GraphPinOutputSettingsWidget.cpp
+++ b/Engine/Plugins/Experimental/TextureGraph/Source/TextureGraphEditor/Private/Pins/STG_GraphPinOutputSettingsWidget.cpp
@@ -375,7 +375,8 @@ void STG_GraphPinOutputSettingsWidget::GenerateStringsFromEnum(TArray<FString>&
 	}
 }
 
-void STG_GraphPinOutputSettingsWidget::GenerateValuesFromEnum(TArray<uint8>& OutEnumValues, const FString& EnumPathName) const
+template<typename T>
+void STG_GraphPinOutputSettingsWidget::GenerateValuesFromEnum(TArray<T>& OutEnumValues, const FString& EnumPathName) const
 {
 	UEnum* EnumPtr = FindObject<UEnum>(nullptr, *EnumPathName);
 	if (EnumPtr)
@@ -385,7 +386,7 @@ void STG_GraphPinOutputSettingsWidget::GenerateValuesFromEnum(TArray<uint8>& Out
 			if (!EnumPtr->HasMetaData(TEXT("Hidden"), i))
 			{
 				FString DisplayName = EnumPtr->GetDisplayNameTextByIndex(i).ToString();
-				uint8 EnumValue = EnumPtr->GetValueByIndex(i);
+				T EnumValue = EnumPtr->GetValueByIndex(i);
 				UE_LOG(LogTemp, Warning, TEXT("Enum Value: %d, Display Name: %s"), EnumValue, *DisplayName);
 				OutEnumValues.Add(EnumValue);
 			}
@@ -393,14 +394,15 @@ void STG_GraphPinOutputSettingsWidget::GenerateValuesFromEnum(TArray<uint8>& Out
 	}
 }
 
+template<typename T>
 int STG_GraphPinOutputSettingsWidget::GetValueFromIndex(const FString& EnumPathName, int Index) const
 {
 	int Value = 0;
-	TArray<uint8> EnumValues;
+	TArray<T> EnumValues;
 	UEnum* EnumPtr = FindObject<UEnum>(nullptr, *EnumPathName);
 	if (EnumPtr)
 	{
-		GenerateValuesFromEnum(EnumValues, EnumPathName);
+		GenerateValuesFromEnum<T>(EnumValues, EnumPathName);
 		if (EnumValues.Num() > Index)
 		{
 			Value = EnumValues[Index];
@@ -448,7 +450,7 @@ TSharedRef<SWidget> STG_GraphPinOutputSettingsWidget::OnGenerateWidthEnumMenu()
 void STG_GraphPinOutputSettingsWidget::HandleWidthChanged(FString Name,int Index)
 {
 	auto Settings = GetSettings();
-	Settings.Width = (EResolution)GetValueFromIndex(StaticEnum<EResolution>()->GetPathName(), Index);
+	Settings.Width = (EResolution)GetValueFromIndex<int>(StaticEnum<EResolution>()->GetPathName(), Index);
 	
 	OnOutputSettingsChanged.ExecuteIfBound(Settings);
 
@@ -488,7 +490,7 @@ TSharedRef<SWidget> STG_GraphPinOutputSettingsWidget::OnGenerateHeightEnumMenu()
 void STG_GraphPinOutputSettingsWidget::HandleHeightChanged(FString Name, int Index)
 {
 	auto Settings = GetSettings();
-	Settings.Height = (EResolution)GetValueFromIndex(StaticEnum<EResolution>()->GetPathName(), Index);
+	Settings.Height = (EResolution)GetValueFromIndex<int>(StaticEnum<EResolution>()->GetPathName(), Index);
 	OnOutputSettingsChanged.ExecuteIfBound(Settings);
 
 	SelectedHeightIndex = Index;
@@ -526,7 +528,7 @@ TSharedRef<SWidget> STG_GraphPinOutputSettingsWidget::OnGenerateFormatEnumMenu()
 void STG_GraphPinOutputSettingsWidget::HandleFormatChanged(FString Name, int Index)
 {
 	auto Settings = GetSettings();
-	Settings.TextureFormat = (ETG_TextureFormat)GetValueFromIndex(StaticEnum<ETG_TextureFormat>()->GetPathName(), Index);
+	Settings.TextureFormat = (ETG_TextureFormat)GetValueFromIndex<uint8>(StaticEnum<ETG_TextureFormat>()->GetPathName(), Index);
 
 	OnOutputSettingsChanged.ExecuteIfBound(Settings);
 
@@ -565,7 +567,7 @@ TSharedRef<SWidget> STG_GraphPinOutputSettingsWidget::OnGenerateTexturePresetTyp
 void STG_GraphPinOutputSettingsWidget::HandleTexturePresetTypeChanged(FString Name, int Index)
 {
 	auto Settings = GetSettings();
-	Settings.TexturePresetType = (ETG_TexturePresetType)GetValueFromIndex(StaticEnum<ETG_TexturePresetType>()->GetPathName(), Index);
+	Settings.TexturePresetType = (ETG_TexturePresetType)GetValueFromIndex<uint8>(StaticEnum<ETG_TexturePresetType>()->GetPathName(), Index);
 
 	Settings.OnSetTexturePresetType(Settings.TexturePresetType);
 
@@ -606,7 +608,7 @@ TSharedRef<SWidget> STG_GraphPinOutputSettingsWidget::OnGenerateLodGroupEnumMenu
 void STG_GraphPinOutputSettingsWidget::HandleLodGroupChanged(FString Name, int Index)
 {
 	auto Settings = GetSettings();
-	Settings.LODGroup = (TextureGroup)GetValueFromIndex(StaticEnum<TextureGroup>()->GetPathName(), Index);
+	Settings.LODGroup = (TextureGroup)GetValueFromIndex<int>(StaticEnum<TextureGroup>()->GetPathName(), Index);
 
 	OnOutputSettingsChanged.ExecuteIfBound(Settings);
 
@@ -650,7 +652,7 @@ TSharedRef<SWidget> STG_GraphPinOutputSettingsWidget::OnGenerateCompressionEnumM
 void STG_GraphPinOutputSettingsWidget::HandleCompressionChanged(FString Name, int Index)
 {
 	auto Settings = GetSettings();
-	Settings.Compression = (TextureCompressionSettings)GetValueFromIndex(StaticEnum<TextureCompressionSettings>()->GetPathName(), Index);
+	Settings.Compression = (TextureCompressionSettings)GetValueFromIndex<int>(StaticEnum<TextureCompressionSettings>()->GetPathName(), Index);
 
 	OnOutputSettingsChanged.ExecuteIfBound(Settings);
 
diff --git a/Engine/Plugins/Experimental/TextureGraph/Source/TextureGraphEditor/Private/Pins/STG_GraphPinOutputSettingsWidget.h b/Engine/Plugins/Experimental/TextureGraph/Source/TextureGraphEditor/Private/Pins/STG_GraphPinOutputSettingsWidget.h
--- a/Engine/Plugins/Experimental/TextureGraph/Source/TextureGraphEditor/Private/Pins/STG_GraphPinOutputSettingsWidget.h
+++ b/Engine/Plugins/Experimental/TextureGraph/Source/TextureGraphEditor/Private/Pins/STG_GraphPinOutputSettingsWidget.h
@@ -57,7 +57,11 @@ public:
 	TSharedPtr<SWidget> AddSRGBWidget();
 	FTG_OutputSettings GetSettings() const;
 	void GenerateStringsFromEnum(TArray<FString>& OutEnumNames, const FString& EnumPathName);
-	void GenerateValuesFromEnum(TArray<uint8>& OutEnumValues, const FString& EnumPathName) const;
+
+	template<typename T>
+	void GenerateValuesFromEnum(TArray<T>& OutEnumValues, const FString& EnumPathName) const;
+
+	template<typename T>
 	int GetValueFromIndex(const FString& EnumPathName, int Index) const;
 	FString GetEnumValueDisplayName(const FString& EnumPathName, int EnumValue) const;
 	EVisibility ShowParameters() const;
diff --git a/Engine/Plugins/Experimental/TextureGraph/Source/TextureGraphEditor/Private/STG_SelectionPreview.cpp b/Engine/Plugins/Experimental/TextureGraph/Source/TextureGraphEditor/Private/STG_SelectionPreview.cpp
--- a/Engine/Plugins/Experimental/TextureGraph/Source/TextureGraphEditor/Private/STG_SelectionPreview.cpp
+++ b/Engine/Plugins/Experimental/TextureGraph/Source/TextureGraphEditor/Private/STG_SelectionPreview.cpp
@@ -541,16 +541,16 @@ void STG_SelectionPreview::ConstructBlobView(BlobPtr InBlob)
 	.HAlign(HAlign_Fill)
 	[
 		SNew( STextBlock)
-		.Text(this, &STG_SelectionPreview::GetBufferDescriptionText)
+		.Text(this, &STG_SelectionPreview::GetOutputDetailsText)
 		.Font(FAppStyle::Get().GetFontStyle("PropertyWindow.NormalFont"))
 		.Justification(ETextJustify::Center)
 		.AutoWrapText(true)
 	];
 }
 
-FText STG_SelectionPreview::GetBufferDescriptionText() const
+FText STG_SelectionPreview::GetOutputDetailsText() const
 {
-	return FText::FromString(BufferDescriptionString + "\n" + PixelInfo);
+	return FText::FromString(OutputDetailsText + "\n" + PixelInfo);
 }
 
 void STG_SelectionPreview::OnTexturePreviewMouseHover()
@@ -756,6 +756,7 @@ void STG_SelectionPreview::OnSelectionChanged(UTG_EdGraphNode* InNode)
 
 							// Set Texture Properties
 							SetTextureProperties();
+							UpdateOutputDetailsText(Blob);
 						}
 					});
 			}
@@ -770,6 +771,7 @@ void STG_SelectionPreview::OnSelectionChanged(UTG_EdGraphNode* InNode)
 							SetTexture(Blob);
 							// Set Texture Properties
 							SetTextureProperties();
+							UpdateOutputDetailsText(Blob);
 						}
 					});
 			}
@@ -800,12 +802,14 @@ void STG_SelectionPreview::OnSelectionChanged(UTG_EdGraphNode* InNode)
 			// this will make sure the ViewportClient exits after
 			// clearing to the expected color
 			SetTexture(nullptr);
+			UpdateOutputDetailsText(TGPin);
 		}
 	}
 	else
 	{
 		//Null blob will set the black preview
 		SetTexture(Blob);
+		ResetOutputDetailsText();
 	}
 
 	//Update the output selection value when clicked on Graph
@@ -821,6 +825,62 @@ void STG_SelectionPreview::OnSelectionChanged(UTG_EdGraphNode* InNode)
 	OnBlobSelectionChanged.ExecuteIfBound(Blob);
 }
 
+void STG_SelectionPreview::UpdateOutputDetailsText(BlobPtr InBlob)
+{
+	if (InBlob)
+	{
+		DeviceBufferPtr Buffer = InBlob->GetBufferRef().GetPtr();
+		// Hopefully found a device buffer if not early exit
+		if (!Buffer)
+		{
+			ResetOutputDetailsText();
+			return;
+		}
+
+		OutputDetailsText = BufferToString(Buffer->Descriptor());
+	}
+	else
+	{
+		ResetOutputDetailsText();
+	}
+}
+
+void STG_SelectionPreview::UpdateOutputDetailsText(UTG_Pin* Pin)
+{
+	if (Pin->IsArgColor())
+	{
+		FLinearColor ColorValue;
+		check(Pin->GetValue(ColorValue));
+		const FColor SRGBColor = ColorValue.ToFColorSRGB();
+		OutputDetailsText = "Color Value : " + SRGBColor.ToString();
+		PixelInfo = "";
+	}
+	else if (Pin->IsArgScalar())
+	{
+		float ScalarValue = 0.0;
+		check(Pin->GetValue(ScalarValue));
+		OutputDetailsText = "Scalar Value : " + FString::Printf(TEXT("%0.3f"), ScalarValue);
+		PixelInfo = "";
+	}
+	else if (Pin->IsArgVector())
+	{
+		FVector4f VectorValue(0,0,0,0);
+		check(Pin->GetValue(VectorValue));
+		OutputDetailsText = "Vector Value : " + VectorValue.ToString();
+		PixelInfo = "";
+	}
+	else
+	{
+		ResetOutputDetailsText();
+	}
+}
+
+void STG_SelectionPreview::ResetOutputDetailsText()
+{
+	OutputDetailsText = "";
+	PixelInfo = "";
+}
+
 int STG_SelectionPreview::GetSelectedPinIndex(UTG_EdGraphNode* InNode,const UEdGraphPin* SelectedPin)
 {
 	const UTG_EdGraphSchema* Schema = Cast<const UTG_EdGraphSchema>(InNode->GetSchema());
@@ -869,7 +929,6 @@ void STG_SelectionPreview::SetTexture(BlobPtr InBlob)
 	
 	IsSingleChannel = Buffer->Descriptor().ItemsPerPoint == 1;
 	bSRGB = Buffer->Descriptor().bIsSRGB;
-	BufferDescriptionString = BufferToString(Buffer->Descriptor());
 
 	Texture = GetTextureFromBuffer(Buffer);
 }
diff --git a/Engine/Plugins/Experimental/TextureGraph/Source/TextureGraphEditor/Private/STG_SelectionPreview.h b/Engine/Plugins/Experimental/TextureGraph/Source/TextureGraphEditor/Private/STG_SelectionPreview.h
--- a/Engine/Plugins/Experimental/TextureGraph/Source/TextureGraphEditor/Private/STG_SelectionPreview.h
+++ b/Engine/Plugins/Experimental/TextureGraph/Source/TextureGraphEditor/Private/STG_SelectionPreview.h
@@ -44,12 +44,15 @@ public:
 	void OnCheckBoxStateChanged(ECheckBoxState NewState);
 	ECheckBoxState GetCheckBoxState() const;
 	void ConstructBlobView(BlobPtr InBlob = nullptr);
-	FText GetBufferDescriptionText() const;
+	FText GetOutputDetailsText() const;
 	FString BufferToString(BufferDescriptor Desc);
 	void UpdatePixelInfo();
 	void OnTexturePreviewMouseHover();
 	
 	void OnSelectionChanged(UTG_EdGraphNode* InNode);
+	void UpdateOutputDetailsText(BlobPtr InBlob);
+	void UpdateOutputDetailsText(UTG_Pin* Pin);
+	void ResetOutputDetailsText();
 	void UpdatePreview();
 
 	UTexture* GetTexture() const;
@@ -116,7 +119,7 @@ private:
 	UTG_EdGraphNode* LockedNode;
 	/** Which output should be shown */
 	FString SpecifiedOutputName = "Output";
-	FString BufferDescriptionString = "";
+	FString OutputDetailsText = "";
 	FString PixelInfo = "";
 	int32 MaxOutputs;
 	bool IsLocked = false;
