SaveAs doesn't show pop-up if nothing has changed

diff --git a/Engine/Plugins/Experimental/TextureGraph/Source/TextureGraphEditor/Private/TG_Editor.cpp b/Engine/Plugins/Experimental/TextureGraph/Source/TextureGraphEditor/Private/TG_Editor.cpp
--- a/Engine/Plugins/Experimental/TextureGraph/Source/TextureGraphEditor/Private/TG_Editor.cpp
+++ b/Engine/Plugins/Experimental/TextureGraph/Source/TextureGraphEditor/Private/TG_Editor.cpp
@@ -1552,10 +1552,9 @@ void FTG_Editor::SaveAssetAs_Execute()
 	// Extra nodes are replicated on asset save, to be saved in the underlying TextureGraph
 	ReplicateExtraNodes();
 
-	if (UpdateOriginalTextureGraph())
-	{
-		ITG_Editor::SaveAssetAs_Execute();
-	}
+	UpdateOriginalTextureGraph();
+
+	ITG_Editor::SaveAssetAs_Execute();
 }
 bool FTG_Editor::OnRequestClose(EAssetEditorCloseReason InCloseReason)
 {
