Fixed issue: Specific channel settings in SplitChannels node are not taken into account

Fix/How: Replaced source desc with output desc to handle output settings updates. in the case of auto, Args desc are handled internally.

diff --git a/Engine/Plugins/Experimental/TextureGraph/Source/TextureGraph/Private/Expressions/Channel/TG_Expression_ChannelSplitter.cpp b/Engine/Plugins/Experimental/TextureGraph/Source/TextureGraph/Private/Expressions/Channel/TG_Expression_ChannelSplitter.cpp
--- a/Engine/Plugins/Experimental/TextureGraph/Source/TextureGraph/Private/Expressions/Channel/TG_Expression_ChannelSplitter.cpp
+++ b/Engine/Plugins/Experimental/TextureGraph/Source/TextureGraph/Private/Expressions/Channel/TG_Expression_ChannelSplitter.cpp
@@ -29,10 +29,9 @@ static void GenericEvaluate(FTG_EvaluationContext* InContext, FString ChannelNam
 		->AddArg(ARG_BLOB(Source, "SourceTexture"))
 		;
 
-	const FString Name = TEXT("ChannelSplitter_") + ChannelName; // FString::Printf(TEXT("Grayscale.[%s].[%d].[%llu]"), *gmask->ID(), InContext->TargetId, InContext->Cycle->Batch()->BatchId());
+	const FString Name = TEXT("ChannelSplitter_") + ChannelName; 
 
-	BufferDescriptor Desc = Source->GetDescriptor();
-	Desc.ItemsPerPoint = 1;
+	BufferDescriptor Desc = Output.GetBufferDescriptor();
 	
 	Output = RenderJob->InitResult(Name, &Desc);
 	InContext->Cycle->AddJob(InContext->TargetId, std::move(RenderJob));
