Fixed issue: Changing the width results in height change even when height is set to auto.

Fix/How: In case when height /width is set to auto, Now we first try to find a valid height/width from the args desc. If we dont have any valid value then we try to fall back to graph level width/height value.

diff --git a/Engine/Plugins/Experimental/TextureGraph/Source/TextureGraphEngine/Job/Job.cpp b/Engine/Plugins/Experimental/TextureGraph/Source/TextureGraphEngine/Job/Job.cpp
--- a/Engine/Plugins/Experimental/TextureGraph/Source/TextureGraphEngine/Job/Job.cpp
+++ b/Engine/Plugins/Experimental/TextureGraph/Source/TextureGraphEngine/Job/Job.cpp
@@ -474,9 +474,28 @@ TiledBlobPtr Job::InitResult(FString InNewName, const BufferDescriptor* InDesire
 
 	TargetTextureSetPtr& Target = MixObj->GetSettings()->Target(TargetId);
 	if (ResultDesc.Width <= 0)
-		ResultDesc.Width = Target->GetWidth();
+	{
+		if(ArgsDescCombined.Width <= 0)
+		{
+			ResultDesc.Width = Target->GetWidth();
+		}
+		else
+		{
+			ResultDesc.Width = ArgsDescCombined.Width;			
+		}
+		
+	}
 	if (ResultDesc.Height <= 0)
-		ResultDesc.Height = Target->GetHeight();
+	{
+		if(ArgsDescCombined.Height <= 0)
+		{
+			ResultDesc.Height = Target->GetHeight();
+		}
+		else
+		{
+			ResultDesc.Height = ArgsDescCombined.Height;	
+		}
+	}
 
 	if (!NumTilesX)
 		NumTilesX = MixObj->GetNumXTiles();
