Fixed issue: Changing the height settings from texture descriptor is not working when width is set to auto

Fix/How: Removed the dependency from width which fixes the issue.

diff --git a/Engine/Plugins/Experimental/TextureGraph/Source/TextureGraphEngine/Data/RawBuffer.cpp b/Engine/Plugins/Experimental/TextureGraph/Source/TextureGraphEngine/Data/RawBuffer.cpp
--- a/Engine/Plugins/Experimental/TextureGraph/Source/TextureGraphEngine/Data/RawBuffer.cpp
+++ b/Engine/Plugins/Experimental/TextureGraph/Source/TextureGraphEngine/Data/RawBuffer.cpp
@@ -276,6 +276,10 @@ BufferDescriptor BufferDescriptor::CombineWithPreference(const BufferDescriptor*
 		if (OverrideDesc->Width)
 		{
 			CombinedDesc.Width = OverrideDesc->Width;
+		}
+
+		if (OverrideDesc->Height)
+		{
 			CombinedDesc.Height = OverrideDesc->Height;
 		}
 
@@ -295,6 +299,10 @@ BufferDescriptor BufferDescriptor::CombineWithPreference(const BufferDescriptor*
 		if (CombinedDesc.Width <= 0 && RefDesc->Width > 0)
 		{
 			CombinedDesc.Width = RefDesc->Width;
+		}
+
+		if (CombinedDesc.Height <= 0 && RefDesc->Height > 0)
+		{
 			CombinedDesc.Height = RefDesc->Height;
 		}
 
