FIxing math nodes outputing a 1x1 black texture if only one variant argument is connected
all nodes deriving from the Expression_variant would expose the same behavior.
By removing the early exit we allow each individual expressions to customize the expected behavior in the virtual EvaluateTexture(...) method

Now all the math nodes will show correctly the result and the thumbnail
Changed the default value of Pow exponent to 1

diff --git a/Engine/Plugins/Experimental/TextureGraph/Source/TextureGraph/Private/Expressions/TG_Expression_Variant.cpp b/Engine/Plugins/Experimental/TextureGraph/Source/TextureGraph/Private/Expressions/TG_Expression_Variant.cpp
--- a/Engine/Plugins/Experimental/TextureGraph/Source/TextureGraph/Private/Expressions/TG_Expression_Variant.cpp
+++ b/Engine/Plugins/Experimental/TextureGraph/Source/TextureGraph/Private/Expressions/TG_Expression_Variant.cpp
@@ -110,20 +110,6 @@ void UTG_Expression_Variant::Evaluate(FTG_EvaluationContext* InContext)
 	}
 	case FTG_Variant::EType::Texture:
 	{
-			// If any of the arguments are not connected then just return black
-		if (ErrorCheckInputTextures())
-		{
-			std::vector<FTG_Variant> Args = GetEvaluateArgs();
-			for (FTG_Variant& Arg : Args)
-			{
-				if (!Arg.IsTexture() || !Arg.GetTexture())
-				{
-					Output.EditTexture() = FTG_Texture::GetBlack();
-					return;
-				}
-			}
-		}
-
 		Output.EditTexture() = EvaluateTexture(InContext);
 		break;
 	}
diff --git a/Engine/Plugins/Experimental/TextureGraph/Source/TextureGraph/Public/Expressions/Maths/TG_Expression_Maths_TwoInputs.h b/Engine/Plugins/Experimental/TextureGraph/Source/TextureGraph/Public/Expressions/Maths/TG_Expression_Maths_TwoInputs.h
--- a/Engine/Plugins/Experimental/TextureGraph/Source/TextureGraph/Public/Expressions/Maths/TG_Expression_Maths_TwoInputs.h
+++ b/Engine/Plugins/Experimental/TextureGraph/Source/TextureGraph/Public/Expressions/Maths/TG_Expression_Maths_TwoInputs.h
@@ -139,11 +139,11 @@ class TEXTUREGRAPH_API UTG_Expression_Pow : public UTG_Expression_Variant
 	
 	// The base of the power function
 	UPROPERTY(meta = (TGType = "TG_Input"))
-	FTG_Variant							Base;
+	FTG_Variant							Base = 0;
 
 	// The exponent
 	UPROPERTY(meta = (TGType = "TG_Input"))
-	FTG_Variant							Exponent = 0;
+	FTG_Variant							Exponent = 1; // Default exponent to 1 for a noop behavior
 
 public:
 	virtual FText						GetTooltipText() const override { return FText::FromString(TEXT("Calculates Pow(Base, Exponent).")); } 
