COnversion of simple type constants to flat texture needs to account for expected precision
Scalar and Vector are converted to a half format

Also address an inconsistancy with MaterialID output not writting the full RGB values

expose the Log button on devleopment build

diff --git a/Engine/Plugins/Experimental/TextureGraph/Shaders/Expressions/Expression_MaterialID.usf b/Engine/Plugins/Experimental/TextureGraph/Shaders/Expressions/Expression_MaterialID.usf
--- a/Engine/Plugins/Experimental/TextureGraph/Shaders/Expressions/Expression_MaterialID.usf
+++ b/Engine/Plugins/Experimental/TextureGraph/Shaders/Expressions/Expression_MaterialID.usf
@@ -143,7 +143,8 @@ float3 SampleMaterialIDTex(float2 UV)
 float4 FSH_MaterialIDMask(in float2 UvSource : TEXCOORD0) : SV_Target0
 {
 	float3 RGB;
-	float4 Mask = float4(MatchAtPoint(UvSource, RGB),0,0,1);
+	float MaskValue = MatchAtPoint(UvSource, RGB);
+	float4 Mask = float4(MaskValue, MaskValue, MaskValue, 1.0);
 	
 	return Mask;
 
diff --git a/Engine/Plugins/Experimental/TextureGraph/Source/TextureGraph/Private/TG_GraphEvaluation.cpp b/Engine/Plugins/Experimental/TextureGraph/Source/TextureGraph/Private/TG_GraphEvaluation.cpp
--- a/Engine/Plugins/Experimental/TextureGraph/Source/TextureGraph/Private/TG_GraphEvaluation.cpp
+++ b/Engine/Plugins/Experimental/TextureGraph/Source/TextureGraph/Private/TG_GraphEvaluation.cpp
@@ -100,11 +100,32 @@ const FString FTG_Evaluation::GVectorToTextureAutoConv_Name = TEXT("_Auto_Conv_V
 const FString FTG_Evaluation::GColorToTextureAutoConv_Name = TEXT("_Auto_Conv_LinearColor_To_Tex_");
 const FString FTG_Evaluation::GFloatToTextureAutoConv_Name = TEXT("_Auto_Conv_Float_To_Tex_");
 
+// Produce a BufferDescriptor ideal to store a constant value of the type specified by Variant Type
+// The texture generated with the Descriptor will contains enough precision for the consstant to save
+BufferDescriptor GetFlatColorDesc(ETG_VariantType InVariantType)
+{
+	switch (InVariantType)
+	{
+	case ETG_VariantType::Scalar:
+		return T_FlatColorTexture::GetFlatColorDesc(FTG_Evaluation::GFloatToTextureAutoConv_Name, BufferFormat::Half);
+
+	case ETG_VariantType::Color:
+		return T_FlatColorTexture::GetFlatColorDesc(FTG_Evaluation::GColorToTextureAutoConv_Name, BufferFormat::Byte);
+
+	case ETG_VariantType::Vector:
+		return T_FlatColorTexture::GetFlatColorDesc(FTG_Evaluation::GVectorToTextureAutoConv_Name, BufferFormat::Half);
+
+	default:
+		return  BufferDescriptor();
+	}
+}
+
+
 void FloatToFTG_Texture_Converter(FTG_Evaluation::VarConverterInfo& Info)
 {
 	auto Input = Info.InVar->GetAs<float>();
 	auto& Output = Info.OutVar->EditAs<FTG_Texture>();
-	BufferDescriptor Desc = T_FlatColorTexture::GetFlatColorDesc(FTG_Evaluation::GFloatToTextureAutoConv_Name);
+	BufferDescriptor Desc = GetFlatColorDesc(ETG_VariantType::Scalar);
 	Output = T_FlatColorTexture::Create(Info.Context->Cycle, Desc, FLinearColor(Input, Input, Input), Info.Context->TargetId);
 }
 
@@ -112,7 +133,7 @@ void FLinearColorToFTG_Texture_Converter(FTG_Evaluation::VarConverterInfo& Info)
 {
 	auto Input = Info.InVar->GetAs<FLinearColor>();
 	auto& Output = Info.OutVar->EditAs<FTG_Texture>();
-	BufferDescriptor Desc = T_FlatColorTexture::GetFlatColorDesc(FTG_Evaluation::GColorToTextureAutoConv_Name);
+	BufferDescriptor Desc = GetFlatColorDesc(ETG_VariantType::Color);
 	Output = T_FlatColorTexture::Create(Info.Context->Cycle, Desc, Input, Info.Context->TargetId);
 }
 
@@ -120,7 +141,7 @@ void FVector4fToFTG_Texture_Converter(FTG_Evaluation::VarConverterInfo& Info)
 {
 	auto Input = Info.InVar->GetAs<FVector4f>();
 	auto& Output = Info.OutVar->EditAs<FTG_Texture>();
-	BufferDescriptor Desc = T_FlatColorTexture::GetFlatColorDesc(FTG_Evaluation::GVectorToTextureAutoConv_Name);
+	BufferDescriptor Desc = GetFlatColorDesc(ETG_VariantType::Vector);
 	Output = T_FlatColorTexture::Create(Info.Context->Cycle, Desc, FLinearColor(Input.X, Input.Y, Input.Z, Input.W), Info.Context->TargetId);
 }
 
@@ -144,7 +165,7 @@ void FloatToFTG_Variant_Converter(FTG_Evaluation::VarConverterInfo& Info)
 		Output.Data.Set<FVector4f>(FVector4f((float)Input));
 		break;
 	case FTG_Variant::EType::Texture:
-		BufferDescriptor Desc = T_FlatColorTexture::GetFlatColorDesc(TEXT("_Auto_Conv_Int_To_Tex_"));
+		BufferDescriptor Desc = GetFlatColorDesc(ETG_VariantType::Scalar);
 		auto texture = FTG_Texture(
 			T_FlatColorTexture::Create(Info.Context->Cycle, Desc,
 				FLinearColor((float)Input, (float)Input, (float)Input),
@@ -171,7 +192,7 @@ void FLinearColorToFTG_Variant_Converter(FTG_Evaluation::VarConverterInfo& Info)
 		Output.Data.Set<FVector4f>(FVector4f(Input.R, Input.G, Input.B, Input.A));
 		break;
 	case FTG_Variant::EType::Texture:
-		BufferDescriptor Desc = T_FlatColorTexture::GetFlatColorDesc(TEXT("_Auto_Conv_Color_To_Tex_"));
+		BufferDescriptor Desc = GetFlatColorDesc(ETG_VariantType::Color);
 		auto texture = FTG_Texture(
 			T_FlatColorTexture::Create(Info.Context->Cycle, Desc,
 				Input,
@@ -196,7 +217,7 @@ void FVector4fToFTG_Variant_Converter(FTG_Evaluation::VarConverterInfo& Info)
 		Output.Data.Set<FVector4f>(Input);
 		break;
 	case FTG_Variant::EType::Texture:
-		BufferDescriptor Desc = T_FlatColorTexture::GetFlatColorDesc(TEXT("_Auto_Conv_Color_To_Tex_"));
+		BufferDescriptor Desc = GetFlatColorDesc(ETG_VariantType::Vector);
 		auto texture = FTG_Texture(
 			T_FlatColorTexture::Create(Info.Context->Cycle, Desc,
 				FLinearColor(Input.X, Input.Y, Input.Z, Input.W),
@@ -312,7 +333,7 @@ void FTG_VariantToFTG_Texture_Converter(FTG_Evaluation::VarConverterInfo& Info)
 		Color = FLinearColor(Input.Data.Get<FVector4f>().X, Input.Data.Get<FVector4f>().Y, Input.Data.Get<FVector4f>().Z, Input.Data.Get<FVector4f>().W);
 		break;
 	}
-	BufferDescriptor Desc = T_FlatColorTexture::GetFlatColorDesc(TEXT("_Auto_Conv_Variant_To_Tex_"));
+	BufferDescriptor Desc = GetFlatColorDesc(SourceType);
 	Output = T_FlatColorTexture::Create(Info.Context->Cycle, Desc, Color, Info.Context->TargetId);
 
 }
diff --git a/Engine/Plugins/Experimental/TextureGraph/Source/TextureGraph/Private/TG_Variant.cpp b/Engine/Plugins/Experimental/TextureGraph/Source/TextureGraph/Private/TG_Variant.cpp
--- a/Engine/Plugins/Experimental/TextureGraph/Source/TextureGraph/Private/TG_Variant.cpp
+++ b/Engine/Plugins/Experimental/TextureGraph/Source/TextureGraph/Private/TG_Variant.cpp
@@ -131,23 +131,27 @@ FTG_Texture FTG_Variant::GetTexture(FTG_EvaluationContext* InContext, FTG_Textur
 
 	FLinearColor Color;
 	FString Name;
+	BufferFormat TexelFormat = BufferFormat::Byte;
 
 	if (IsColor())
 	{
 		Name = FTG_Evaluation::GColorToTextureAutoConv_Name;
 		Color = Data.Get<FLinearColor>();
+		TexelFormat = BufferFormat::Byte;
 	}
 	else if (IsVector())
 	{
 		Name = FTG_Evaluation::GColorToTextureAutoConv_Name;
 		FVector4f Vec = Data.Get<FVector4f>();
 		Color = FLinearColor((float)Vec.X, (float)Vec.Y, (float)Vec.Z, (float)Vec.W);
+		TexelFormat = BufferFormat::Half;
 	}
 	else if (IsScalar())
 	{
 		Name = FTG_Evaluation::GFloatToTextureAutoConv_Name;
 		float Value = Data.Get<float>();
 		Color = FLinearColor(Value, Value, Value, Value);
+		TexelFormat = BufferFormat::Half;
 	}
 	else
 	{
@@ -177,7 +181,7 @@ FTG_Texture FTG_Variant::GetTexture(FTG_EvaluationContext* InContext, FTG_Textur
 			return FTG_Texture::GetMagenta();
 	}
 
-	BufferDescriptor Desc = !DesiredDesc ? T_FlatColorTexture::GetFlatColorDesc(Name) : *DesiredDesc;
+	BufferDescriptor Desc = !DesiredDesc ? T_FlatColorTexture::GetFlatColorDesc(Name, TexelFormat) : *DesiredDesc;
 	TiledBlobPtr Output = T_FlatColorTexture::Create(InContext->Cycle, Desc, Color, InContext->TargetId);
 
 	return Output;
diff --git a/Engine/Plugins/Experimental/TextureGraph/Source/TextureGraphEngine/Transform/Expressions/T_FlatColorTexture.cpp b/Engine/Plugins/Experimental/TextureGraph/Source/TextureGraphEngine/Transform/Expressions/T_FlatColorTexture.cpp
--- a/Engine/Plugins/Experimental/TextureGraph/Source/TextureGraphEngine/Transform/Expressions/T_FlatColorTexture.cpp
+++ b/Engine/Plugins/Experimental/TextureGraph/Source/TextureGraphEngine/Transform/Expressions/T_FlatColorTexture.cpp
@@ -12,12 +12,12 @@
 
 IMPLEMENT_GLOBAL_SHADER(FSH_FlatColorTexture, "/Plugin/TextureGraph/Layer/ChannelSourceFlat.usf", "FSH_FlatColorTexture", SF_Pixel);
 
-BufferDescriptor T_FlatColorTexture::GetFlatColorDesc(FString name)
+BufferDescriptor T_FlatColorTexture::GetFlatColorDesc(FString name, BufferFormat InBufferFormat)
 {
 	BufferDescriptor Desc;
 	
 	Desc.Name = name;
-	Desc.Format = BufferFormat::Byte;
+	Desc.Format = InBufferFormat;
 	Desc.ItemsPerPoint = 4;
 	Desc.Width = 1;
 	Desc.Height = 1;
@@ -41,8 +41,12 @@ TiledBlobPtr T_FlatColorTexture::Create(MixUpdateCyclePtr InCycle, BufferDescrip
 	Desc.ItemsPerPoint = 4;
 
 	BufferDescriptor OutputDesc = BufferDescriptor::Combine(Desc, DesiredOutputDesc);
-	OutputDesc.DefaultValue = Color; //
-	OutputDesc.bIsSRGB = true;
+	if (Desc.Format == BufferFormat::Byte) // Byte size automatically enable sRGB
+	{
+		OutputDesc.bIsSRGB = true;
+	}
+
+	OutputDesc.DefaultValue = Color;
 
 	RenderJob->AddArg(ARG_LINEAR_COLOR(Color, "Color"))
 		 ->AddArg(WithUnbounded(ARG_INT(OutputDesc.ItemsPerPoint,"ItemsPerPoint")))
diff --git a/Engine/Plugins/Experimental/TextureGraph/Source/TextureGraphEngine/Transform/Expressions/T_FlatColorTexture.h b/Engine/Plugins/Experimental/TextureGraph/Source/TextureGraphEngine/Transform/Expressions/T_FlatColorTexture.h
--- a/Engine/Plugins/Experimental/TextureGraph/Source/TextureGraphEngine/Transform/Expressions/T_FlatColorTexture.h
+++ b/Engine/Plugins/Experimental/TextureGraph/Source/TextureGraphEngine/Transform/Expressions/T_FlatColorTexture.h
@@ -34,6 +34,6 @@ public:
 	//////////////////////////////////////////////////////////////////////////
 	/// Static functions
 	//////////////////////////////////////////////////////////////////////////
-	static BufferDescriptor			GetFlatColorDesc(FString name);
+	static BufferDescriptor			GetFlatColorDesc(FString name, BufferFormat InBufferFormat = BufferFormat::Byte);
 	static TiledBlobPtr				Create(MixUpdateCyclePtr InCycle, BufferDescriptor DesiredOutputDesc, FLinearColor Color, int InTargetId);
 };
