Fix warning /  typo issue 422937

diff --git a/Engine/Plugins/Experimental/TextureGraph/Source/TextureGraphEngine/Transform/Expressions/T_FlatColorTexture.cpp b/Engine/Plugins/Experimental/TextureGraph/Source/TextureGraphEngine/Transform/Expressions/T_FlatColorTexture.cpp
--- a/Engine/Plugins/Experimental/TextureGraph/Source/TextureGraphEngine/Transform/Expressions/T_FlatColorTexture.cpp
+++ b/Engine/Plugins/Experimental/TextureGraph/Source/TextureGraphEngine/Transform/Expressions/T_FlatColorTexture.cpp
@@ -41,7 +41,7 @@ TiledBlobPtr T_FlatColorTexture::Create(MixUpdateCyclePtr InCycle, BufferDescrip
 	Desc.ItemsPerPoint = 4;
 
 	BufferDescriptor OutputDesc = BufferDescriptor::Combine(Desc, DesiredOutputDesc);
-	if (Desc.Format == BufferFormat::Byte) // Byte size automatically enable sRGB
+	if (OutputDesc.Format == BufferFormat::Byte) // Byte size automatically enable sRGB
 	{
 		OutputDesc.bIsSRGB = true;
 	}
