Fix crash when undoing a material pin connection to a texture

diff --git a/Engine/Plugins/Experimental/TextureGraph/Source/TextureGraph/Private/Expressions/TG_Expression_MaterialBase.cpp b/Engine/Plugins/Experimental/TextureGraph/Source/TextureGraph/Private/Expressions/TG_Expression_MaterialBase.cpp
--- a/Engine/Plugins/Experimental/TextureGraph/Source/TextureGraph/Private/Expressions/TG_Expression_MaterialBase.cpp
+++ b/Engine/Plugins/Experimental/TextureGraph/Source/TextureGraph/Private/Expressions/TG_Expression_MaterialBase.cpp
@@ -109,7 +109,7 @@ TiledBlobPtr UTG_Expression_MaterialBase::CreateRenderMaterialJob(FTG_Evaluation
 
 	Desc.DefaultValue = FLinearColor::Black;
 
-	const TiledBlobPtr RefBlob = LinkMaterialParameters(InContext, MaterialJob, InRenderMaterial->GetMaterial(), Desc);
+	LinkMaterialParameters(InContext, MaterialJob, InRenderMaterial->GetMaterial(), Desc);
 
 	const TiledBlob_PromisePtr MaterialResult = std::static_pointer_cast<TiledBlob_Promise>(MaterialJob->InitResult(InRenderMaterial->GetName(), &Desc));
 	MaterialJob->AddArg(WithUnbounded(ARG_BOOL(TiledMode, "TiledMode")));
@@ -130,10 +130,8 @@ TiledBlobPtr UTG_Expression_MaterialBase::CreateRenderMaterialJob(FTG_Evaluation
 
 }
 
-TiledBlobPtr UTG_Expression_MaterialBase::LinkMaterialParameters(FTG_EvaluationContext* InContext, JobUPtr& InMaterialJob, const UMaterial* InMaterial, BufferDescriptor InDescriptor)
+void UTG_Expression_MaterialBase::LinkMaterialParameters(FTG_EvaluationContext* InContext, JobUPtr& InMaterialJob, const UMaterial* InMaterial, BufferDescriptor InDescriptor)
 {
-	TiledBlobPtr Ref = nullptr;
-
 	if (InMaterialJob)
 	{	
 		TArray<FMaterialParameterInfo> OutParameterInfo;
@@ -178,19 +176,19 @@ TiledBlobPtr UTG_Expression_MaterialBase::LinkMaterialParameters(FTG_EvaluationC
 			if (Var && !Var->IsEmpty())
 			{
 				FTG_Texture& ParamValue = Var->EditAs<FTG_Texture>();
-				if (!Ref)
-					Ref = ParamValue;
-#define WITH_COMBINER_Job 1
-#ifdef WITH_COMBINER_Job
-				auto CombinedBlob = T_CombineTiledBlob::Create(InContext->Cycle, ParamValue.GetBufferDescriptor(), 0, ParamValue.RasterBlob);
-
-				auto ArgBlob = ARG_BLOB(CombinedBlob, TCHAR_TO_UTF8(*ParameterInfo.Name.ToString()));
-#else
-				auto ArgBlob = ARG_BLOB(Ref, TCHAR_TO_UTF8(*ParameterInfo.Name.ToString()));
-#endif	
-				ArgBlob->SetHandleTiles(TiledMode);
+
+				// there could be a case where the var has an empty blob, we don't need to do anything
+				// we fallback to the default internal material parameter value for this input pin
+				if (ParamValue.RasterBlob)
+				{
+					auto CombinedBlob = T_CombineTiledBlob::Create(InContext->Cycle, ParamValue.GetBufferDescriptor(), 0, ParamValue.RasterBlob);
+
+					auto ArgBlob = ARG_BLOB(CombinedBlob, TCHAR_TO_UTF8(*ParameterInfo.Name.ToString()));
+	
+					ArgBlob->SetHandleTiles(TiledMode);
 				
-				InMaterialJob->AddArg(ArgBlob);
+					InMaterialJob->AddArg(ArgBlob);
+				}
 			}
 		}
 
@@ -233,8 +231,6 @@ TiledBlobPtr UTG_Expression_MaterialBase::LinkMaterialParameters(FTG_EvaluationC
 			}
 		}*/
 	}
-
-	return Ref;
 }
 
 void UTG_Expression_MaterialBase::AddSignatureParam(TArray<FMaterialParameterInfo> OutParameterInfo, FName CPPTypeName,  
@@ -365,4 +361,4 @@ void UTG_Expression_MaterialBase::CopyVarGeneric(const FTG_Argument& Arg, FTG_Va
 		}
 	}
 	
-}
+}
\ No newline at end of file
diff --git a/Engine/Plugins/Experimental/TextureGraph/Source/TextureGraph/Public/Expressions/TG_Expression_MaterialBase.h b/Engine/Plugins/Experimental/TextureGraph/Source/TextureGraph/Public/Expressions/TG_Expression_MaterialBase.h
--- a/Engine/Plugins/Experimental/TextureGraph/Source/TextureGraph/Public/Expressions/TG_Expression_MaterialBase.h
+++ b/Engine/Plugins/Experimental/TextureGraph/Source/TextureGraph/Public/Expressions/TG_Expression_MaterialBase.h
@@ -65,7 +65,7 @@ protected:
 	
 	TiledBlobPtr CreateRenderMaterialJob(FTG_EvaluationContext* InContext, const FString& InName, const FString& InMaterialPath, const BufferDescriptor& InDescriptor, EDrawMaterialAttributeTarget InDrawMaterialAttributeTarget);
 	TiledBlobPtr CreateRenderMaterialJob(FTG_EvaluationContext* InContext, const RenderMaterial_BPPtr& InRenderMaterial, const BufferDescriptor& InDescriptor, EDrawMaterialAttributeTarget InDrawMaterialAttributeTarget);
-	TiledBlobPtr LinkMaterialParameters(FTG_EvaluationContext* InContext, JobUPtr& InMaterialJob, const UMaterial* InMaterial, BufferDescriptor InDescriptor);
+	void LinkMaterialParameters(FTG_EvaluationContext* InContext, JobUPtr& InMaterialJob, const UMaterial* InMaterial, BufferDescriptor InDescriptor);
 
 	virtual void CopyVarGeneric(const FTG_Argument& Arg, FTG_Var* InVar, bool CopyVarToArg);
 	virtual void SetMaterialInternal(UMaterialInterface* InMaterial);
