Remove MixerMeshComponent and Mixer word from the texture graph plugin.

diff --git a/Engine/Plugins/Experimental/TextureGraph/Source/TextureGraph/Public/TG_GraphEvaluation.h b/Engine/Plugins/Experimental/TextureGraph/Source/TextureGraph/Public/TG_GraphEvaluation.h
--- a/Engine/Plugins/Experimental/TextureGraph/Source/TextureGraph/Public/TG_GraphEvaluation.h
+++ b/Engine/Plugins/Experimental/TextureGraph/Source/TextureGraph/Public/TG_GraphEvaluation.h
@@ -50,7 +50,7 @@ struct TEXTUREGRAPH_API FTG_VarMap
 // 
 struct TEXTUREGRAPH_API FTG_EvaluationContext
 {
-	// The cycle used to call in the concrete Mixer engine system
+	// The cycle used to call in the concrete TextureGraph engine system
 	MixUpdateCyclePtr	Cycle;
 
 	// in and out buckets of vars passed as the arguments to the expression evaluate call
diff --git a/Engine/Plugins/Experimental/TextureGraph/Source/TextureGraphEditor/Private/TG_Editor.h b/Engine/Plugins/Experimental/TextureGraph/Source/TextureGraphEditor/Private/TG_Editor.h
--- a/Engine/Plugins/Experimental/TextureGraph/Source/TextureGraphEditor/Private/TG_Editor.h
+++ b/Engine/Plugins/Experimental/TextureGraph/Source/TextureGraphEditor/Private/TG_Editor.h
@@ -207,8 +207,6 @@ private:
 
 	void											OnFinishedChangingOutput(const FPropertyChangedEvent& PropertyChangedEvent);
 
-	void											CreateMixerUI(UMixInterface* mix);
-
 	void											UpdateParametersUI();
 
 	/*Get the Current selected folder from content browser*/
diff --git a/Engine/Plugins/Experimental/TextureGraph/Source/TextureGraphEngine/Job/Job.cpp b/Engine/Plugins/Experimental/TextureGraph/Source/TextureGraphEngine/Job/Job.cpp
--- a/Engine/Plugins/Experimental/TextureGraph/Source/TextureGraphEngine/Job/Job.cpp
+++ b/Engine/Plugins/Experimental/TextureGraph/Source/TextureGraphEngine/Job/Job.cpp
@@ -753,7 +753,7 @@ AsyncPrepareResult Job::PrepareTargets(JobBatch* Batch)
 	}
 
 	/// Keep this consistent across all the different types. Inconsistent sizes of these arrays
-	/// was bowling over Mixer Insight. Consistency is better.
+	/// was bowling over TextureGraph Insight. Consistency is better.
 	TileResults.Resize(ResultPromise->Rows(), ResultPromise->Cols());
 	TileInvalidationMatrix.Resize(ResultPromise->Rows(), ResultPromise->Cols());
 
diff --git a/Engine/Plugins/Experimental/TextureGraph/Source/TextureGraphEngine/Job/JobBatch.cpp b/Engine/Plugins/Experimental/TextureGraph/Source/TextureGraphEngine/Job/JobBatch.cpp
--- a/Engine/Plugins/Experimental/TextureGraph/Source/TextureGraphEngine/Job/JobBatch.cpp
+++ b/Engine/Plugins/Experimental/TextureGraph/Source/TextureGraphEngine/Job/JobBatch.cpp
@@ -220,7 +220,7 @@ AsyncJobResultPtr JobBatch::Exec(OnAllJobsDoneCallback callback)
 {
 	check(IsInGameThread());
 
-	/// This Batch is starting to execute ... it will be locked forever (until reset for Mixer Insight)
+	/// This Batch is starting to execute ... it will be locked forever (until reset for TextureGraph Insight)
 	bIsLocked = true;
 
 	SCOPE_CYCLE_COUNTER(STAT_JobBatch_Exec);
diff --git a/Engine/Plugins/Experimental/TextureGraph/Source/TextureGraphEngine/Model/Mix/Mix.h b/Engine/Plugins/Experimental/TextureGraph/Source/TextureGraphEngine/Model/Mix/Mix.h
--- a/Engine/Plugins/Experimental/TextureGraph/Source/TextureGraphEngine/Model/Mix/Mix.h
+++ b/Engine/Plugins/Experimental/TextureGraph/Source/TextureGraphEngine/Model/Mix/Mix.h
@@ -3,7 +3,6 @@
 
 #include "CoreMinimal.h" 
 #include "MixInterface.h"
-#include "MixerMeshComponent.h"
 
 #include <memory>
 #include <vector> 
@@ -12,7 +11,6 @@
 
 DECLARE_LOG_CATEGORY_EXTERN(LogMix, All, All);
 
-class UMixerProject;
 class UMixParameters;
 
 class RenderMesh;
@@ -35,17 +33,7 @@ private:
 public:
 	static UMix*						NullMix();
 
-protected:	
-	UPROPERTY()
-	TObjectPtr<UMixerMeshComponent>		MeshComponent;
-
 public:
 	virtual								~UMix() override;
-	
 	virtual void						Update(MixUpdateCyclePtr Cycle) override;
-	
-	//////////////////////////////////////////////////////////////////////////
-	/// Inline Functions
-	//////////////////////////////////////////////////////////////////////////
-	virtual UMixerMeshComponent*		GetMeshComponent() const override { return MeshComponent; }
 };
diff --git a/Engine/Plugins/Experimental/TextureGraph/Source/TextureGraphEngine/Model/Mix/MixInterface.cpp b/Engine/Plugins/Experimental/TextureGraph/Source/TextureGraphEngine/Model/Mix/MixInterface.cpp
--- a/Engine/Plugins/Experimental/TextureGraph/Source/TextureGraphEngine/Model/Mix/MixInterface.cpp
+++ b/Engine/Plugins/Experimental/TextureGraph/Source/TextureGraphEngine/Model/Mix/MixInterface.cpp
@@ -5,8 +5,6 @@
 #include "MixManager.h"
 #include <GameFramework/Actor.h>
 
-#include "MixerMeshComponent.h"
-
 const TMap<FString, FString> UMixInterface::s_uriAlias = UMixInterface::InitURIAlias();
 const TMap<FString, FString> UMixInterface::s_uriAliasDecryptor = UMixInterface::InitURIAliasDecryptor();
 
@@ -111,10 +109,7 @@ UMixSettings* UMixInterface::GetSettings() const
 
 void UMixInterface::SetMesh(RenderMeshPtr MeshObj, int MeshType, FVector Scale, FVector2D Dimension)
 {
-	GetMeshComponent()->SetMesh(MeshObj);
-
 	GetSettings()->SetMesh(MeshObj);
-
 	InvalidateAll();
 }
 
diff --git a/Engine/Plugins/Experimental/TextureGraph/Source/TextureGraphEngine/Model/Mix/MixInterface.h b/Engine/Plugins/Experimental/TextureGraph/Source/TextureGraphEngine/Model/Mix/MixInterface.h
--- a/Engine/Plugins/Experimental/TextureGraph/Source/TextureGraphEngine/Model/Mix/MixInterface.h
+++ b/Engine/Plugins/Experimental/TextureGraph/Source/TextureGraphEngine/Model/Mix/MixInterface.h
@@ -8,7 +8,6 @@
 #include "MixInterface.generated.h"
 
 class UMix;
-class UMixerMeshComponent;
 class ULayerStack;
 class UMixSettings;
 class AActor;
@@ -82,8 +81,6 @@ public:
 
 	virtual RenderMeshPtr				GetMesh() const; // Get the mesh assigned in Settings
 
-	virtual UMixerMeshComponent*		GetMeshComponent() const { checkNoEntry() return nullptr; }
-
 	virtual int32						Width() const; 
 	virtual int32						Height() const;
 	
diff --git a/Engine/Plugins/Experimental/TextureGraph/Source/TextureGraphEngine/Model/Mix/MixerMeshComponent.cpp b/Engine/Plugins/Experimental/TextureGraph/Source/TextureGraphEngine/Model/Mix/MixerMeshComponent.cpp
deleted file mode 100644
diff --git a/Engine/Plugins/Experimental/TextureGraph/Source/TextureGraphEngine/Model/Mix/MixerMeshComponent.h b/Engine/Plugins/Experimental/TextureGraph/Source/TextureGraphEngine/Model/Mix/MixerMeshComponent.h
deleted file mode 100644
diff --git a/Engine/Plugins/Experimental/TextureGraph/Source/TextureGraphInsight/Public/Model/TextureGraphInsightSession.h b/Engine/Plugins/Experimental/TextureGraph/Source/TextureGraphInsight/Public/Model/TextureGraphInsightSession.h
--- a/Engine/Plugins/Experimental/TextureGraph/Source/TextureGraphInsight/Public/Model/TextureGraphInsightSession.h
+++ b/Engine/Plugins/Experimental/TextureGraph/Source/TextureGraphInsight/Public/Model/TextureGraphInsightSession.h
@@ -13,7 +13,7 @@ using MixPtrW = UMixInterface*;
 class TEXTUREGRAPHINSIGHT_API TextureGraphInsightSession
 {
 public:
-	/// Keep all the pointers to "real" Mixer Engine objects in the live runtime cache
+	/// Keep all the pointers to "real" TextureGraph Engine objects in the live runtime cache
 	class LiveCache
 	{
 	public:
@@ -95,7 +95,7 @@ protected:
 	SessionRecordPtr				Record;
 	std::mutex						RecordMutex;
 
-	/// Keep all the pointers to "real" Mixer Engine objects in the live cache
+	/// Keep all the pointers to "real" TextureGraph Engine objects in the live cache
 	/// A session live cache is renewed for each engine
 	std::shared_ptr<LiveCache>		Cache;
 
diff --git a/Engine/Plugins/Experimental/TextureGraph/Source/TextureGraphInsightEditor/Private/TextureGraphInsightEditor.cpp b/Engine/Plugins/Experimental/TextureGraph/Source/TextureGraphInsightEditor/Private/TextureGraphInsightEditor.cpp
--- a/Engine/Plugins/Experimental/TextureGraph/Source/TextureGraphInsightEditor/Private/TextureGraphInsightEditor.cpp
+++ b/Engine/Plugins/Experimental/TextureGraph/Source/TextureGraphInsightEditor/Private/TextureGraphInsightEditor.cpp
@@ -86,7 +86,7 @@ void FTextureGraphInsightEditorModule::StartupModule()
 	FInsightTabCommands::Register();
 
 	FGlobalTabmanager::Get()->RegisterNomadTabSpawner(TextureGraphInsightEditorTabName, FOnSpawnTab::CreateRaw(this, &FTextureGraphInsightEditorModule::OnSpawnPluginTab))
-		.SetDisplayName(LOCTEXT("FTextureGraphInsightEditorTabTitle", "Mixer Insight"))
+		.SetDisplayName(LOCTEXT("FTextureGraphInsightEditorTabTitle", "TextureGraph Insight"))
 		.SetMenuType(ETabSpawnerMenuType::Hidden);
 
 
@@ -120,7 +120,7 @@ bool FTextureGraphInsightEditorModule::Tick(float DeltaTime)
 }
 TSharedRef<SDockTab> FTextureGraphInsightEditorModule::OnSpawnPluginTab(const FSpawnTabArgs& SpawnTabArgs)
 {
-	// (re)Instantiate the Mixer Insight singleton
+	// (re)Instantiate the TextureGraph Insight singleton
 	if (TextureGraphInsight::Instance())
 	{
 		TextureGraphInsight::Destroy();
@@ -130,7 +130,7 @@ TSharedRef<SDockTab> FTextureGraphInsightEditorModule::OnSpawnPluginTab(const FS
 
 	const TSharedRef<SDockTab> NomadTab = SNew(SDockTab)
 		.TabRole(ETabRole::NomadTab)
-		.Label(NSLOCTEXT("TextureGraphInsight", "Mixer Insight", "Mixer Insight"));
+		.Label(NSLOCTEXT("TextureGraphInsight", "TextureGraph Insight", "TextureGraph Insight"));
 
 	if (!TabManager.IsValid())
 	{
@@ -177,7 +177,7 @@ TSharedRef<SDockTab> FTextureGraphInsightEditorModule::OnSpawnPluginTab(const FS
 			{
 				return SNew(SDockTab)
 					.TabRole(ETabRole::PanelTab)
-					.Label(LOCTEXT("FTextureGraphInsightEditorTabTitleMixes", "Mixer Mixes"))
+					.Label(LOCTEXT("FTextureGraphInsightEditorTabTitleMixes", "TextureGraph Assets"))
 					[
 						SNew(STextureGraphInsightMixListView)
 					];
@@ -187,7 +187,7 @@ TSharedRef<SDockTab> FTextureGraphInsightEditorModule::OnSpawnPluginTab(const FS
 			{
 				return SNew(SDockTab)
 					.TabRole(ETabRole::PanelTab)
-					.Label(LOCTEXT("FTextureGraphInsightEditorTabTitleAction", "Mixer Actions"))
+					.Label(LOCTEXT("FTextureGraphInsightEditorTabTitleAction", "TextureGraph Actions"))
 					[
 						SNew(STextureGraphInsightActionView)
 					];
@@ -197,7 +197,7 @@ TSharedRef<SDockTab> FTextureGraphInsightEditorModule::OnSpawnPluginTab(const FS
 			{
 				return SNew(SDockTab)
 					.TabRole(ETabRole::PanelTab)
-					.Label(LOCTEXT("FTextureGraphInsightEditorTabTitleJobBatches", "Mixer Jobs & Batches"))
+					.Label(LOCTEXT("FTextureGraphInsightEditorTabTitleJobBatches", "TextureGraph Jobs & Batches"))
 					[
 						SNew(STextureGraphInsightSessionView)
 					];
@@ -207,7 +207,7 @@ TSharedRef<SDockTab> FTextureGraphInsightEditorModule::OnSpawnPluginTab(const FS
 			{
 				return SNew(SDockTab)
 					.TabRole(ETabRole::PanelTab)
-					.Label(LOCTEXT("FTextureGraphInsightEditorTabTitleResources", "Mixer Resources"))
+					.Label(LOCTEXT("FTextureGraphInsightEditorTabTitleResources", "TextureGraph Resources"))
 					[
 						SNew(STextureGraphInsightResourceView)
 					]; 
@@ -217,7 +217,7 @@ TSharedRef<SDockTab> FTextureGraphInsightEditorModule::OnSpawnPluginTab(const FS
 			{
 				return SNew(SDockTab)
 					.TabRole(ETabRole::PanelTab)
-					.Label(LOCTEXT("FTextureGraphInsightEditorTabTitleDevices", "Mixer Devices"))
+					.Label(LOCTEXT("FTextureGraphInsightEditorTabTitleDevices", "TextureGraph Devices"))
 					[
 						SNew(STextureGraphInsightDeviceListView)
 					];
@@ -227,7 +227,7 @@ TSharedRef<SDockTab> FTextureGraphInsightEditorModule::OnSpawnPluginTab(const FS
 			{
 				return SNew(SDockTab)
 					.TabRole(ETabRole::PanelTab)
-					.Label(LOCTEXT("FTextureGraphInsightEditorTabTitleInspector", "Mixer Inspector"))
+					.Label(LOCTEXT("FTextureGraphInsightEditorTabTitleInspector", "TextureGraph Inspector"))
 					[
 						SNew(STextureGraphInsightInspectorView)
 					];
diff --git a/Engine/Plugins/Experimental/TextureGraph/Source/TextureGraphInsightEditor/Private/TextureGraphInsightEditorCommands.cpp b/Engine/Plugins/Experimental/TextureGraph/Source/TextureGraphInsightEditor/Private/TextureGraphInsightEditorCommands.cpp
--- a/Engine/Plugins/Experimental/TextureGraph/Source/TextureGraphInsightEditor/Private/TextureGraphInsightEditorCommands.cpp
+++ b/Engine/Plugins/Experimental/TextureGraph/Source/TextureGraphInsightEditor/Private/TextureGraphInsightEditorCommands.cpp
@@ -6,7 +6,7 @@
 
 void FTextureGraphInsightEditorCommands::RegisterCommands()
 {
-	// UI_COMMAND(OpenPluginWindow, "Mixer Insight", "Bring up Mixer Insight window", EUserInterfaceActionType::Button, FInputChord());
+	// UI_COMMAND(OpenPluginWindow, "TextureGraph Insight", "Bring up TextureGraph Insight window", EUserInterfaceActionType::Button, FInputChord());
 }
 
 #undef LOCTEXT_NAMESPACE
