Exposiong MirrorX and MirrorY to the Transform node
add support for autoconversion from scalar/vector/color to vector2f

diff --git a/Engine/Plugins/Experimental/TextureGraph/Shaders/Expressions/Expression_Transform.usf b/Engine/Plugins/Experimental/TextureGraph/Shaders/Expressions/Expression_Transform.usf
--- a/Engine/Plugins/Experimental/TextureGraph/Shaders/Expressions/Expression_Transform.usf
+++ b/Engine/Plugins/Experimental/TextureGraph/Shaders/Expressions/Expression_Transform.usf
@@ -23,6 +23,8 @@ float ScaleX;
 float ScaleY;
 
 float FilterMode;
+float MirrorX;
+float MirrorY;
 
 float StaggerX;
 float StaggerY;
@@ -131,15 +133,18 @@ float4 FSH_TransformBlit(float2 uv : TEXCOORD0) : SV_Target0
 	// Eval the cell grid masked  for debug
 	float cellGridMask = CoverageMask * Paint_BoxEdge(spacedCellPos, 0, 1.0, -pixelSize * 0.5, pixelSize * 1.25);
 	
-	
-	
+		
 	// also compute the mask for spacedCell and exclude anything out of the [0,1] box
 	float CellMask = MaskBox(spacedCellPos, 0, 1);
 	CoverageMask *= CellMask;
 	
 	// Finally fetch the source texture
 	float2 fetchPos = spacedCellPos; 
-	fetchPos.y = 1.0 - fetchPos.y;
+	
+	// Apply mirroring to uvs
+	fetchPos.x = lerp(fetchPos.x, 1.0 - fetchPos.x, MirrorX);	
+	fetchPos.y = lerp(1.0 - fetchPos.y, fetchPos.y, MirrorY);
+	
 	float4 SourceColor = float4(0, 0, 0, 1);
 	SourceColor = SourceTexture.Sample(SamplerStates_Linear_Wrap, fetchPos);
 		
diff --git a/Engine/Plugins/Experimental/TextureGraph/Source/TextureGraph/Private/Expressions/Procedural/TG_Expression_Transform.cpp b/Engine/Plugins/Experimental/TextureGraph/Source/TextureGraph/Private/Expressions/Procedural/TG_Expression_Transform.cpp
--- a/Engine/Plugins/Experimental/TextureGraph/Source/TextureGraph/Private/Expressions/Procedural/TG_Expression_Transform.cpp
+++ b/Engine/Plugins/Experimental/TextureGraph/Source/TextureGraph/Private/Expressions/Procedural/TG_Expression_Transform.cpp
@@ -49,6 +49,8 @@ void UTG_Expression_Transform::Evaluate(FTG_EvaluationContext* InContext)
 	T_Transform::ColorParameter ColorParam{
 		.FillColor = FillColor,
 		.WrapFilterMode = WrapMode,
+		.MirrorX = MirrorX,
+		.MirrorY = MirrorY,
 		.ShowDebugGrid = ShowDebugGrid
 	};
 
diff --git a/Engine/Plugins/Experimental/TextureGraph/Source/TextureGraph/Private/TG_GraphEvaluation.cpp b/Engine/Plugins/Experimental/TextureGraph/Source/TextureGraph/Private/TG_GraphEvaluation.cpp
--- a/Engine/Plugins/Experimental/TextureGraph/Source/TextureGraph/Private/TG_GraphEvaluation.cpp
+++ b/Engine/Plugins/Experimental/TextureGraph/Source/TextureGraph/Private/TG_GraphEvaluation.cpp
@@ -83,6 +83,11 @@ void FloatToFVector4f_Converter(FTG_Evaluation::VarConverterInfo& Info)
 	Info.OutVar->EditAs<FVector4f>() = FLinearColor(Input, Input, Input);
 }
 
+void FloatToFVector2f_Converter(FTG_Evaluation::VarConverterInfo& Info)
+{
+	auto Input = Info.InVar->GetAs<float>();
+	Info.OutVar->EditAs<FVector2f>() = FVector2f(Input, Input);
+}
 
 void FLinearColorToFVector4f_Converter(FTG_Evaluation::VarConverterInfo& Info)
 {
@@ -95,6 +100,17 @@ void FVector4fToFLinearColor_Converter(FTG_Evaluation::VarConverterInfo& Info)
 	Info.OutVar->EditAs<FLinearColor>() = Input;
 }
 
+void FLinearColorToFVector2f_Converter(FTG_Evaluation::VarConverterInfo& Info)
+{
+	auto Input = Info.InVar->GetAs<FLinearColor>();
+	Info.OutVar->EditAs<FVector2f>() = FVector2f(Input.R, Input.G);
+}
+void FVector4fToFVector2f_Converter(FTG_Evaluation::VarConverterInfo& Info)
+{
+	auto Input = Info.InVar->GetAs<FVector4f>();
+	Info.OutVar->EditAs<FVector2f>() = FVector2f(Input);
+}
+
 
 const FString FTG_Evaluation::GVectorToTextureAutoConv_Name = TEXT("_Auto_Conv_Vector_To_Tex_");
 const FString FTG_Evaluation::GColorToTextureAutoConv_Name = TEXT("_Auto_Conv_LinearColor_To_Tex_");
@@ -308,6 +324,32 @@ void FTG_VariantToFVector4f_Converter(FTG_Evaluation::VarConverterInfo& Info)
 	}
 }
 
+
+void FTG_VariantToFVector2f_Converter(FTG_Evaluation::VarConverterInfo& Info)
+{
+	auto Input = Info.InVar->GetAs<FTG_Variant>();
+	auto& Output = Info.OutVar->EditAs<FVector2f>();
+
+	FTG_Variant::EType SourceType = Input.GetType();
+
+	switch (SourceType)
+	{
+	case FTG_Variant::EType::Texture:
+		return;
+		break;
+	case FTG_Variant::EType::Scalar:
+		Output = FVector2f((float)Input.Data.Get<float>());
+		break;
+	case FTG_Variant::EType::Color:
+		Output = FVector2f(Input.Data.Get<FLinearColor>().R, Input.Data.Get<FLinearColor>().G);
+		break;
+	case FTG_Variant::EType::Vector:
+		Output = FVector2f(Input.Data.Get<FVector4f>());
+		break;
+	}
+}
+
+
 void FTG_VariantToFTG_Texture_Converter(FTG_Evaluation::VarConverterInfo& Info)
 {
 	auto Input = Info.InVar->GetAs<FTG_Variant>();
@@ -356,10 +398,15 @@ FTG_Evaluation::ConverterMap FTG_Evaluation::DefaultConverters
 		VAR_CONVERTER(float, FLinearColor, FloatToFLinearColor_Converter),
 		VAR_CONVERTER(float, FVector4f, FloatToFVector4f_Converter),
 		VAR_CONVERTER(float, FTG_Texture, FloatToFTG_Texture_Converter),
+		VAR_CONVERTER(float, FVector2f, FloatToFVector2f_Converter),
+
 		VAR_CONVERTER(FLinearColor, FVector4f, FLinearColorToFVector4f_Converter),
 		VAR_CONVERTER(FLinearColor, FTG_Texture, FLinearColorToFTG_Texture_Converter),
+		VAR_CONVERTER(FLinearColor, FVector2f, FLinearColorToFVector2f_Converter),
+
 		VAR_CONVERTER(FVector4f, FLinearColor, FVector4fToFLinearColor_Converter),
 		VAR_CONVERTER(FVector4f, FTG_Texture, FVector4fToFTG_Texture_Converter),
+		VAR_CONVERTER(FVector4f, FVector2f, FVector4fToFVector2f_Converter),
 
 		VAR_CONVERTER(float, FTG_Variant, FloatToFTG_Variant_Converter),
 		VAR_CONVERTER(FLinearColor, FTG_Variant, FLinearColorToFTG_Variant_Converter),
@@ -381,6 +428,11 @@ FTG_Evaluation::ConverterMap FTG_Evaluation::DefaultConverters
 		VAR_CONVERTER(FTG_Variant.Vector, FTG_Texture, FTG_VariantToFTG_Texture_Converter),
 		VAR_CONVERTER(FTG_Variant.Texture, FTG_Texture, FTG_VariantToFTG_Texture_Converter),
 
+		VAR_CONVERTER(FTG_Variant.Scalar, FVector2f, FTG_VariantToFVector2f_Converter),
+		VAR_CONVERTER(FTG_Variant.Color, FVector2f, FTG_VariantToFVector2f_Converter),
+		VAR_CONVERTER(FTG_Variant.Vector, FVector2f, FTG_VariantToFVector2f_Converter),
+
+
 		VAR_CONVERTER_NULL(FTG_Variant.Scalar, FTG_Variant),
 		VAR_CONVERTER_NULL(FTG_Variant.Color, FTG_Variant),
 		VAR_CONVERTER_NULL(FTG_Variant.Vector, FTG_Variant),
diff --git a/Engine/Plugins/Experimental/TextureGraph/Source/TextureGraph/Public/Expressions/Procedural/TG_Expression_Transform.h b/Engine/Plugins/Experimental/TextureGraph/Source/TextureGraph/Public/Expressions/Procedural/TG_Expression_Transform.h
--- a/Engine/Plugins/Experimental/TextureGraph/Source/TextureGraph/Public/Expressions/Procedural/TG_Expression_Transform.h
+++ b/Engine/Plugins/Experimental/TextureGraph/Source/TextureGraph/Public/Expressions/Procedural/TG_Expression_Transform.h
@@ -34,6 +34,14 @@ public:
 	UPROPERTY(EditAnywhere, Category = NoCategory, meta = (TGType = "TG_Setting", DisplayName = "Wrap Mapping"))
 	bool									WrapMode = true;
 
+	// Filter mode is wrap or clamped
+	UPROPERTY(EditAnywhere, Category = NoCategory, meta = (TGType = "TG_Setting", DisplayName = "Mirror X"))
+	bool									MirrorX = false;
+
+	// Filter mode is wrap or clamped
+	UPROPERTY(EditAnywhere, Category = NoCategory, meta = (TGType = "TG_Setting", DisplayName = "Mirror Y"))
+	bool									MirrorY = false;
+
 	// The coverage of the transform in range [0, 1]. Defaults to 1,1
 	UPROPERTY(EditAnywhere, Category = NoCategory, meta = (TGType = "TG_Setting", UIMin = "0.0", ClampMin = "0.0", UIMax = "1.0", ClampMax = "1.0"))
 	FVector2f								Coverage = { 1, 1 };
diff --git a/Engine/Plugins/Experimental/TextureGraph/Source/TextureGraphEngine/Transform/Expressions/T_Transform.cpp b/Engine/Plugins/Experimental/TextureGraph/Source/TextureGraphEngine/Transform/Expressions/T_Transform.cpp
--- a/Engine/Plugins/Experimental/TextureGraph/Source/TextureGraphEngine/Transform/Expressions/T_Transform.cpp
+++ b/Engine/Plugins/Experimental/TextureGraph/Source/TextureGraphEngine/Transform/Expressions/T_Transform.cpp
@@ -61,6 +61,8 @@ TiledBlobPtr T_Transform::Create(MixUpdateCyclePtr Cycle, BufferDescriptor Desir
 		->AddArg(ARG_FLOAT(CellParam.Spacing.X, "SpacingX"))
 		->AddArg(ARG_FLOAT(CellParam.Spacing.Y, "SpacingY"))
 		->AddArg(ARG_FLOAT(float(ColorParam.WrapFilterMode), "FilterMode"))
+		->AddArg(ARG_FLOAT(float(ColorParam.MirrorX), "MirrorX"))
+		->AddArg(ARG_FLOAT(float(ColorParam.MirrorY), "MirrorY"))
 		->AddArg(ARG_FLOAT(ColorParam.ShowDebugGrid, "BlendDebugGrid"))
 		;
 
diff --git a/Engine/Plugins/Experimental/TextureGraph/Source/TextureGraphEngine/Transform/Expressions/T_Transform.h b/Engine/Plugins/Experimental/TextureGraph/Source/TextureGraphEngine/Transform/Expressions/T_Transform.h
--- a/Engine/Plugins/Experimental/TextureGraph/Source/TextureGraphEngine/Transform/Expressions/T_Transform.h
+++ b/Engine/Plugins/Experimental/TextureGraph/Source/TextureGraphEngine/Transform/Expressions/T_Transform.h
@@ -42,6 +42,8 @@ public:
 		SHADER_PARAMETER(float, SpacingX)
 		SHADER_PARAMETER(float, SpacingY)
 		SHADER_PARAMETER(float, FilterMode)
+		SHADER_PARAMETER(float, MirrorX)
+		SHADER_PARAMETER(float, MirrorY)
 		SHADER_PARAMETER(float, BlendDebugGrid)
 	END_SHADER_PARAMETER_STRUCT()
 
@@ -82,6 +84,8 @@ public:
 	{
 		FLinearColor FillColor = { 0,0,0, 1 };
 		bool WrapFilterMode = true;
+		bool MirrorX = false;
+		bool MirrorY = false;
 		float ShowDebugGrid = 0.0f;
 	};
 
