Removing the input one expression from palette.

diff --git a/Engine/Plugins/Experimental/TextureGraph/Source/TextureGraph/Public/Expressions/Maths/TG_Expression_Maths_OneInput.h b/Engine/Plugins/Experimental/TextureGraph/Source/TextureGraph/Public/Expressions/Maths/TG_Expression_Maths_OneInput.h
--- a/Engine/Plugins/Experimental/TextureGraph/Source/TextureGraph/Public/Expressions/Maths/TG_Expression_Maths_OneInput.h
+++ b/Engine/Plugins/Experimental/TextureGraph/Source/TextureGraph/Public/Expressions/Maths/TG_Expression_Maths_OneInput.h
@@ -38,7 +38,7 @@ public:
 //////////////////////////////////////////////////////////////////////////
 /// Generic single input math function
 //////////////////////////////////////////////////////////////////////////
-UCLASS()
+UCLASS(Abstract)
 class TEXTUREGRAPH_API UTG_Expression_OneInput : public UTG_Expression_Variant
 {
 	GENERATED_BODY()
