Resetting the output expression to scalar when the source is not connected. Also setting the viewport to black when non textured output is connected to source.

diff --git a/Engine/Plugins/Experimental/TextureGraph/Source/TextureGraph/Private/Expressions/Output/TG_Expression_Output.cpp b/Engine/Plugins/Experimental/TextureGraph/Source/TextureGraph/Private/Expressions/Output/TG_Expression_Output.cpp
--- a/Engine/Plugins/Experimental/TextureGraph/Source/TextureGraph/Private/Expressions/Output/TG_Expression_Output.cpp
+++ b/Engine/Plugins/Experimental/TextureGraph/Source/TextureGraph/Private/Expressions/Output/TG_Expression_Output.cpp
@@ -69,10 +69,10 @@ void UTG_Expression_Output::Evaluate(FTG_EvaluationContext* InContext)
 	}
 	else
 	{
-		Output.EditTexture() = FTG_Texture::GetBlack();
+		// Resetting it to scalar.
+		Output.EditScalar() = 0.0f;
 	}
 
-
 	// We only update the output last render for the root graph. We don't need to call it for subgraphs.
 	// Last render eventually gets set for target texture set and then displayed on the viewport.
 	if (InContext->GraphDepth == 0)
@@ -87,6 +87,10 @@ void UTG_Expression_Output::Evaluate(FTG_EvaluationContext* InContext)
 			auto OutputFlatTexture = Source.GetTexture(InContext, FTG_Texture::GetBlack(), &DesiredDesc);
 			InContext->Cycle->GetTarget(0)->GetLastRender().SetTexture(GetTitleName(), OutputFlatTexture.RasterBlob);
 		}
+		else
+		{
+			InContext->Cycle->GetTarget(0)->GetLastRender().SetTexture(GetTitleName(), FTG_Texture::GetBlack());
+		}
 	}
 }
 
