TextureGraph: Changed how thumbnails are created and applied to slate widgets

diff --git a/Engine/Plugins/Experimental/TextureGraph/Source/TextureGraph/Public/TG_Texture.h b/Engine/Plugins/Experimental/TextureGraph/Source/TextureGraph/Public/TG_Texture.h
--- a/Engine/Plugins/Experimental/TextureGraph/Source/TextureGraph/Public/TG_Texture.h
+++ b/Engine/Plugins/Experimental/TextureGraph/Source/TextureGraph/Public/TG_Texture.h
@@ -59,7 +59,6 @@ PRAGMA_ENABLE_DEPRECATION_WARNINGS
 	FORCEINLINE operator BufferDescriptor() const { return ToBufferDescriptor(); }
 };
 
-class FAssetThumbnail;
 USTRUCT()
 struct TEXTUREGRAPH_API FTG_Texture
 {
diff --git a/Engine/Plugins/Experimental/TextureGraph/Source/TextureGraphEditor/Private/EdGraph/STG_EditorGraphNode.cpp b/Engine/Plugins/Experimental/TextureGraph/Source/TextureGraphEditor/Private/EdGraph/STG_EditorGraphNode.cpp
--- a/Engine/Plugins/Experimental/TextureGraph/Source/TextureGraphEditor/Private/EdGraph/STG_EditorGraphNode.cpp
+++ b/Engine/Plugins/Experimental/TextureGraph/Source/TextureGraphEditor/Private/EdGraph/STG_EditorGraphNode.cpp
@@ -26,8 +26,11 @@
 #include "TG_HelperFunctions.h"
 #include "Widgets/Text/SInlineEditableTextBlock.h"
 #include "SlateOptMacros.h"
+#include "STG_EditorViewport.h"
 #include "Widgets/SToolTip.h"
 #include "STG_NodeThumbnail.h"
+#include "Job/Scheduler.h"
+#include "Job/ThumbnailsService.h"
 #include "Pins/STG_TextureDescriptor.h"
 #include "Widgets/Layout/SSeparator.h"
 #define LOCTEXT_NAMESPACE "TextureGraphEditor"
@@ -37,26 +40,37 @@ static const FSlateBrush* CacheImg_Pin_NotConnectable = nullptr;
 
 STG_EditorGraphNode::~STG_EditorGraphNode()
 {
-	if (TSEditorGraphNode)
+	if (TGEditorGraphNode)
 	{
-		TSEditorGraphNode->OnNodeReconstructDelegate.Remove(OnNodeChangedHandle);
-		TSEditorGraphNode->OnNodePostEvaluateDelegate.Remove(OnPostEvaluateHandle);
-		TSEditorGraphNode->OnPinSelectionChangeDelegate.Remove(OnPinSelectionChangedHandle);
+		TGEditorGraphNode->OnNodeReconstructDelegate.Remove(OnNodeChangedHandle);
+		TGEditorGraphNode->OnNodePostEvaluateDelegate.Remove(OnPostEvaluateHandle);
+		TGEditorGraphNode->OnPinSelectionChangeDelegate.Remove(OnPinSelectionChangedHandle);
+	}
+
+	const ThumbnailsServicePtr SvcThumbnail = TextureGraphEngine::GetScheduler()->GetThumbnailsService().lock();
+	if (SvcThumbnail)
+	{
+		SvcThumbnail->OnUpdateThumbnailDelegate.Remove(OnUpdateThumbHandle);
 	}
 }
 
 void STG_EditorGraphNode::Construct(const FArguments& InArgs, UTG_EdGraphNode* InNode)
 {
 	GraphNode = InNode;
-	TSEditorGraphNode = InNode;
+	TGEditorGraphNode = InNode;
 
-	if (TSEditorGraphNode)
+	if (TGEditorGraphNode)
 	{
-		OnNodeChangedHandle = TSEditorGraphNode->OnNodeReconstructDelegate.AddRaw(this, &STG_EditorGraphNode::OnNodeReconstruct);
-		OnPostEvaluateHandle = TSEditorGraphNode->OnNodePostEvaluateDelegate.AddRaw(this, &STG_EditorGraphNode::OnNodePostEvaluate);
-		OnPinSelectionChangedHandle = TSEditorGraphNode->OnPinSelectionChangeDelegate.AddRaw(this, &STG_EditorGraphNode::OnPinSelectionChanged);
+		OnNodeChangedHandle = TGEditorGraphNode->OnNodeReconstructDelegate.AddRaw(this, &STG_EditorGraphNode::OnNodeReconstruct);
+		OnPostEvaluateHandle = TGEditorGraphNode->OnNodePostEvaluateDelegate.AddRaw(this, &STG_EditorGraphNode::OnNodePostEvaluate);
+		OnPinSelectionChangedHandle = TGEditorGraphNode->OnPinSelectionChangeDelegate.AddRaw(this, &STG_EditorGraphNode::OnPinSelectionChanged);
+	}
+	const ThumbnailsServicePtr SvcThumbnail = TextureGraphEngine::GetScheduler()->GetThumbnailsService().lock();
+	if (SvcThumbnail)
+	{
+		// register yourself to the thumbnail service to get events when a thumb batch is done.
+		OnUpdateThumbHandle = SvcThumbnail->OnUpdateThumbnailDelegate.AddRaw(this, &STG_EditorGraphNode::OnUpdateThumbnail);
 	}
-	
 	BodyBrush = *GetNodeBodyBrush();
 	BodyBrush.OutlineSettings.Color = UTG_EdGraphSchema::NodeBodyColorOutline;
 	
@@ -463,7 +477,7 @@ EVisibility STG_EditorGraphNode::ShowParameters() const
 
 EVisibility STG_EditorGraphNode::ShowOverrideSettings() const
 {
-	const bool bOverrideSettings = GraphNode && (ENodeAdvancedPins::Shown == GraphNode->AdvancedPinDisplay) && TSEditorGraphNode->GetTextureOutputPins().Num() > 0;
+	const bool bOverrideSettings = GraphNode && (ENodeAdvancedPins::Shown == GraphNode->AdvancedPinDisplay) && TGEditorGraphNode->GetTextureOutputPins().Num() > 0;
 	return bOverrideSettings ? EVisibility::Visible : EVisibility::Collapsed;
 }
 
@@ -493,7 +507,7 @@ TSharedRef<SWidget> STG_EditorGraphNode::CreateNodeSettings()
 	auto Widget = SNew(SVerticalBox);
 	const UTG_EdGraphSchema* Schema = Cast<const UTG_EdGraphSchema>(GraphNode->GetSchema());
 
-	for (auto Pin : TSEditorGraphNode->GetTextureOutputPins())
+	for (auto Pin : TGEditorGraphNode->GetTextureOutputPins())
 	{
 		UTG_Pin* TGPin = Schema->GetTGPinFromEdPin(Pin);
 		FTG_Texture& Texture = TGPin->GetNodePtr()->GetGraph()->GetVar(TGPin->GetId())->EditAs<FTG_Texture>();
@@ -612,7 +626,7 @@ TSharedRef<SWidget> STG_EditorGraphNode::CreateTitleDetailsWidget()
 
 EVisibility STG_EditorGraphNode::IsTitleDetailVisible() const
 {
-	return TSEditorGraphNode->GetTitleDetail() == "" ? EVisibility::Collapsed : EVisibility::Visible;
+	return TGEditorGraphNode->GetTitleDetail() == "" ? EVisibility::Collapsed : EVisibility::Visible;
 }
 
 void STG_EditorGraphNode::CreatePinWidgets()
@@ -633,7 +647,7 @@ void STG_EditorGraphNode::CreatePinWidgets()
 
 		bool bPinDesiresToBeHidden = CurPin->bHidden || (bHideNoConnectionPins && !bPinHasConections); 
 		
-		const UTG_EdGraphSchema* Schema = Cast<const UTG_EdGraphSchema>(TSEditorGraphNode->GetSchema());
+		const UTG_EdGraphSchema* Schema = Cast<const UTG_EdGraphSchema>(TGEditorGraphNode->GetSchema());
 		UTG_Pin* TSPin = Schema->GetTGPinFromEdPin(CurPin);
 		FProperty* Property = TSPin->GetExpressionProperty();
 		if (Property && Property->HasMetaData("HideNodeUI"))
@@ -866,7 +880,7 @@ FReply STG_EditorGraphNode::OnOutputIconClick(const FGeometry& SenderGeometry, c
 
 TSharedRef<SWidget> STG_EditorGraphNode::CreateTitleRightWidget()
 {
-	auto TGNode = TSEditorGraphNode->GetNode();
+	auto TGNode = TGEditorGraphNode->GetNode();
 	TArray<UTG_Pin*> Pins;
 	TGNode->GetOutputPins(Pins);
 	TSharedPtr<SWidget> Thumb = SNullWidget::NullWidget;
@@ -881,17 +895,16 @@ TSharedRef<SWidget> STG_EditorGraphNode::CreateTitleRightWidget()
 			auto ThumbnailWidget = FindOrCreateThumbWidget(Pin->GetId());
 
 			// get blob from EdGraphNode's cache if available
-			if (TiledBlobPtr CachedThumb = TSEditorGraphNode->GetCachedThumbBlob(Pin->GetId()))
+			UTG_EdGraph* EdGraph = Cast<UTG_EdGraph>(TGEditorGraphNode->GetGraph());
+			if (TiledBlobPtr CachedThumb = EdGraph->GetCachedThumbBlob(Pin->GetId()))
 			{
-				CachedThumb->OnFinalise().then([CachedThumb, ThumbnailWidget]()
+				if (CachedThumb->IsFinalised())
 				{
-					// NOTE: If later, "this" were to be captured here, we should check DoesSharedInstanceExist()
-					// as there is a chance this might be invoked when the slate widgets have already been destroyed
 					if (ThumbnailWidget.IsValid())
 					{
 						ThumbnailWidget->UpdateBlob(CachedThumb);
 					}
-				});
+				}
 			}
 			
 			if (Thumb == SNullWidget::NullWidget)
@@ -914,8 +927,8 @@ TSharedRef<SWidget> STG_EditorGraphNode::CreateTitleRightWidget()
 
 void STG_EditorGraphNode::SetPinIcon(const TSharedRef<SGraphPin> PinToAdd)
 {
-	check(TSEditorGraphNode);
-	UTG_Node* TSNode = TSEditorGraphNode->GetNode();
+	check(TGEditorGraphNode);
+	UTG_Node* TSNode = TGEditorGraphNode->GetNode();
 	UEdGraphPin* EdPin = PinToAdd->GetPinObj();
 
 	if (TSNode && EdPin)
@@ -963,13 +976,11 @@ void STG_EditorGraphNode::OnNodeReconstruct()
 
 void STG_EditorGraphNode::OnNodePostEvaluate(const FTG_EvaluationContext* InContext)
 {
-	if (!InContext->IsTweaking())
-		UpdateThumbnail(InContext);
 }
 
 void STG_EditorGraphNode::OnPinSelectionChanged(UEdGraphPin* EdPin)
 {
-	const UTG_EdGraphSchema* Schema = Cast<const UTG_EdGraphSchema>(TSEditorGraphNode->GetSchema());
+	const UTG_EdGraphSchema* Schema = Cast<const UTG_EdGraphSchema>(TGEditorGraphNode->GetSchema());
 	UTG_Pin* Pin = Schema->GetTGPinFromEdPin(EdPin);
 
 	if (Pin)
@@ -982,63 +993,30 @@ void STG_EditorGraphNode::OnPinSelectionChanged(UEdGraphPin* EdPin)
 	}
 }
 
-void STG_EditorGraphNode::UpdateThumbnail(const FTG_EvaluationContext* InContext)
+
+void STG_EditorGraphNode::ApplyThumbToWidget()
 {
-	// Loop over all Output pins
-	for (const TPair<FName, FTG_VarMap::FVarArgument>& V : InContext->Outputs.VarArguments)
+	UTG_EdGraph* EdGraph = Cast<UTG_EdGraph>(TGEditorGraphNode->GetGraph());
+	TArray<UTG_Pin*> Pins;
+	TGEditorGraphNode->GetNode()->GetOutputPins(Pins);
+	for(const UTG_Pin* Pin : Pins)
 	{
-		FTG_VarMap::FVarArgument VarArg = V.Value;
-		auto TypeName = VarArg.Argument.GetCPPTypeName();
-		UTG_Pin* Pin = TSEditorGraphNode->GetNode()->GetOutputPin(V.Key);
-		check(Pin);
-		
-		// if they are FTG_Texture or Textured Variant, make a thumbnail for that.
-		if (VarArg.Argument.IsTexture())
+		auto ThumbWidget = FindOrCreateThumbWidget(Pin->GetId());
+		if (ThumbWidget.IsValid())
 		{
-			FTG_Texture OutTexture;
-			if(Pin->GetValue(OutTexture))
-			{
-				UMixInterface* Mix = InContext->Cycle->GetMix();
-				auto TargetId = InContext->TargetId;
-
-				if (!OutTexture.RasterBlob)
-				{
-					OutTexture = FTG_Texture::GetBlack();
-				}
-
-				TiledBlobPtr ThumbBlob = T_Thumbnail::Bind(Mix, Pin, OutTexture.RasterBlob, TargetId);
-
-				TSEditorGraphNode->CacheThumbBlob(Pin->GetId(), ThumbBlob);
-				ThumbBlob->OnFinalise()
-					.then([=, this](const Blob* FinalisedBlob) mutable
-						{
-							// OnFinalise can sometimes occur after the editor is closed and thus can potentially
-							// deallocate all corresponding slate objects
-							if (DoesSharedInstanceExist())
-							{
-								return ThumbBlob->CombineTiles(true, false);
-							}
-							return static_cast<AsyncBufferResultPtr>(cti::make_ready_continuable<BufferResultPtr>(std::make_shared<BufferResult>()));
-						})
-					.then([=, this](BufferResultPtr) mutable
-						{
-							// OnFinalise can sometimes occur after the editor is closed and thus can potentially
-							// deallocate all corresponding slate objects
-							if (DoesSharedInstanceExist())
-							{
-								auto ThumbWidget = FindOrCreateThumbWidget(Pin->GetId());
-
-								if (ThumbWidget.IsValid())
-								{
-									ThumbWidget->UpdateBlob(ThumbBlob);
-								}
-							}
-						});
-			}
+			ThumbWidget->UpdateBlob(EdGraph->GetCachedThumbBlob(Pin->GetId()));
 		}
 	}
 }
 
+void STG_EditorGraphNode::OnUpdateThumbnail(JobBatchPtr JobBatch)
+{
+	// check if the JobBatch has the same TextureGraph as ours (As the Thumbnail Service could be catering to multiple TextureGraphs)
+	if (JobBatch->GetCycle()->GetMix() == TGEditorGraphNode->GetOutermostObject())
+	{
+		ApplyThumbToWidget();
+	}
+}
 
 //////////////////////////////////////////////////////////////////////////
 
diff --git a/Engine/Plugins/Experimental/TextureGraph/Source/TextureGraphEditor/Private/EdGraph/STG_EditorGraphNode.h b/Engine/Plugins/Experimental/TextureGraph/Source/TextureGraphEditor/Private/EdGraph/STG_EditorGraphNode.h
--- a/Engine/Plugins/Experimental/TextureGraph/Source/TextureGraphEditor/Private/EdGraph/STG_EditorGraphNode.h
+++ b/Engine/Plugins/Experimental/TextureGraph/Source/TextureGraphEditor/Private/EdGraph/STG_EditorGraphNode.h
@@ -83,7 +83,8 @@ protected:
 	virtual TSharedRef<SWidget> CreateTitleWidget(TSharedPtr<SNodeTitle> NodeTitle) override;
 
 private:
-	UTG_EdGraphNode* TSEditorGraphNode = nullptr;
+	UTG_EdGraphNode* TGEditorGraphNode = nullptr;
+	FDelegateHandle OnUpdateThumbHandle;
 	FDelegateHandle OnPostEvaluateHandle;
 	FDelegateHandle OnNodeChangedHandle;
 	FDelegateHandle OnPinSelectionChangedHandle;
@@ -94,7 +95,8 @@ private:
 	FLinearColor GetPinThumbSelectionColor();
 	FLinearColor GetPinThumbDeselectionColor();
 	TSharedPtr<STG_NodeThumbnail> FindOrCreateThumbWidget(FTG_Id PinId);
-	void UpdateThumbnail(const FTG_EvaluationContext* InContext);
+	void ApplyThumbToWidget();
+	void OnUpdateThumbnail(std::shared_ptr<JobBatch> JobBatch);
 	UTexture* GetTextureFromBuffer(DeviceBufferPtr Buffer);
 	FReply OnOutputIconClick(const FGeometry& SenderGeometry, const FPointerEvent& MouseEvent);
 	TArray<FString> GetTitleDetailTextLines() const;
diff --git a/Engine/Plugins/Experimental/TextureGraph/Source/TextureGraphEditor/Private/EdGraph/STG_NodeThumbnail.cpp b/Engine/Plugins/Experimental/TextureGraph/Source/TextureGraphEditor/Private/EdGraph/STG_NodeThumbnail.cpp
--- a/Engine/Plugins/Experimental/TextureGraph/Source/TextureGraphEditor/Private/EdGraph/STG_NodeThumbnail.cpp
+++ b/Engine/Plugins/Experimental/TextureGraph/Source/TextureGraphEditor/Private/EdGraph/STG_NodeThumbnail.cpp
@@ -79,49 +79,48 @@ void STG_NodeThumbnail::UpdateParams(TiledBlobPtr InBlob)
 {
 	if (!InBlob)
 		return;
+	
+	// we should assume that the block is finalised at this point
+	check (InBlob->IsFinalised());
 
-	/// We don't assume that te blob would be finalised when we get to this point
-	InBlob->OnFinalise().then([=, this]()
+	if (InBlob->IsFinalised())
+	{
+		UTexture* BlobTexture = GetTextureFromBlob(InBlob);
+		check(BlobTexture);
+
+		//FRHITexture* SlateTexture = (FRHITexture*)(AssetThumbnail->GetViewportRenderTargetTexture());
+		UTexture2D* Texture2D = Cast<UTexture2D>(BlobTexture);
+		UTextureRenderTarget2D* TextureRT2D = Cast<UTextureRenderTarget2D>(BlobTexture);
+	
+		float ShowChecker = 1.0;
+		if (BlobTexture)
 		{
-			if (DoesSharedInstanceExist())
-			{
-				UTexture* BlobTexture = GetTextureFromBlob(InBlob);
-				check(BlobTexture);
-
-				//FRHITexture* SlateTexture = (FRHITexture*)(AssetThumbnail->GetViewportRenderTargetTexture());
-				UTexture2D* Texture2D = Cast<UTexture2D>(BlobTexture);
-				UTextureRenderTarget2D* TextureRT2D = Cast<UTextureRenderTarget2D>(BlobTexture);
-				
-				float ShowChecker = 1.0;
-				if (BlobTexture)
-				{
-					ShowChecker = 0.0;
-					BrushMaterial->SetTextureParameterValue("ThumbTex", BlobTexture);
-				}
-
-				float SingleChannel = 0.0;
-
-				if (Texture2D)
-				{
-					SingleChannel = GPixelFormats[Texture2D->GetPixelFormat()].NumComponents == 1 ? 1.0 : 0.0;
-				}
-
-				if (TextureRT2D)
-				{
-					SingleChannel = GPixelFormats[TextureRT2D->GetFormat()].NumComponents == 1 ? 1.0 : 0.0;
-				}
-
-				BrushMaterial->SetScalarParameterValue("ShowChecker", ShowChecker);
-				BrushMaterial->SetScalarParameterValue("SingleChannel", SingleChannel);
-				Brush->SetResourceObject(BrushMaterial);
-			}
-		});
+			ShowChecker = 0.0;
+			BrushMaterial->SetTextureParameterValue("ThumbTex", BlobTexture);
+		}
+
+		float SingleChannel = 0.0;
+
+		if (Texture2D)
+		{
+			SingleChannel = GPixelFormats[Texture2D->GetPixelFormat()].NumComponents == 1 ? 1.0 : 0.0;
+		}
+
+		if (TextureRT2D)
+		{
+			SingleChannel = GPixelFormats[TextureRT2D->GetFormat()].NumComponents == 1 ? 1.0 : 0.0;
+		}
+
+		BrushMaterial->SetScalarParameterValue("ShowChecker", ShowChecker);
+		BrushMaterial->SetScalarParameterValue("SingleChannel", SingleChannel);
+		Brush->SetResourceObject(BrushMaterial);
+	}
 }
 
 void STG_NodeThumbnail::UpdateBlob(TiledBlobPtr InBlob)
 {
 	UpdateParams(InBlob);
-
+	
 	TSharedPtr<SOverlay> overlay;
 	ChildSlot
 	[
diff --git a/Engine/Plugins/Experimental/TextureGraph/Source/TextureGraphEditor/Private/EdGraph/TG_EdGraph.cpp b/Engine/Plugins/Experimental/TextureGraph/Source/TextureGraphEditor/Private/EdGraph/TG_EdGraph.cpp
--- a/Engine/Plugins/Experimental/TextureGraph/Source/TextureGraphEditor/Private/EdGraph/TG_EdGraph.cpp
+++ b/Engine/Plugins/Experimental/TextureGraph/Source/TextureGraphEditor/Private/EdGraph/TG_EdGraph.cpp
@@ -8,11 +8,13 @@
 #include "TextureGraph.h"
 #include "TG_Editor.h"
 #include "TG_Graph.h"
+#include "Transform/Layer/T_Thumbnail.h"
 
 void UTG_EdGraph::InitializeFromTextureGraph(UTextureGraph* InTextureGraph, TWeakPtr<FTG_Editor> InTGEditor)
 {
 	TextureGraph = InTextureGraph;
 	TGEditor = InTGEditor;
+
 	// And now build the EdGraph viewmodel matching the script
 	BuildEdGraphViewmodel();
 }
@@ -43,6 +45,16 @@ bool UTG_EdGraph::CreateViewModelLinkFromModelPins(const UTG_Pin* pinFrom, const
 	return false;
 }
 
+void UTG_EdGraph::CacheThumbBlob(FTG_Id PinId, TiledBlobPtr InBlob)
+{
+	PinThumbBlobMap.FindOrAdd(PinId) = InBlob;
+}
+
+TiledBlobPtr UTG_EdGraph::GetCachedThumbBlob(FTG_Id PinId)
+{
+	return (PinThumbBlobMap.IsEmpty() || !PinThumbBlobMap.Contains(PinId)) ? nullptr : *PinThumbBlobMap.Find(PinId);
+}
+
 void UTG_EdGraph::BuildEdGraphViewmodel()
 {
 	// For every model nodes, let's create an equivalent viewmodel node
@@ -120,6 +132,39 @@ void UTG_EdGraph::RefreshEditorDetails() const
 	}
 }
 
+void UTG_EdGraph::OnNodeCreateThumbnail(UTG_Node* InNode, const FTG_EvaluationContext* InContext)
+{
+	UTG_EdGraphNode* EdGraphNode = GetViewModelNode(InNode->GetId());
+	if (EdGraphNode)
+	{
+		TArray<UTG_Pin*> OutputPins;
+		InNode->GetOutputPins(OutputPins);
+		// Loop over all Output pins
+		for(UTG_Pin* Pin : OutputPins)
+		{
+			check (Pin);
+			if (Pin->IsArgTexture())
+			{
+				FTG_Texture OutTexture;
+				if(Pin->GetValue(OutTexture))
+				{
+					UMixInterface* Mix = InContext->Cycle->GetMix();
+					auto TargetId = InContext->TargetId;
+
+					if (!OutTexture.RasterBlob)
+					{
+						OutTexture = FTG_Texture::GetBlack();
+					}
+
+					TiledBlobPtr ThumbBlob = T_Thumbnail::Bind(Mix, Pin, OutTexture.RasterBlob, TargetId);
+
+					CacheThumbBlob(Pin->GetId(), ThumbBlob);
+				}
+			}
+		}
+	}
+}
+
 void UTG_EdGraph::GraphChanged(UTG_Graph* InGraph, UTG_Node* InNode, bool Tweaking)
 {
 	if (InNode)
@@ -145,5 +190,6 @@ void UTG_EdGraph::OnNodePostEvaluation(UTG_Node* InNode, const FTG_EvaluationCon
 	{
 		EdGraphNode->OnNodePostEvaluate(Context);
 	}
+	OnNodeCreateThumbnail(InNode, Context);
 }
 	
\ No newline at end of file
diff --git a/Engine/Plugins/Experimental/TextureGraph/Source/TextureGraphEditor/Private/EdGraph/TG_EdGraphNode.cpp b/Engine/Plugins/Experimental/TextureGraph/Source/TextureGraphEditor/Private/EdGraph/TG_EdGraphNode.cpp
--- a/Engine/Plugins/Experimental/TextureGraph/Source/TextureGraphEditor/Private/EdGraph/TG_EdGraphNode.cpp
+++ b/Engine/Plugins/Experimental/TextureGraph/Source/TextureGraphEditor/Private/EdGraph/TG_EdGraphNode.cpp
@@ -343,9 +343,6 @@ void UTG_EdGraphNode::ReconstructNode()
 	// Store copy of old pins
 	TArray<UEdGraphPin*> OldPins = MoveTemp(Pins);
 	Pins.Reset();
-
-	// clear cache of thumb blobs
-	PinThumbBlobMap.Reset();
 	
 	// Generate new pins
 	AllocateDefaultPins();
diff --git a/Engine/Plugins/Experimental/TextureGraph/Source/TextureGraphEditor/Private/STG_OutputSelectionDlg.cpp b/Engine/Plugins/Experimental/TextureGraph/Source/TextureGraphEditor/Private/STG_OutputSelectionDlg.cpp
--- a/Engine/Plugins/Experimental/TextureGraph/Source/TextureGraphEditor/Private/STG_OutputSelectionDlg.cpp
+++ b/Engine/Plugins/Experimental/TextureGraph/Source/TextureGraphEditor/Private/STG_OutputSelectionDlg.cpp
@@ -86,7 +86,7 @@ void STG_OutputSelectionDlg::AddExportItems()
 			if (Variant.IsTexture())
 			{
 				UTG_EdGraphNode* EdNode = EdGraph->GetViewModelNode(Node->GetId());
-				TiledBlobPtr ThumbBlob = EdNode->GetCachedThumbBlob(Id);
+				TiledBlobPtr ThumbBlob = EdGraph->GetCachedThumbBlob(Id);
 
 				if (!ThumbBlob)
 				{
diff --git a/Engine/Plugins/Experimental/TextureGraph/Source/TextureGraphEditor/Public/EdGraph/TG_EdGraph.h b/Engine/Plugins/Experimental/TextureGraph/Source/TextureGraphEditor/Public/EdGraph/TG_EdGraph.h
--- a/Engine/Plugins/Experimental/TextureGraph/Source/TextureGraphEditor/Public/EdGraph/TG_EdGraph.h
+++ b/Engine/Plugins/Experimental/TextureGraph/Source/TextureGraphEditor/Public/EdGraph/TG_EdGraph.h
@@ -52,6 +52,13 @@ public:
 	 * Force refresh Editor details
 	 */
 	void RefreshEditorDetails() const;
+
+	/**
+	 * Tells the node to create a thumbnail for its output pins
+	 * @param Node 
+	 * @param EvaluationContext 
+	 */
+	void OnNodeCreateThumbnail(UTG_Node* Node, const FTG_EvaluationContext* EvaluationContext);
 	
 	UPROPERTY()
 	TObjectPtr<UTextureGraph> TextureGraph = nullptr;
@@ -59,6 +66,9 @@ public:
 	TWeakPtr<FTG_Editor> TGEditor = nullptr;
 	FTG_PinSelectionManager PinSelectionManager;
 
+	void CacheThumbBlob(FTG_Id PinId, TiledBlobPtr InBlob);
+	TiledBlobPtr GetCachedThumbBlob(FTG_Id PinId);
+
 private:
 	// Generate the viewmodel matching the Script
 	void BuildEdGraphViewmodel();
@@ -70,6 +80,7 @@ private:
 	void GraphChanged(UTG_Graph* InGraph, UTG_Node* InNode, bool Tweaking);
 
 
+	TMap<FTG_Id, TiledBlobPtr> PinThumbBlobMap;
 public:
 	//void UpdateParams();
 };
diff --git a/Engine/Plugins/Experimental/TextureGraph/Source/TextureGraphEditor/Public/EdGraph/TG_EdGraphNode.h b/Engine/Plugins/Experimental/TextureGraph/Source/TextureGraphEditor/Public/EdGraph/TG_EdGraphNode.h
--- a/Engine/Plugins/Experimental/TextureGraph/Source/TextureGraphEditor/Public/EdGraph/TG_EdGraphNode.h
+++ b/Engine/Plugins/Experimental/TextureGraph/Source/TextureGraphEditor/Public/EdGraph/TG_EdGraphNode.h
@@ -96,9 +96,6 @@ public:
 	TArray<UEdGraphPin*> GetOutputPins() const;
 	TArray<UEdGraphPin*> GetTextureOutputPins() const;
 
-	void CacheThumbBlob(FTG_Id PinId, TiledBlobPtr InBlob){ PinThumbBlobMap.FindOrAdd(PinId) = InBlob; }
-	TiledBlobPtr GetCachedThumbBlob(FTG_Id PinId){ return PinThumbBlobMap.IsEmpty() ? nullptr : *PinThumbBlobMap.Find(PinId); }
-
 protected:
 	friend class UTG_EdGraph;
 
@@ -118,7 +115,6 @@ protected:
 		TObjectPtr<UTG_Node> Node;
 
 private:
-	TMap<FTG_Id, TiledBlobPtr> PinThumbBlobMap;
 	UEdGraphPin* SelectedPin;
 };
 
diff --git a/Engine/Plugins/Experimental/TextureGraph/Source/TextureGraphEngine/Job/ThumbnailsService.cpp b/Engine/Plugins/Experimental/TextureGraph/Source/TextureGraphEngine/Job/ThumbnailsService.cpp
--- a/Engine/Plugins/Experimental/TextureGraph/Source/TextureGraphEngine/Job/ThumbnailsService.cpp
+++ b/Engine/Plugins/Experimental/TextureGraph/Source/TextureGraphEngine/Job/ThumbnailsService.cpp
@@ -31,13 +31,16 @@ AsyncJobResultPtr ThumbnailsService::Tick()
  
 	TextureGraphEngine::GetInstance()->GetScheduler()->GetObserverSource()->BatchAdded(CurrentBatch); // notify observer
  
-	return CurrentBatch->Exec([=](JobBatch*)	/// Instead of passing it as an argument, JobBatch should be a return type; this is to keep it from cyclic dependancy
+	return CurrentBatch->Exec([=, this](JobBatch*)	/// Instead of passing it as an argument, JobBatch should be a return type; this is to keep it from cyclic dependancy
 	{
+		OnUpdateThumbnailDelegate.Broadcast(CurrentBatch);
 		TextureGraphEngine::GetInstance()->GetScheduler()->GetObserverSource()->BatchJobsDone(CurrentBatch);
+		
 	})
  	.then([this, CurrentBatch]()
  	{
  		TextureGraphEngine::GetInstance()->GetScheduler()->GetObserverSource()->BatchDone(CurrentBatch); // notify observer
+ 		
  		return std::make_shared<JobResult>();
  	});
 }
diff --git a/Engine/Plugins/Experimental/TextureGraph/Source/TextureGraphEngine/Job/ThumbnailsService.h b/Engine/Plugins/Experimental/TextureGraph/Source/TextureGraphEngine/Job/ThumbnailsService.h
--- a/Engine/Plugins/Experimental/TextureGraph/Source/TextureGraphEngine/Job/ThumbnailsService.h
+++ b/Engine/Plugins/Experimental/TextureGraph/Source/TextureGraphEngine/Job/ThumbnailsService.h
@@ -33,6 +33,8 @@ protected:
 	JobBatchPtr						GetNextUpdateCycle();
 
 public:
+	DECLARE_MULTICAST_DELEGATE_OneParam(FOnThumbnailBatchDoneDelegate, JobBatchPtr Batch)
+	FOnThumbnailBatchDoneDelegate	OnUpdateThumbnailDelegate;
 									ThumbnailsService();
 	virtual							~ThumbnailsService() override;
 
diff --git a/Engine/Plugins/Experimental/TextureGraph/Source/TextureGraphEngine/Transform/Layer/T_Thumbnail.cpp b/Engine/Plugins/Experimental/TextureGraph/Source/TextureGraphEngine/Transform/Layer/T_Thumbnail.cpp
--- a/Engine/Plugins/Experimental/TextureGraph/Source/TextureGraphEngine/Transform/Layer/T_Thumbnail.cpp
+++ b/Engine/Plugins/Experimental/TextureGraph/Source/TextureGraphEngine/Transform/Layer/T_Thumbnail.cpp
@@ -34,13 +34,13 @@ TiledBlobRef T_Thumbnail::Bind(UMixInterface* Mix, UObject* Model, TiledBlobPtr
 
 	JobUPtr JobPtr = std::make_unique<Job>(Mix, InTargetId, std::static_pointer_cast<BlobTransform>(RenderMaterial));
 	JobPtr->AddArg(ARG_BLOB(InBlobToBind, "SourceTexture"));
-
+	JobPtr->SetTiled(false);
 	BufferDescriptor Desc = InBlobToBind->GetDescriptor();
 	Desc.Width = RenderMaterial_Thumbnail::GThumbWidth;
 	Desc.Height = RenderMaterial_Thumbnail::GThumbHeight;
 	
 	TiledBlobPtr Result = JobPtr->InitResult(TEXT("T_Thumbnail"), &Desc);
-	//Result->MakeSingleBlob();
+	Result->MakeSingleBlob();
 
 	AddThumbJobToCycle(InBlobToBind, std::move(JobPtr), Model, InTargetId);
 
