Combine blob issues fixed that were resulting in inconsistent rendering across sessions.

diff --git a/Engine/Plugins/Experimental/TextureGraph/Source/TextureGraphEngine/Data/TiledBlob.h b/Engine/Plugins/Experimental/TextureGraph/Source/TextureGraphEngine/Data/TiledBlob.h
--- a/Engine/Plugins/Experimental/TextureGraph/Source/TextureGraphEngine/Data/TiledBlob.h
+++ b/Engine/Plugins/Experimental/TextureGraph/Source/TextureGraphEngine/Data/TiledBlob.h
@@ -134,6 +134,7 @@ public:
 		check(IsValidTileIndex(TileX, TileY));
 		return BlobRef(Tiles[TileX][TileY], false);
 	}
+	FORCEINLINE BlobPtr				GetSingleBlob() const { return SingleBlob; }
 
 	JobPtrW							Job() const { return JobObj; }
 	JobPtrW&						Job() { return JobObj; }
diff --git a/Engine/Plugins/Experimental/TextureGraph/Source/TextureGraphEngine/Device/FX/Device_FX.cpp b/Engine/Plugins/Experimental/TextureGraph/Source/TextureGraphEngine/Device/FX/Device_FX.cpp
--- a/Engine/Plugins/Experimental/TextureGraph/Source/TextureGraphEngine/Device/FX/Device_FX.cpp
+++ b/Engine/Plugins/Experimental/TextureGraph/Source/TextureGraphEngine/Device/FX/Device_FX.cpp
@@ -351,7 +351,7 @@ AsyncDeviceBufferRef Device_FX::DrawTilesToBuffer_Deferred(DeviceBufferRef Buffe
 					DeviceBuffer_FX* TileBuffer = static_cast<DeviceBuffer_FX*>(Tiles[TileX][TileY].get());
 					TexPtr TileTex = TileBuffer->GetTexture();
 
-					UE_LOG(LogDevice, Log, TEXT("Combine tiles: %s [%d, %d] => %llu [%s]"), *FXBuffer->GetName(), TileX, TileY, TileBuffer->Hash(false)->Value(), *TileBuffer->GetName());
+					UE_LOG(LogDevice, Verbose, TEXT("Combine tiles: %s [%d, %d] => %llu [%s]"), *FXBuffer->GetName(), TileX, TileY, TileBuffer->Hash(false)->Value(), *TileBuffer->GetName());
 
 					check(TileTex);
 
diff --git a/Engine/Plugins/Experimental/TextureGraph/Source/TextureGraphEngine/FxMat/RenderMaterial_FX.cpp b/Engine/Plugins/Experimental/TextureGraph/Source/TextureGraphEngine/FxMat/RenderMaterial_FX.cpp
--- a/Engine/Plugins/Experimental/TextureGraph/Source/TextureGraphEngine/FxMat/RenderMaterial_FX.cpp
+++ b/Engine/Plugins/Experimental/TextureGraph/Source/TextureGraphEngine/FxMat/RenderMaterial_FX.cpp
@@ -13,10 +13,9 @@ RenderMaterial_FX::RenderMaterial_FX(FString InName, FxMaterialPtr InFXMaterial)
 	: RenderMaterial(InName)
 	, FXMaterial(InFXMaterial)
 {
-	static constexpr int MaxName = 256;
+	static constexpr int32 MaxName = 256;
 	char FullName[MaxName] = {0};
-	check(InName.Len() < MaxName);
-	memcpy((char*)FullName, (const char*)TCHAR_TO_ANSI(*InName), InName.Len());
+	FMemory::Memcpy((char*)FullName, (const char*)TCHAR_TO_ANSI(*InName), (size_t)std::min((int32)InName.Len(), MaxName - 1));
 	HashValue = std::make_shared<CHash>(DataUtil::Hash((const uint8*)FullName, sizeof(FullName)), true);
 }
 
diff --git a/Engine/Plugins/Experimental/TextureGraph/Source/TextureGraphEngine/Helper/DataUtil.cpp b/Engine/Plugins/Experimental/TextureGraph/Source/TextureGraphEngine/Helper/DataUtil.cpp
--- a/Engine/Plugins/Experimental/TextureGraph/Source/TextureGraphEngine/Helper/DataUtil.cpp
+++ b/Engine/Plugins/Experimental/TextureGraph/Source/TextureGraphEngine/Helper/DataUtil.cpp
@@ -349,9 +349,7 @@ void CHash::UpdateLinks()
 	if (Linked.empty())
 		return;
 
-	/// has to be
 	check(IsInGameThread());
-	check(IsFinal());
 
 	auto ThisHash = shared_from_this();
 
@@ -366,9 +364,10 @@ void CHash::UpdateLinks()
 		}
 	}
 
-	/// Ok, the links should've updated over here. We can clear our _linked list now
-	/// as they'll never be used after this. 
-	Linked.clear();
+	/// Ok, the links should've updated over here and We can clear our linked list now
+	/// if the hash is final, as they'll never be used after this. 
+	if (IsFinal())
+		Linked.clear();
 }
 
 HashTypeVec CHash::GetIntermediateHashes() const
diff --git a/Engine/Plugins/Experimental/TextureGraph/Source/TextureGraphEngine/Job/Job.cpp b/Engine/Plugins/Experimental/TextureGraph/Source/TextureGraphEngine/Job/Job.cpp
--- a/Engine/Plugins/Experimental/TextureGraph/Source/TextureGraphEngine/Job/Job.cpp
+++ b/Engine/Plugins/Experimental/TextureGraph/Source/TextureGraphEngine/Job/Job.cpp
@@ -764,7 +764,7 @@ AsyncPrepareResult Job::PrepareTargets(JobBatch* Batch)
 	{
 		for (size_t TileY = 0; TileY < ResultPromise->Cols(); TileY++)
 		{
-			CHashPtr ResultTileHash = TileHash(TileX, TileY);
+			CHashPtr ResultTileHash = bCanHandleTiles ? TileHash(TileX, TileY) : TileHash(-1, -1);
 			BlobRef TileBlob = !(Batch->IsNoCache() || bIsNoCache) ? TextureGraphEngine::GetBlobber()->FindSingle(ResultTileHash->Value()) : BlobRef();
 
 			if (!TileBlob && DuplicatesThisJob.size())
diff --git a/Engine/Plugins/Experimental/TextureGraph/Source/TextureGraphEngine/Job/JobBatch.cpp b/Engine/Plugins/Experimental/TextureGraph/Source/TextureGraphEngine/Job/JobBatch.cpp
--- a/Engine/Plugins/Experimental/TextureGraph/Source/TextureGraphEngine/Job/JobBatch.cpp
+++ b/Engine/Plugins/Experimental/TextureGraph/Source/TextureGraphEngine/Job/JobBatch.cpp
@@ -362,16 +362,16 @@ void JobBatch::OnJobDone(Job* JobObj, int64 JobId)
 
 	if (JobsFinished[JobId] == false) 
 	{
-		auto numJobsAlive = --NumJobsRunning; /// Decrement and then catch the current value
+		auto NumJobsAlive = --NumJobsRunning; /// Decrement and then catch the current value
 		check(JobId < (int64)NumJobs());
-		check(numJobsAlive >= 0);
+		check(NumJobsAlive >= 0);
 
-		UE_LOG(LogBatch, VeryVerbose, TEXT("Job done: %d [Num left: %d, Batch Id: %llu]"), JobId, (uint32)numJobsAlive, BatchId);
+		UE_LOG(LogBatch, VeryVerbose, TEXT("Job done: %d [Num left: %d, Batch Id: %llu]"), JobId, (uint32)NumJobsAlive, BatchId);
 
 		JobsFinished[JobId] = true;
 
 		/// when all more Jobs are done then notify
-		if (numJobsAlive <= 0)
+		if (NumJobsAlive <= 0)
 		{
 			UE_LOG(LogBatch, VeryVerbose, TEXT("All Jobs done for Batch: %llu"), BatchId);
 			EndBatch();
diff --git a/Engine/Plugins/Experimental/TextureGraph/Source/TextureGraphEngine/Transform/Expressions/T_Transform.cpp b/Engine/Plugins/Experimental/TextureGraph/Source/TextureGraphEngine/Transform/Expressions/T_Transform.cpp
--- a/Engine/Plugins/Experimental/TextureGraph/Source/TextureGraphEngine/Transform/Expressions/T_Transform.cpp
+++ b/Engine/Plugins/Experimental/TextureGraph/Source/TextureGraphEngine/Transform/Expressions/T_Transform.cpp
@@ -16,7 +16,7 @@ IMPLEMENT_GLOBAL_SHADER(FSH_TransformBlit, "/Plugin/TextureGraph/Expressions/Exp
 TiledBlobPtr T_Transform::Create(MixUpdateCyclePtr Cycle, BufferDescriptor DesiredOutputDesc, TiledBlobPtr Source,
 	const TransformParameter& TransformParam, const CellParameter& CellParam, const ColorParameter& ColorParam, int32 TargetId)
 {
-	RenderMaterial_FXPtr RenderMaterial = TextureGraphEngine::GetMaterialManager()->CreateMaterial_FX<VSH_Simple, FSH_TransformBlit>(TEXT("T_TransformBlit"));
+	RenderMaterial_FXPtr RenderMaterial = TextureGraphEngine::GetMaterialManager()->CreateMaterial_FX<VSH_Simple, FSH_TransformBlit>(TEXT("T_Transform"));
 
 	check(RenderMaterial);
 
@@ -35,11 +35,11 @@ TiledBlobPtr T_Transform::Create(MixUpdateCyclePtr Cycle, BufferDescriptor Desir
 
 
 
-	FTileInfo tileInfo;
+	FTileInfo TileInfo;
 
 	JobUPtr RenderJob = std::make_unique<Job>(Cycle->GetMix(), TargetId, std::static_pointer_cast<BlobTransform>(RenderMaterial));
 	RenderJob
-		->AddArg(ARG_TILEINFO(tileInfo, "TileInfo"))
+		->AddArg(ARG_TILEINFO(TileInfo, "TileInfo"))
 		->AddArg(ARG_BLOB(CombinedBlob, "SourceTexture"))
 		->AddArg(ARG_LINEAR_COLOR(ColorParam.FillColor, "FillColor"))
 		->AddArg(ARG_FLOAT((TransformParam.Coverage.X < 0.005 ? 0.005 : TransformParam.Coverage.X), "CoverageX"))
