Fixed crash when trying to access a Pin which could be null. Also removed a redundant call for Var to copy itself to expression.

diff --git a/Engine/Plugins/Experimental/TextureGraph/Source/TextureGraph/Private/TG_Pin.cpp b/Engine/Plugins/Experimental/TextureGraph/Source/TextureGraph/Private/TG_Pin.cpp
--- a/Engine/Plugins/Experimental/TextureGraph/Source/TextureGraph/Private/TG_Pin.cpp
+++ b/Engine/Plugins/Experimental/TextureGraph/Source/TextureGraph/Private/TG_Pin.cpp
@@ -247,7 +247,6 @@ void UTG_Pin::SetSelfVarValueFromString(const FString& InValueStr, bool bIsTweak
 	}
 	EditSelfVar()->SetValueFromStr(DefaultValue);
 
-	EditSelfVar()->CopyTo(GetNodePtr()->GetExpression(), Argument);
 	NotifyPinSelfVarChanged(bIsTweaking);
 }
 
diff --git a/Engine/Plugins/Experimental/TextureGraph/Source/TextureGraphEditor/Private/STG_SelectionPreview.cpp b/Engine/Plugins/Experimental/TextureGraph/Source/TextureGraphEditor/Private/STG_SelectionPreview.cpp
--- a/Engine/Plugins/Experimental/TextureGraph/Source/TextureGraphEditor/Private/STG_SelectionPreview.cpp
+++ b/Engine/Plugins/Experimental/TextureGraph/Source/TextureGraphEditor/Private/STG_SelectionPreview.cpp
@@ -847,31 +847,38 @@ void STG_SelectionPreview::UpdateOutputDetailsText(BlobPtr InBlob)
 
 void STG_SelectionPreview::UpdateOutputDetailsText(UTG_Pin* Pin)
 {
-	if (Pin->IsArgColor())
+	if (Pin == nullptr)
 	{
-		FLinearColor ColorValue;
-		check(Pin->GetValue(ColorValue));
-		const FColor SRGBColor = ColorValue.ToFColorSRGB();
-		OutputDetailsText = "Color Value : " + SRGBColor.ToString();
-		PixelInfo = "";
-	}
-	else if (Pin->IsArgScalar())
-	{
-		float ScalarValue = 0.0;
-		check(Pin->GetValue(ScalarValue));
-		OutputDetailsText = "Scalar Value : " + FString::Printf(TEXT("%0.3f"), ScalarValue);
-		PixelInfo = "";
-	}
-	else if (Pin->IsArgVector())
-	{
-		FVector4f VectorValue(0,0,0,0);
-		check(Pin->GetValue(VectorValue));
-		OutputDetailsText = "Vector Value : " + VectorValue.ToString();
-		PixelInfo = "";
+		ResetOutputDetailsText();
 	}
 	else
 	{
-		ResetOutputDetailsText();
+		if (Pin->IsArgColor())
+		{
+			FLinearColor ColorValue;
+			check(Pin->GetValue(ColorValue));
+			const FColor SRGBColor = ColorValue.ToFColorSRGB();
+			OutputDetailsText = "Color Value : " + SRGBColor.ToString();
+			PixelInfo = "";
+		}
+		else if (Pin->IsArgScalar())
+		{
+			float ScalarValue = 0.0;
+			check(Pin->GetValue(ScalarValue));
+			OutputDetailsText = "Scalar Value : " + FString::Printf(TEXT("%0.3f"), ScalarValue);
+			PixelInfo = "";
+		}
+		else if (Pin->IsArgVector())
+		{
+			FVector4f VectorValue(0,0,0,0);
+			check(Pin->GetValue(VectorValue));
+			OutputDetailsText = "Vector Value : " + VectorValue.ToString();
+			PixelInfo = "";
+		}
+		else
+		{
+			ResetOutputDetailsText();
+		}
 	}
 }
 
