Allow texture graph on Mac and Linux

diff --git a/Engine/Plugins/Experimental/TextureGraph/TextureGraph.uplugin b/Engine/Plugins/Experimental/TextureGraph/TextureGraph.uplugin
--- a/Engine/Plugins/Experimental/TextureGraph/TextureGraph.uplugin
+++ b/Engine/Plugins/Experimental/TextureGraph/TextureGraph.uplugin
@@ -18,9 +18,7 @@
 			"Name": "TextureGraphEngine",
 			"Type": "Runtime",
 			"LoadingPhase": "PostConfigInit",
-			"PlatformAllowList": [
-				"Win64"
-			]
+			"PlatformAllowList": [ "Win64", "Linux", "Mac" ]
 		},
 		{
 			"Name": "TextureGraphInsight",
@@ -42,17 +40,13 @@
 			"Name": "TextureGraph",
 			"Type": "Runtime",
 			"LoadingPhase": "Default",
-			"PlatformAllowList": [
-				"Win64"
-			]
+			"PlatformAllowList": [ "Win64", "Linux", "Mac" ]
 		},
 		{
 			"Name": "TextureGraphEditor",
 			"Type": "Editor",
 			"LoadingPhase": "Default",
-			"PlatformAllowList": [
-				"Win64"
-			]
+			"PlatformAllowList": [ "Win64", "Linux", "Mac" ]
 		}
 	],
 	"Plugins": [
@@ -73,7 +67,5 @@
 			"Enabled": true
 		}
 	],
-	"SupportedTargetPlatforms": [
-		"Win64"
-	]
+	"SupportedTargetPlatforms": [ "Win64", "Linux", "Mac" ]
 }
\ No newline at end of file
