[Backout] - CL31207109
Original CL Desc
-----------------------------------------------------------------
Allow texture graph on Mac and Linux

diff --git a/Engine/Plugins/Experimental/TextureGraph/TextureGraph.uplugin b/Engine/Plugins/Experimental/TextureGraph/TextureGraph.uplugin
--- a/Engine/Plugins/Experimental/TextureGraph/TextureGraph.uplugin
+++ b/Engine/Plugins/Experimental/TextureGraph/TextureGraph.uplugin
@@ -18,7 +18,9 @@
 			"Name": "TextureGraphEngine",
 			"Type": "Runtime",
 			"LoadingPhase": "PostConfigInit",
-			"PlatformAllowList": [ "Win64", "Linux", "Mac" ]
+			"PlatformAllowList": [
+				"Win64"
+			]
 		},
 		{
 			"Name": "TextureGraphInsight",
@@ -40,13 +42,17 @@
 			"Name": "TextureGraph",
 			"Type": "Runtime",
 			"LoadingPhase": "Default",
-			"PlatformAllowList": [ "Win64", "Linux", "Mac" ]
+			"PlatformAllowList": [
+				"Win64"
+			]
 		},
 		{
 			"Name": "TextureGraphEditor",
 			"Type": "Editor",
 			"LoadingPhase": "Default",
-			"PlatformAllowList": [ "Win64", "Linux", "Mac" ]
+			"PlatformAllowList": [
+				"Win64"
+			]
 		}
 	],
 	"Plugins": [
@@ -67,5 +73,7 @@
 			"Enabled": true
 		}
 	],
-	"SupportedTargetPlatforms": [ "Win64", "Linux", "Mac" ]
+	"SupportedTargetPlatforms": [
+		"Win64"
+	]
 }
\ No newline at end of file
