Additional options on the Invert node along with thumbnails asserts fixes.

diff --git a/Engine/Plugins/Experimental/TextureGraph/Shaders/Expressions/Expression_Invert.usf b/Engine/Plugins/Experimental/TextureGraph/Shaders/Expressions/Expression_Invert.usf
--- a/Engine/Plugins/Experimental/TextureGraph/Shaders/Expressions/Expression_Invert.usf
+++ b/Engine/Plugins/Experimental/TextureGraph/Shaders/Expressions/Expression_Invert.usf
@@ -5,10 +5,29 @@
 Texture2D SourceTexture;
 float MaxValue;
 
+#ifndef INVERT_INCLUDE_ALPHA
+#define INVERT_INCLUDE_ALPHA 0
+#endif
+
+#ifndef INVERT_CLAMP
+#define INVERT_CLAMP 0
+#endif
+
 float4 FSH_Invert(float2 uv : TEXCOORD0) : SV_Target0
 {
 	float4 Color = SourceTexture.Sample(SamplerStates_NoBorder, uv);
-
-	// TODO: Support HDR values
-	return saturate(float4(MaxValue, MaxValue, MaxValue, MaxValue) - Color);
+	float4 InvertedColor;
+	
+#if INVERT_INCLUDE_ALPHA 
+	InvertedColor = float4(MaxValue, MaxValue, MaxValue, MaxValue) - Color;
+#else
+	float3 InvertedColorXYZ = float3(MaxValue, MaxValue, MaxValue) - Color.xyz;
+	InvertedColor = float4(InvertedColorXYZ.xyz, Color.a);
+#endif
+	
+#if INVERT_CLAMP 
+	InvertedColor = saturate(InvertedColor);
+#endif 
+	
+	return InvertedColor;
 }
diff --git a/Engine/Plugins/Experimental/TextureGraph/Source/TextureGraph/Private/Expressions/Maths/TG_Expression_Invert.cpp b/Engine/Plugins/Experimental/TextureGraph/Source/TextureGraph/Private/Expressions/Maths/TG_Expression_Invert.cpp
--- a/Engine/Plugins/Experimental/TextureGraph/Source/TextureGraph/Private/Expressions/Maths/TG_Expression_Invert.cpp
+++ b/Engine/Plugins/Experimental/TextureGraph/Source/TextureGraph/Private/Expressions/Maths/TG_Expression_Invert.cpp
@@ -18,5 +18,5 @@ void UTG_Expression_Invert::Evaluate(FTG_EvaluationContext* InContext)
 	const RenderMaterial_FXPtr RenderMaterial = TextureGraphEngine::GetMaterialManager()->CreateMaterial_FX<VSH_Simple, FSH_Invert>(TEXT("T_Invert"));
 	check(RenderMaterial);
 
-	Output = T_Invert::Create(InContext->Cycle, Output.GetBufferDescriptor(), Input, MaxValue, InContext->TargetId);
+	Output = T_Invert::Create(InContext->Cycle, Output.GetBufferDescriptor(), Input, MaxValue, InContext->TargetId, IncludeAlpha, Clamp);
 }
diff --git a/Engine/Plugins/Experimental/TextureGraph/Source/TextureGraph/Public/Expressions/Maths/TG_Expression_Invert.h b/Engine/Plugins/Experimental/TextureGraph/Source/TextureGraph/Public/Expressions/Maths/TG_Expression_Invert.h
--- a/Engine/Plugins/Experimental/TextureGraph/Source/TextureGraph/Public/Expressions/Maths/TG_Expression_Invert.h
+++ b/Engine/Plugins/Experimental/TextureGraph/Source/TextureGraph/Public/Expressions/Maths/TG_Expression_Invert.h
@@ -31,10 +31,18 @@ public:
 	UPROPERTY(meta = (TGType = "TG_Input", PinDisplayName = ""))
 	FTG_Texture							Input;
 
+	// Whether to invert the Alpha channel as well?
+	UPROPERTY(EditAnywhere, Category = NoCategory, meta = (TGType = "TG_Setting", PinNotConnectable = true))
+	bool								IncludeAlpha = false;
+
+	// Saturate/clamp values in the [0, 1] range or not
+	UPROPERTY(EditAnywhere, Category = NoCategory, meta = (TGType = "TG_Setting", PinNotConnectable = true))
+	bool								Clamp = false;
+
 	// The inverted output. This is calculated as MaxValue - Input for every pixel in every channel
 	UPROPERTY(EditAnywhere, Category = NoCategory, meta = (TGType = "TG_Output", PinDisplayName = ""))
 	FTG_Texture							Output;
 	
-	virtual FText						GetTooltipText() const override { return FText::FromString(TEXT("Inverts the input, i.e. computes one minus the pixel value.")); } 
+	virtual FText						GetTooltipText() const override { return FText::FromString(TEXT("Inverts the input, i.e. computes one minus the pixel value. Alpha is only inverted if the flag is checked.")); } 
 };
 
diff --git a/Engine/Plugins/Experimental/TextureGraph/Source/TextureGraphEditor/Private/EdGraph/STG_NodeThumbnail.cpp b/Engine/Plugins/Experimental/TextureGraph/Source/TextureGraphEditor/Private/EdGraph/STG_NodeThumbnail.cpp
--- a/Engine/Plugins/Experimental/TextureGraph/Source/TextureGraphEditor/Private/EdGraph/STG_NodeThumbnail.cpp
+++ b/Engine/Plugins/Experimental/TextureGraph/Source/TextureGraphEditor/Private/EdGraph/STG_NodeThumbnail.cpp
@@ -80,9 +80,8 @@ void STG_NodeThumbnail::UpdateParams(TiledBlobPtr InBlob)
 	if (!InBlob)
 		return;
 	
-	// we should assume that the block is finalised at this point
-	check (InBlob->IsFinalised());
-
+	/// If the blob hasn't been finalised then we don't do anything. This sort of scenario can happen if you end 
+	/// up with your thumbnail generation split between multiple batches on the ThumbnailsService
 	if (InBlob->IsFinalised())
 	{
 		UTexture* BlobTexture = GetTextureFromBlob(InBlob);
diff --git a/Engine/Plugins/Experimental/TextureGraph/Source/TextureGraphEngine/2D/Tex.cpp b/Engine/Plugins/Experimental/TextureGraph/Source/TextureGraphEngine/2D/Tex.cpp
--- a/Engine/Plugins/Experimental/TextureGraph/Source/TextureGraphEngine/2D/Tex.cpp
+++ b/Engine/Plugins/Experimental/TextureGraph/Source/TextureGraphEngine/2D/Tex.cpp
@@ -1145,6 +1145,7 @@ AsyncTiledBlobRef Tex::ToBlob(int32 XTiles, int32 YTiles, uint32 Width /* = 0 */
 			BlobDesc.Name = Desc.Name;
 			
 			TiledBlobPtr TiledBlobObj = TiledBlob::InitFromTiles(BlobDesc, ResultTiles);
+			TiledBlobObj->FinaliseNow(false, nullptr);
 
 			if (HashObj && HashObj->IsFinal())
 			{
diff --git a/Engine/Plugins/Experimental/TextureGraph/Source/TextureGraphEngine/Data/Blob.cpp b/Engine/Plugins/Experimental/TextureGraph/Source/TextureGraphEngine/Data/Blob.cpp
--- a/Engine/Plugins/Experimental/TextureGraph/Source/TextureGraphEngine/Data/Blob.cpp
+++ b/Engine/Plugins/Experimental/TextureGraph/Source/TextureGraphEngine/Data/Blob.cpp
@@ -40,6 +40,11 @@ Blob::Blob() : Buffer(Device_Null::Get()->Create(BufferDescriptor(), nullptr))
 Blob::Blob(DeviceBufferRef InBuffer) 
 	: Buffer(InBuffer)
 {
+	if (InBuffer && InBuffer->IsValid() && !InBuffer->IsNull())
+	{
+		bIsFinalised = true;
+		FinaliseTS = FDateTime::Now();
+	}
 }
 
 /// Allocate a NULL device buffer by default
diff --git a/Engine/Plugins/Experimental/TextureGraph/Source/TextureGraphEngine/Data/Blobber.cpp b/Engine/Plugins/Experimental/TextureGraph/Source/TextureGraphEngine/Data/Blobber.cpp
--- a/Engine/Plugins/Experimental/TextureGraph/Source/TextureGraphEngine/Data/Blobber.cpp
+++ b/Engine/Plugins/Experimental/TextureGraph/Source/TextureGraphEngine/Data/Blobber.cpp
@@ -46,8 +46,11 @@ Blobber::BlobCacheEntryPtr Blobber::FindInternal(HashType Hash)
 
 	FScopeLock CacheLock(&BlobLookupLock);
 	HashType OrgHash = Hash;
+	static const int32 MaxIter = 4;
 
-	while (Hash != DataUtil::GNullHash)
+	int32 NumIter = 0;
+
+	while (Hash != DataUtil::GNullHash && NumIter++ < MaxIter)
 	{
 		check(Hash != DataUtil::GNullHash);
 		const BlobCacheEntryPtr* BlobEntry = BlobCache.Find(Hash, false);
@@ -83,6 +86,14 @@ Blobber::BlobCacheEntryPtr Blobber::FindInternal(HashType Hash)
 			/// Otherwise, we need to use the RHS of the mapping
 			CHashPtr MappedHash = HashIter->second;
 			UE_LOG(LogBlob, VeryVerbose, TEXT("Blobber Find_Internal Remapped: %#016llu => %#016llu"), Hash, MappedHash->Value());
+
+			if (Hash == *MappedHash)
+			{
+				HashMappings.erase(HashIter);
+				return nullptr;
+			}
+
+			check(Hash != *MappedHash);
 			Hash = MappedHash->Value();
 
 			/// This is a cyclical link
diff --git a/Engine/Plugins/Experimental/TextureGraph/Source/TextureGraphEngine/Data/TiledBlob.cpp b/Engine/Plugins/Experimental/TextureGraph/Source/TextureGraphEngine/Data/TiledBlob.cpp
--- a/Engine/Plugins/Experimental/TextureGraph/Source/TextureGraphEngine/Data/TiledBlob.cpp
+++ b/Engine/Plugins/Experimental/TextureGraph/Source/TextureGraphEngine/Data/TiledBlob.cpp
@@ -997,6 +997,30 @@ void TiledBlob::SetLODLevel(int32 Level, BlobPtr Blob_, BlobPtrW LODParent_, Blo
 	});
 }
 
+void TiledBlob::FinaliseNow(bool bNoCalcHash, CHashPtr FixedHash)
+{
+	check(IsInGameThread());
+
+	UE_LOG(LogJob, Verbose, TEXT("Finalised tiled BlobObj: %s"), *Desc.Name);
+
+	/// Debug: check that all Tiles are valid
+	for (int32 TileX = 0; TileX < Tiles.Rows(); TileX++)
+	{
+		for (int32 TileY = 0; TileY < Tiles.Cols(); TileY++)
+		{
+			check(Tiles[TileX][TileY]);
+			BlobPtr Tile = Tiles[TileX][TileY];
+			Tile->FinaliseNow(bNoCalcHash, nullptr);
+		}
+	}
+
+	if (SingleBlob)
+		SingleBlob->FinaliseNow(bNoCalcHash, nullptr);
+
+	bIsFinalised = true;
+	FinaliseTS = FDateTime::Now();
+}
+
 //////////////////////////////////////////////////////////////////////////
 TiledBlob_Promise::TiledBlob_Promise(TiledBlobPtr Source) : TiledBlob(Source->GetDescriptor(), Source->GetTiles())
 {
@@ -1059,20 +1083,7 @@ void TiledBlob_Promise::FinaliseNow(bool bNoCalcHash, CHashPtr FixedHash)
 	check(IsInGameThread());
 
 	UE_LOG(LogJob, Verbose, TEXT("Finalised promised BlobObj: %s"), *Desc.Name);
-
-	/// Debug: check that all Tiles are valid
-	for (int32 TileX = 0; TileX < Tiles.Rows(); TileX++)
-	{
-		for (int32 TileY = 0; TileY < Tiles.Cols(); TileY++)
-		{
-			check(Tiles[TileX][TileY]);
-			BlobPtr Tile = Tiles[TileX][TileY];
-			Tile->FinaliseNow(bNoCalcHash, nullptr);
-		}
-	}
-
-	bIsFinalised = true;
-	FinaliseTS = FDateTime::Now();
+	TiledBlob::FinaliseNow(bNoCalcHash, FixedHash);
 
 	/// Trigger the callbacks - Should always be the last step of finalize
 	NotifyCallbacks();
diff --git a/Engine/Plugins/Experimental/TextureGraph/Source/TextureGraphEngine/Data/TiledBlob.h b/Engine/Plugins/Experimental/TextureGraph/Source/TextureGraphEngine/Data/TiledBlob.h
--- a/Engine/Plugins/Experimental/TextureGraph/Source/TextureGraphEngine/Data/TiledBlob.h
+++ b/Engine/Plugins/Experimental/TextureGraph/Source/TextureGraphEngine/Data/TiledBlob.h
@@ -83,6 +83,7 @@ public:
 	virtual void					ResolveLateBound(const BufferDescriptor& Desc, bool bOverrideExisting = false);
 	virtual void					CopyResolveLateBound(BlobPtr RHS);
 	virtual void					SetTransient();
+	virtual void					FinaliseNow(bool bNoCalcHash, CHashPtr FixedHash) override;
 
 	//////////////////////////////////////////////////////////////////////////
 	/// Min/Max related
diff --git a/Engine/Plugins/Experimental/TextureGraph/Source/TextureGraphEngine/Job/Job.cpp b/Engine/Plugins/Experimental/TextureGraph/Source/TextureGraphEngine/Job/Job.cpp
--- a/Engine/Plugins/Experimental/TextureGraph/Source/TextureGraphEngine/Job/Job.cpp
+++ b/Engine/Plugins/Experimental/TextureGraph/Source/TextureGraphEngine/Job/Job.cpp
@@ -392,19 +392,16 @@ bool Job::CheckCached()
 	{
 		check(ExistingResult->IsTiled());
 
-		BufferDescriptor ResultDesc = ExistingResult->GetDescriptor();
-
-		if (ResultDesc.bIsTransient == IsDiscard())
-		{
-			Result = ExistingResult;
-		}
+		if (ResultOrg)
+			ResultOrg->FinaliseFrom(ExistingResult.get());
 		else
 		{
 			ResultOrg = std::make_shared<TiledBlob_Promise>(ExistingResult->GetDescriptor(), ExistingResult->Rows(), ExistingResult->Cols(), ExistingResult->Hash());
 			ResultOrg->FinaliseFrom(ExistingResult.get());
-			Result = TiledBlobRef(std::static_pointer_cast<TiledBlob>(ResultOrg), true, false);
 		}
 
+		Result = TiledBlobRef(std::static_pointer_cast<TiledBlob>(ResultOrg), true, false);
+
 		FinalJobResult = std::make_shared<JobResult>(GetResultRef(), nullptr);
 		bIsCulled = true;
 
@@ -1061,6 +1058,7 @@ void Job::MarkJobDone()
 {
 	Stats.EndNativeTime = Util::Time();
 
+	check((!Result || Result->IsFinalised()) && (!ResultOrg || ResultOrg->IsFinalised()));
 	if (RunInfo.Batch)
 		RunInfo.Batch->OnJobDone(this, GetJobId());
 
@@ -1097,7 +1095,7 @@ void Job::AddResultToBlobber()
 #endif 
 
 		Result = TextureGraphEngine::GetBlobber()->AddTiledResult(TempHash, ResultOrg, CacheOpt);
-		ResultOrg = nullptr;
+		//ResultOrg = nullptr;
 
 		check(Result && Result->IsTiled());
 	}
@@ -1198,7 +1196,6 @@ int32 Job::Exec()
 	Run(RunInfo).apply(cti::transforms::wait());
 	Stats.EndNativeTime = Util::Time();
 
-	MarkJobDone();
 	SetPromise(0);
 
 	return 0;
@@ -1251,13 +1248,13 @@ void Job::GetDependencies(JobPtrVec& Prior, JobPtrVec& After, JobRunInfo InRunIn
 	FinaliseJob->RunInfo = RunInfoCopy;
 	FinaliseJob->Generator = ThisJob;
 
-	JobRunInfo prepareRunInfo = RunInfo;
-	prepareRunInfo.Dev = PrepareJob->GetTransform()->TargetDevice(0);
-	PrepareJob->GetDependencies(Prior, After, prepareRunInfo);
+	JobRunInfo PrepareRunInfo = RunInfo;
+	PrepareRunInfo.Dev = PrepareJob->GetTransform()->TargetDevice(0);
+	PrepareJob->GetDependencies(Prior, After, PrepareRunInfo);
 
-	JobRunInfo afterRunInfo = RunInfo;
-	afterRunInfo.Dev = FinaliseJob->GetTransform()->TargetDevice(0);
-	FinaliseJob->GetDependencies(Prior, After, afterRunInfo);
+	JobRunInfo AfterRunInfo = RunInfo;
+	AfterRunInfo.Dev = FinaliseJob->GetTransform()->TargetDevice(0);
+	FinaliseJob->GetDependencies(Prior, After, AfterRunInfo);
 
 	JobsGeneratedPrior.push_back(PrepareJob);
 	JobsGeneratedAfter.push_back(FinaliseJob);
@@ -1291,13 +1288,7 @@ void Job::GetDependencies(JobPtrVec& Prior, JobPtrVec& After, JobRunInfo InRunIn
 
 void Job::PostExec()
 {
-	//size_t threadId = Util::GetCurrentThreadId();
-
-	//check(threadId > 0 && threadId == _runInfo.device->GetPostExecThreadId());
-
 	UE_LOG(LogJob, VeryVerbose, TEXT("PostExec::%s"), *Transform->GetName());
-
-
 	DeviceNativeTask::PostExec();
 }
 
diff --git a/Engine/Plugins/Experimental/TextureGraph/Source/TextureGraphEngine/Job/Job.h b/Engine/Plugins/Experimental/TextureGraph/Source/TextureGraphEngine/Job/Job.h
--- a/Engine/Plugins/Experimental/TextureGraph/Source/TextureGraphEngine/Job/Job.h
+++ b/Engine/Plugins/Experimental/TextureGraph/Source/TextureGraphEngine/Job/Job.h
@@ -76,7 +76,6 @@ protected:
 	JobPtrW							Generator;			/// The job that generated this job. This is useful in certain situations
 	FString							DebugJobName;		/// The name of the job
 
-
 	//std::vector<std::shared_ptr<cti::promise<TiledBlobPtr>>> _blobPreparedPromises;
 
 	typedef std::pair<int32, int32>	IndexPair;
@@ -182,10 +181,9 @@ public:
 	FORCEINLINE void				SetQueueId(int32 NewQueueId) { QueueId = NewQueueId; }
 	FORCEINLINE BlobTransformPtr	GetTransform() const { return Transform; }
 	FORCEINLINE TiledBlobRef		GetResult() const { return Result; }
-	
 
 	FORCEINLINE bool				GetTiled() const { return bIsTiled; }
-	FORCEINLINE void				SetTiled(bool bTiled) { bIsTiled = bTiled; }
+	FORCEINLINE Job*				SetTiled(bool bTiled) { bIsTiled = bTiled; return this; }
 
 	FORCEINLINE JobStats			GetStats() const { return Stats; }
 
@@ -213,7 +211,6 @@ public:
 	FORCEINLINE TiledBlob_PromisePtr GetResultPromise() const { return std::static_pointer_cast<TiledBlob_Promise>(Result.get()); }
 
 	FORCEINLINE UObject*			GetErrorOwner() const { return ErrorOwner.Get(); }
-
 };
 
 //////////////////////////////////////////////////////////////////////////
diff --git a/Engine/Plugins/Experimental/TextureGraph/Source/TextureGraphEngine/Transform/Expressions/T_Invert.cpp b/Engine/Plugins/Experimental/TextureGraph/Source/TextureGraphEngine/Transform/Expressions/T_Invert.cpp
--- a/Engine/Plugins/Experimental/TextureGraph/Source/TextureGraphEngine/Transform/Expressions/T_Invert.cpp
+++ b/Engine/Plugins/Experimental/TextureGraph/Source/TextureGraphEngine/Transform/Expressions/T_Invert.cpp
@@ -10,19 +10,24 @@
 
 IMPLEMENT_GLOBAL_SHADER(FSH_Invert, "/Plugin/TextureGraph/Expressions/Expression_Invert.usf", "FSH_Invert", SF_Pixel);
 
-TiledBlobPtr T_Invert::Create(MixUpdateCyclePtr Cycle, BufferDescriptor DesiredOutputDesc, TiledBlobPtr Source, float MaxValue, int32 TargetId)
+TiledBlobPtr T_Invert::Create(MixUpdateCyclePtr Cycle, BufferDescriptor DesiredOutputDesc, TiledBlobPtr Source, float MaxValue, int32 TargetId, bool bIncludeAlpha, bool bClamp)
 {
-	RenderMaterial_FXPtr RenderMaterial = TextureGraphEngine::GetMaterialManager()->CreateMaterial_FX<VSH_Simple, FSH_Invert>(TEXT("T_Invert"));
+	FSH_Invert::FPermutationDomain PermutationVector;
+	PermutationVector.Set<FSH_Invert::FInvertIncludeAlpha>(bIncludeAlpha);
+	PermutationVector.Set<FSH_Invert::FInvertClamp>(bClamp);
 
+	RenderMaterial_FXPtr RenderMaterial = TextureGraphEngine::GetMaterialManager()->CreateMaterial_FX<VSH_Simple, FSH_Invert>(TEXT("T_Invert"), PermutationVector);
 	check(RenderMaterial);
 
 	JobUPtr RenderJob = std::make_unique<Job>(Cycle->GetMix(), TargetId, std::static_pointer_cast<BlobTransform>(RenderMaterial));
 	RenderJob
 		->AddArg(ARG_BLOB(Source, "SourceTexture"))
 		->AddArg(ARG_FLOAT(MaxValue, "MaxValue"))
+		->AddArg(WithUnbounded(ARG_BOOL(bIncludeAlpha, "IncludeAlpha")))
+		->AddArg(WithUnbounded(ARG_BOOL(bClamp, "Clamp")))
 		;
 
-	FString Name = TEXT("Invert"); // FString::Printf(TEXT("GaussianBlur.[%s].[%d].[%llu]"), *gmask->ID(), InContext->TargetId, InContext->Cycle->Batch()->BatchId());
+	FString Name = TEXT("Invert");
 
 	TiledBlobPtr Result = RenderJob->InitResult(Name, &DesiredOutputDesc);
 	Cycle->AddJob(TargetId, std::move(RenderJob));
diff --git a/Engine/Plugins/Experimental/TextureGraph/Source/TextureGraphEngine/Transform/Expressions/T_Invert.h b/Engine/Plugins/Experimental/TextureGraph/Source/TextureGraphEngine/Transform/Expressions/T_Invert.h
--- a/Engine/Plugins/Experimental/TextureGraph/Source/TextureGraphEngine/Transform/Expressions/T_Invert.h
+++ b/Engine/Plugins/Experimental/TextureGraph/Source/TextureGraphEngine/Transform/Expressions/T_Invert.h
@@ -22,6 +22,10 @@ public:
 		SHADER_PARAMETER(float, MaxValue)
 	END_SHADER_PARAMETER_STRUCT()
 
+	class FInvertIncludeAlpha : SHADER_PERMUTATION_BOOL("INVERT_INCLUDE_ALPHA");
+	class FInvertClamp : SHADER_PERMUTATION_BOOL("INVERT_CLAMP");
+	using FPermutationDomain = TShaderPermutationDomain<FInvertIncludeAlpha, FInvertClamp>;
+
 public:
 	static bool						ShouldCompilePermutation(const FGlobalShaderPermutationParameters& Parameters) 
 	{
@@ -36,5 +40,5 @@ public:
 	//////////////////////////////////////////////////////////////////////////
 	/// Static functions
 	//////////////////////////////////////////////////////////////////////////
-	static TiledBlobPtr				Create(MixUpdateCyclePtr Cycle, BufferDescriptor DesiredOutputDesc, TiledBlobPtr Source, float MaxValue, int32 TargetId);
+	static TiledBlobPtr				Create(MixUpdateCyclePtr Cycle, BufferDescriptor DesiredOutputDesc, TiledBlobPtr Source, float MaxValue, int32 TargetId, bool bIncludeAlpha, bool bClamp);
 };
diff --git a/Engine/Plugins/Experimental/TextureGraph/Source/TextureGraphEngine/Transform/Layer/T_Thumbnail.cpp b/Engine/Plugins/Experimental/TextureGraph/Source/TextureGraphEngine/Transform/Layer/T_Thumbnail.cpp
--- a/Engine/Plugins/Experimental/TextureGraph/Source/TextureGraphEngine/Transform/Layer/T_Thumbnail.cpp
+++ b/Engine/Plugins/Experimental/TextureGraph/Source/TextureGraphEngine/Transform/Layer/T_Thumbnail.cpp
@@ -25,6 +25,9 @@ JobBatchPtr AddThumbJobToCycle(TiledBlobPtr InBlobToBind, JobUPtr JobPtr, UObjec
 
 TiledBlobRef T_Thumbnail::Bind(UMixInterface* Mix, UObject* Model, TiledBlobPtr InBlobToBind, int32 InTargetId)
 {
+	if (InBlobToBind->IsTransient())
+		return TextureHelper::GetBlack();
+
 	check(InBlobToBind);
 
 	UE_LOG(LogJob, VeryVerbose, TEXT("T_Thumbnail::Bind [%s]"), *InBlobToBind->Name());
@@ -33,13 +36,18 @@ TiledBlobRef T_Thumbnail::Bind(UMixInterface* Mix, UObject* Model, TiledBlobPtr
 	RenderMaterial_FXPtr RenderMaterial = TextureGraphEngine::GetMaterialManager()->CreateMaterialOfType_FX<Fx_FullScreenCopy>(ThmJobName);
 
 	JobUPtr JobPtr = std::make_unique<Job>(Mix, InTargetId, std::static_pointer_cast<BlobTransform>(RenderMaterial));
-	JobPtr->AddArg(ARG_BLOB(InBlobToBind, "SourceTexture"));
-	JobPtr->SetTiled(false);
+	JobPtr
+		->AddArg(ARG_BLOB(InBlobToBind, "SourceTexture"))
+		->SetTiled(false)
+		;
+
 	BufferDescriptor Desc = InBlobToBind->GetDescriptor();
 	Desc.Width = RenderMaterial_Thumbnail::GThumbWidth;
 	Desc.Height = RenderMaterial_Thumbnail::GThumbHeight;
+
+	FString Name = FString::Printf(TEXT("T_Thumbnail [%s]"), *InBlobToBind->GetDescriptor().Name);
 	
-	TiledBlobPtr Result = JobPtr->InitResult(TEXT("T_Thumbnail"), &Desc);
+	TiledBlobPtr Result = JobPtr->InitResult(Name, &Desc);
 	Result->MakeSingleBlob();
 
 	AddThumbJobToCycle(InBlobToBind, std::move(JobPtr), Model, InTargetId);
diff --git a/Engine/Plugins/Experimental/TextureGraph/Source/TextureGraphEngine/Transform/Utility/T_FinaliseBlob.cpp b/Engine/Plugins/Experimental/TextureGraph/Source/TextureGraphEngine/Transform/Utility/T_FinaliseBlob.cpp
--- a/Engine/Plugins/Experimental/TextureGraph/Source/TextureGraphEngine/Transform/Utility/T_FinaliseBlob.cpp
+++ b/Engine/Plugins/Experimental/TextureGraph/Source/TextureGraphEngine/Transform/Utility/T_FinaliseBlob.cpp
@@ -72,23 +72,15 @@ protected:
 						return JobResult->Finalise(true, nullptr);
 				}
 
-				EndNative();
-
 				return static_cast<AsyncBufferResultPtr>(cti::make_ready_continuable<BufferResultPtr>(std::make_shared<BufferResult>()));
 			})
-			.then([this](BufferResultPtr result) mutable
+			.then([this](BufferResultPtr) mutable
 			{
-				auto JobResult = OriginalJobObj->GetResultRef();
+				BlobRef JobResult = OriginalJobObj->GetResultRef();
 
 				if (JobResult)
 				{
 					OriginalJobObj->AddResultToBlobber();
-					//TextureGraphEngine::Blobber()->AddTiledResult(jobResult->Hash(), jobResult,
-					//	{
-					//		false,
-					//		_runInfo.cycle->NoCache()
-					//	});
-
 					TransformResultPtr TResult = std::make_shared<TransformResult>();
 					TResult->Target = JobResult;
 				}
