Fixing aproblem when undoing a change that cause a Dynamic Expression to change signature

#changelist validated
#virtualized

diff --git a/Engine/Plugins/Experimental/TextureGraph/Source/TextureGraph/Private/Expressions/Input/TG_Expression_Graph.cpp b/Engine/Plugins/Experimental/TextureGraph/Source/TextureGraph/Private/Expressions/Input/TG_Expression_Graph.cpp
--- a/Engine/Plugins/Experimental/TextureGraph/Source/TextureGraph/Private/Expressions/Input/TG_Expression_Graph.cpp
+++ b/Engine/Plugins/Experimental/TextureGraph/Source/TextureGraph/Private/Expressions/Input/TG_Expression_Graph.cpp
@@ -108,6 +108,15 @@ void UTG_Expression_TextureGraph::PostEditChangeProperty(FPropertyChangedEvent&
 	Super::PostEditChangeProperty(PropertyChangedEvent);
 }
 
+void UTG_Expression_TextureGraph::PostEditUndo()
+{
+	// Make sure the signature is in sync after undo in case we undo a TextureGRaph assignment:
+	// So recreate it internally without notifying, normally, the node's pins should match
+	DynSignature.Reset();
+	GetSignature();
+
+	Super::PostEditUndo();
+}
 #endif
 
 void UTG_Expression_TextureGraph::OnTextureGraphPreSave(UObject* Object, FObjectPreSaveContext SaveContext)
diff --git a/Engine/Plugins/Experimental/TextureGraph/Source/TextureGraph/Private/Expressions/Input/TG_Expression_Material.cpp b/Engine/Plugins/Experimental/TextureGraph/Source/TextureGraph/Private/Expressions/Input/TG_Expression_Material.cpp
--- a/Engine/Plugins/Experimental/TextureGraph/Source/TextureGraph/Private/Expressions/Input/TG_Expression_Material.cpp
+++ b/Engine/Plugins/Experimental/TextureGraph/Source/TextureGraph/Private/Expressions/Input/TG_Expression_Material.cpp
@@ -14,6 +14,17 @@ void UTG_Expression_Material::PostEditChangeProperty(FPropertyChangedEvent& Prop
 
 	Super::PostEditChangeProperty(PropertyChangedEvent);
 }
+
+void UTG_Expression_Material::PostEditUndo()
+{
+	// Make sure the signature is in sync after undo in case we undo a material assignment:
+	// So recreate it internally without notifying, normally, the node's pins should match
+	DynSignature.Reset();
+	GetSignature();
+
+	Super::PostEditUndo();
+}
+
 #endif
 
 void UTG_Expression_Material::SetMaterialInternal(UMaterialInterface* InMaterial)
diff --git a/Engine/Plugins/Experimental/TextureGraph/Source/TextureGraph/Private/Expressions/Input/TG_Expression_MaterialFunction.cpp b/Engine/Plugins/Experimental/TextureGraph/Source/TextureGraph/Private/Expressions/Input/TG_Expression_MaterialFunction.cpp
--- a/Engine/Plugins/Experimental/TextureGraph/Source/TextureGraph/Private/Expressions/Input/TG_Expression_MaterialFunction.cpp
+++ b/Engine/Plugins/Experimental/TextureGraph/Source/TextureGraph/Private/Expressions/Input/TG_Expression_MaterialFunction.cpp
@@ -45,6 +45,16 @@ void UTG_Expression_MaterialFunction::PostEditChangeProperty(FPropertyChangedEve
 
 	Super::PostEditChangeProperty(PropertyChangedEvent);
 }
+
+void UTG_Expression_MaterialFunction::PostEditUndo()
+{
+	// Make sure the signature is in sync after undo in case we undo a materialFunction assignment:
+	// So recreate it internally without notifying, normally, the node's pins should match
+	DynSignature.Reset();
+	GetSignature();
+
+	Super::PostEditUndo();
+}
 #endif
 
 void UTG_Expression_MaterialFunction::Initialize()
diff --git a/Engine/Plugins/Experimental/TextureGraph/Source/TextureGraph/Private/TG_Graph.cpp b/Engine/Plugins/Experimental/TextureGraph/Source/TextureGraph/Private/TG_Graph.cpp
--- a/Engine/Plugins/Experimental/TextureGraph/Source/TextureGraph/Private/TG_Graph.cpp
+++ b/Engine/Plugins/Experimental/TextureGraph/Source/TextureGraph/Private/TG_Graph.cpp
@@ -11,7 +11,7 @@
 #if WITH_EDITOR
 void UTG_Graph::PostEditChangeProperty(FPropertyChangedEvent& PropertyChangedEvent)
 {
-	UE_LOG(LogTextureGraph, Log, TEXT("PostEditChangeProperty::Graph Changed."));
+	UE_LOG(LogTextureGraph, Log, TEXT("UTG_Graph::PostEditChangeProperty."));
 }
 
 bool UTG_Graph::Modify(bool bAlwaysMarkDirty)
@@ -19,11 +19,18 @@ bool UTG_Graph::Modify(bool bAlwaysMarkDirty)
 	// Runtime state is about to get dirty
 	// And remember it as such for undo redo
 	bIsGraphTraversalDirty = true;
-	UE_LOG(LogTextureGraph, Log, TEXT("UTG_Graph::Modify: Graph Changed."));
+	UE_LOG(LogTextureGraph, Log, TEXT("UTG_Graph::Modify: Graph Modified."));
 
 	return Super::Modify(bAlwaysMarkDirty);
 }
 
+
+void UTG_Graph::PostEditUndo()
+{
+	UE_LOG(LogTextureGraph, Log, TEXT("UTG_Graph::PostEditUndo."));
+	UObject::PostEditUndo();
+}
+
 #endif
 void UTG_Graph::Reset()
 {
@@ -250,11 +257,11 @@ void UTG_Graph::RegenerateNode(UTG_Node* InNode)
 void UTG_Graph::AllocateNodePins(UTG_Node* Node)
 {
 	// All the pins are created from the signature 
-	for (auto& Arg : Node->Signature->GetInArguments())
+	for (auto& Arg : Node->GetSignature().GetInArguments())
 		AllocatePin(Node, Arg);
-	for (auto& Arg : Node->Signature->GetOutArguments())
+	for (auto& Arg : Node->GetSignature().GetOutArguments())
 		AllocatePin(Node, Arg);
-	for (auto& Arg : Node->Signature->GetPrivateArguments())
+	for (auto& Arg : Node->GetSignature().GetPrivateArguments())
 		AllocatePin(Node, Arg);
 }
 
diff --git a/Engine/Plugins/Experimental/TextureGraph/Source/TextureGraph/Private/TG_Node.cpp b/Engine/Plugins/Experimental/TextureGraph/Source/TextureGraph/Private/TG_Node.cpp
--- a/Engine/Plugins/Experimental/TextureGraph/Source/TextureGraph/Private/TG_Node.cpp
+++ b/Engine/Plugins/Experimental/TextureGraph/Source/TextureGraph/Private/TG_Node.cpp
@@ -34,7 +34,7 @@ FTG_Name UTG_Node::GetNodeName() const
 #if WITH_EDITOR
 void UTG_Node::PostEditChangeProperty(FPropertyChangedEvent& PropertyChangedEvent)
 {
-	UE_LOG(LogTextureGraph, Log, TEXT("Node Changed."));
+	UE_LOG(LogTextureGraph, Log, TEXT("UTG_Node::PostEditChangeProperty."));
 }
 
 void UTG_Node::PostEditUndo()
@@ -54,7 +54,10 @@ void UTG_Node::OnPostUndo()
 		{
 			auto Arg = Pin->GetArgument();
 			Pin->EditSelfVar()->CopyFrom(GetExpression(), Arg);	
+
 		}
+
+		Signature = GetExpression()->GetSignature();
 	}
 }
 #endif
@@ -255,7 +258,7 @@ void UTG_Node::Initialize(FTG_Id InId)
 
 int32 UTG_Node::ForEachInputPins(std::function<void(const UTG_Pin* /*pin*/, uint32 /*index*/)> visitor) const
 {
-	int32 Num = Signature->GetInArguments().Num();
+	int32 Num = GetSignature().GetInArguments().Num();
 	for (int32 i = 0; i < Num; ++i)
 	{
 		visitor(Pins[i], i);
@@ -265,8 +268,8 @@ int32 UTG_Node::ForEachInputPins(std::function<void(const UTG_Pin* /*pin*/, uint
 
 int32 UTG_Node::ForEachOutputPins(std::function<void(const UTG_Pin* /*pin*/, uint32 /*index*/)> visitor) const
 {
-	int32 Offset = Signature->GetInArguments().Num();
-	int32 Num = Signature->GetOutArguments().Num();
+	int32 Offset = GetSignature().GetInArguments().Num();
+	int32 Num = GetSignature().GetOutArguments().Num();
 	for (int32 i = 0; i < Num; ++i)
 	{
 		visitor(Pins[i + Offset], i);
@@ -320,20 +323,20 @@ FTG_Id UTG_Node::GetPinId(FName Name) const
 
 FTG_Id UTG_Node::GetInputPinIdAt(FTG_Index inIndex) const
 {
-	if (inIndex >= 0 && inIndex < Signature->GetInArguments().Num())		
+	if (inIndex >= 0 && inIndex < GetSignature().GetInArguments().Num())		
 		return FTG_Id(GetId().NodeIdx(), inIndex);
 	return FTG_Id();
 }
 FTG_Id UTG_Node::GetOutputPinIdAt(FTG_Index outIndex) const
 {
-	if (outIndex >= 0 && outIndex < Signature->GetOutArguments().Num())
-		return FTG_Id(GetId().NodeIdx(), outIndex + Signature->GetInArguments().Num());
+	if (outIndex >= 0 && outIndex < GetSignature().GetOutArguments().Num())
+		return FTG_Id(GetId().NodeIdx(), outIndex + GetSignature().GetInArguments().Num());
 	return FTG_Id();
 }
 FTG_Id UTG_Node::GetPrivatePinIdAt(FTG_Index privateIndex) const
 {
-	if (privateIndex >= 0 && privateIndex < Signature->GetPrivateArguments().Num())
-		return FTG_Id(GetId().NodeIdx(), privateIndex + Signature->GetInArguments().Num() + Signature->GetOutArguments().Num());
+	if (privateIndex >= 0 && privateIndex < GetSignature().GetPrivateArguments().Num())
+		return FTG_Id(GetId().NodeIdx(), privateIndex + GetSignature().GetInArguments().Num() + GetSignature().GetOutArguments().Num());
 	return FTG_Id();
 }
 
@@ -350,32 +353,32 @@ TArray<FTG_Id> MakeIdArray(int32 NodeIdx, int32 Num, int32 IdxOffset)
 
 TArray<FTG_Id> UTG_Node::GetInputPinIds() const
 {
-	return MakeIdArray(GetId().NodeIdx(), Signature->GetInArguments().Num(), 0);
+	return MakeIdArray(GetId().NodeIdx(), GetSignature().GetInArguments().Num(), 0);
 }
 TArray<FTG_Id> UTG_Node::GetOutputPinIds() const
 {
-	int32 IdxOffset = Signature->GetInArguments().Num();
-	return MakeIdArray(GetId().NodeIdx(), Signature->GetOutArguments().Num(), IdxOffset);
+	int32 IdxOffset = GetSignature().GetInArguments().Num();
+	return MakeIdArray(GetId().NodeIdx(), GetSignature().GetOutArguments().Num(), IdxOffset);
 }
 TArray<FTG_Id> UTG_Node::GetPrivatePinIds() const
 {
-	int32 IdxOffset = Signature->GetInArguments().Num() + Signature->GetOutArguments().Num();
-	return MakeIdArray(GetId().NodeIdx(), Signature->GetPrivateArguments().Num(), IdxOffset);
+	int32 IdxOffset = GetSignature().GetInArguments().Num() + GetSignature().GetOutArguments().Num();
+	return MakeIdArray(GetId().NodeIdx(), GetSignature().GetPrivateArguments().Num(), IdxOffset);
 }
 
 FTG_Id UTG_Node::GetInputPinID(FTG_Name& Name) const
 {
-    FTG_Index i = Signature->FindInputArgument(Name);
+    FTG_Index i = GetSignature().FindInputArgument(Name);
     return (i == FTG_Id::INVALID_INDEX ? FTG_Id::INVALID : GetInputPinIdAt(i));
 }
 FTG_Id UTG_Node::GetOutputPinID(FTG_Name& Name) const
 {
-	FTG_Index i = Signature->FindOutputArgument(Name);
+	FTG_Index i = GetSignature().FindOutputArgument(Name);
 	return (i == FTG_Id::INVALID_INDEX ? FTG_Id::INVALID : GetOutputPinIdAt(i));
 }
 FTG_Id UTG_Node::GetPrivatePinID(FTG_Name& Name) const
 {
-	FTG_Index i = Signature->FindPrivateArgument(Name);
+	FTG_Index i = GetSignature().FindPrivateArgument(Name);
 	return (i == FTG_Id::INVALID_INDEX ? FTG_Id::INVALID : GetPrivatePinIdAt(i));
 }
 
diff --git a/Engine/Plugins/Experimental/TextureGraph/Source/TextureGraph/Public/Expressions/Input/TG_Expression_Graph.h b/Engine/Plugins/Experimental/TextureGraph/Source/TextureGraph/Public/Expressions/Input/TG_Expression_Graph.h
--- a/Engine/Plugins/Experimental/TextureGraph/Source/TextureGraph/Public/Expressions/Input/TG_Expression_Graph.h
+++ b/Engine/Plugins/Experimental/TextureGraph/Source/TextureGraph/Public/Expressions/Input/TG_Expression_Graph.h
@@ -48,6 +48,7 @@ public:
 	
 #if WITH_EDITOR
 	virtual void PostEditChangeProperty(FPropertyChangedEvent& PropertyChangedEvent) override;
+	virtual void PostEditUndo() override;
 #endif
 
 	virtual void Evaluate(FTG_EvaluationContext* InContext) override;
diff --git a/Engine/Plugins/Experimental/TextureGraph/Source/TextureGraph/Public/Expressions/Input/TG_Expression_Material.h b/Engine/Plugins/Experimental/TextureGraph/Source/TextureGraph/Public/Expressions/Input/TG_Expression_Material.h
--- a/Engine/Plugins/Experimental/TextureGraph/Source/TextureGraph/Public/Expressions/Input/TG_Expression_Material.h
+++ b/Engine/Plugins/Experimental/TextureGraph/Source/TextureGraph/Public/Expressions/Input/TG_Expression_Material.h
@@ -18,7 +18,8 @@ class TEXTUREGRAPH_API UTG_Expression_Material : public UTG_Expression_MaterialB
 	GENERATED_BODY()
 public:
 #if WITH_EDITOR
-	void PostEditChangeProperty(FPropertyChangedEvent& PropertyChangedEvent) override;
+	virtual void PostEditChangeProperty(FPropertyChangedEvent& PropertyChangedEvent) override;
+	virtual void PostEditUndo() override;
 #endif
 
 	// The input material to employ for rendering
diff --git a/Engine/Plugins/Experimental/TextureGraph/Source/TextureGraph/Public/Expressions/Input/TG_Expression_MaterialFunction.h b/Engine/Plugins/Experimental/TextureGraph/Source/TextureGraph/Public/Expressions/Input/TG_Expression_MaterialFunction.h
--- a/Engine/Plugins/Experimental/TextureGraph/Source/TextureGraph/Public/Expressions/Input/TG_Expression_MaterialFunction.h
+++ b/Engine/Plugins/Experimental/TextureGraph/Source/TextureGraph/Public/Expressions/Input/TG_Expression_MaterialFunction.h
@@ -17,7 +17,8 @@ class TEXTUREGRAPH_API UTG_Expression_MaterialFunction : public UTG_Expression_M
 	GENERATED_BODY()
 public:
 #if WITH_EDITOR
-	void PostEditChangeProperty(FPropertyChangedEvent& PropertyChangedEvent) override;
+	virtual void PostEditChangeProperty(FPropertyChangedEvent& PropertyChangedEvent) override;
+	virtual void PostEditUndo() override;
 #endif
 	
 	// The material function to employ for rendering
diff --git a/Engine/Plugins/Experimental/TextureGraph/Source/TextureGraph/Public/TG_Graph.h b/Engine/Plugins/Experimental/TextureGraph/Source/TextureGraph/Public/TG_Graph.h
--- a/Engine/Plugins/Experimental/TextureGraph/Source/TextureGraph/Public/TG_Graph.h
+++ b/Engine/Plugins/Experimental/TextureGraph/Source/TextureGraph/Public/TG_Graph.h
@@ -102,6 +102,9 @@ protected:
 protected:
 	virtual void PostEditChangeProperty(FPropertyChangedEvent& PropertyChangedEvent) override;
 
+	// Override PostEditUndo method of UObject
+	virtual void PostEditUndo() override;
+
 public:
 	// Overwrite the modify to also mark the transient states as dirty
 	virtual bool Modify(bool bAlwaysMarkDirty = true) override;
diff --git a/Engine/Plugins/Experimental/TextureGraph/Source/TextureGraph/Public/TG_Node.h b/Engine/Plugins/Experimental/TextureGraph/Source/TextureGraph/Public/TG_Node.h
--- a/Engine/Plugins/Experimental/TextureGraph/Source/TextureGraph/Public/TG_Node.h
+++ b/Engine/Plugins/Experimental/TextureGraph/Source/TextureGraph/Public/TG_Node.h
@@ -44,7 +44,7 @@ class TEXTUREGRAPH_API UTG_Node : public UObject
     GENERATED_BODY()
 
     friend class UTG_Graph;
-	FTG_SignaturePtr Signature;
+	TSharedPtr<FTG_Signature> Signature;
 
 	// Note: This property is assigned at construction of the node in the Construct call
 	UPROPERTY()
