Missing tooltips

diff --git a/Engine/Plugins/Experimental/TextureGraph/Source/TextureGraph/Private/TG_Pin.cpp b/Engine/Plugins/Experimental/TextureGraph/Source/TextureGraph/Private/TG_Pin.cpp
--- a/Engine/Plugins/Experimental/TextureGraph/Source/TextureGraph/Private/TG_Pin.cpp
+++ b/Engine/Plugins/Experimental/TextureGraph/Source/TextureGraph/Private/TG_Pin.cpp
@@ -269,8 +269,7 @@ FString	UTG_Pin::LogTooltip() const
 		ExprTooltip = *Prop->GetToolTipText().ToString();
 #endif 
 	
-	return FString::Printf(TEXT("Id%s : %s : %s - %s"),
-		*Id.ToString(),
+	return FString::Printf(TEXT("%s : %s - %s"),
 		*GetAliasName().ToString(),
 		*GetArgumentCPPTypeName().ToString(),
 		*ExprTooltip);
diff --git a/Engine/Plugins/Experimental/TextureGraph/Source/TextureGraph/Public/Expressions/Adjustment/TG_Expression_NormalFromHeightMap.h b/Engine/Plugins/Experimental/TextureGraph/Source/TextureGraph/Public/Expressions/Adjustment/TG_Expression_NormalFromHeightMap.h
--- a/Engine/Plugins/Experimental/TextureGraph/Source/TextureGraph/Public/Expressions/Adjustment/TG_Expression_NormalFromHeightMap.h
+++ b/Engine/Plugins/Experimental/TextureGraph/Source/TextureGraph/Public/Expressions/Adjustment/TG_Expression_NormalFromHeightMap.h
@@ -19,9 +19,11 @@ public:
 	virtual void						Evaluate(FTG_EvaluationContext* InContext) override;
 	virtual bool						Validate(MixUpdateCyclePtr Cycle) override;
 
+	// Normals from height are generated by taking nearby points from the texture having some offset. Using smaller offset values means we are getting samples closer to the surface from the texture. 
 	UPROPERTY(EditAnywhere, Category = NoCategory, meta = (TGType = "TG_Setting", UIMin = "0.001", ClampMin = "0.001", UIMax = "0.1", ClampMax = "0.1"))
 	float								Offset = 0.002;
 
+	// Strength of the normals
 	UPROPERTY(EditAnywhere, Category = NoCategory, meta = (TGType = "TG_Setting", UIMin = "0", ClampMin = "0", UIMax = "2", ClampMax = "2"))
 	float								Strength = 1;
 	
diff --git a/Engine/Plugins/Experimental/TextureGraph/Source/TextureGraph/Public/Expressions/Color/TG_Expression_HSV.h b/Engine/Plugins/Experimental/TextureGraph/Source/TextureGraph/Public/Expressions/Color/TG_Expression_HSV.h
--- a/Engine/Plugins/Experimental/TextureGraph/Source/TextureGraph/Public/Expressions/Color/TG_Expression_HSV.h
+++ b/Engine/Plugins/Experimental/TextureGraph/Source/TextureGraph/Public/Expressions/Color/TG_Expression_HSV.h
@@ -75,15 +75,16 @@ public:
 	UPROPERTY(meta = (TGType = "TG_Input", PinDisplayName = ""))
 	FTG_Texture							Input;
 
+	// Defines the basic color tone, such as red, green, or blue, without considering brightness or intensity. Adjusting the hue changes the overall color appearance while maintaining its saturation and brightness.
 	// The normalized hue. Please divide your [0, 359] hue values by 359 for this input
 	UPROPERTY(EditAnywhere, Category = NoCategory, meta = (TGType = "TG_Setting", UIMin = "0", ClampMin = "0", UIMax = "1", ClampMax = "1", PinDisplayName = "Hue (Normalized)", DisplayName = "Hue (Normalized)"))
 	float								Hue = 1.0f;
 	
-	// The saturation
+	// Controls the intensity of the color. Higher values represent more vivid colors, while lower values produce muted tones.
 	UPROPERTY(EditAnywhere, Category = NoCategory, meta = (TGType = "TG_Setting", UIMin = "0", ClampMin = "0", UIMax = "1", ClampMax = "1"))
 	float								Saturation = 1.0f;
 
-	// The value 
+	// Specifies the brightness or darkness of the color. Higher values correspond to brighter colors, while lower values result in darker shades.
 	UPROPERTY(EditAnywhere, Category = NoCategory, meta = (TGType = "TG_Setting", UIMin = "0", ClampMin = "0", UIMax = "1", ClampMax = "1"))
 	float								Value = 1.0f;
 
diff --git a/Engine/Plugins/Experimental/TextureGraph/Source/TextureGraph/Public/Expressions/Filter/TG_Expression_Blur.h b/Engine/Plugins/Experimental/TextureGraph/Source/TextureGraph/Public/Expressions/Filter/TG_Expression_Blur.h
--- a/Engine/Plugins/Experimental/TextureGraph/Source/TextureGraph/Public/Expressions/Filter/TG_Expression_Blur.h
+++ b/Engine/Plugins/Experimental/TextureGraph/Source/TextureGraph/Public/Expressions/Filter/TG_Expression_Blur.h
@@ -27,22 +27,28 @@ public:
 	// Used to implement EditCondition logic for both Node UI and Details View
 	virtual bool CanEditChange(const FProperty* InProperty) const override;
 #endif
-	
-	UPROPERTY(EditAnywhere, BlueprintReadWrite, Category = NoCategory, meta = (TGType = "TG_Setting", ClampMin = 0, ClampMax = 100, EditConditionHides))
+
+	// Controls how much an texture is blurred
+	UPROPERTY(EditAnywhere, BlueprintReadWrite, Category = NoCategory, meta = (TGType = "TG_Setting", ClampMin = 0, ClampMax = 20, EditConditionHides))
 	int32 Radius = 1;
-	
+
+	// Adjusting the angle changes the direction at which the blur effect is applied." 
 	UPROPERTY(EditAnywhere, BlueprintReadWrite, Category = NoCategory, meta = (TGType = "TG_Setting", ClampMin = -180, ClampMax = 180, EditConditionHides))
 	float Angle = 0.0f;
-	
+
+	// Determines how intense the blurring effect is applied. Higher values result in stronger blur, while lower values produce a milder effect.
 	UPROPERTY(EditAnywhere, BlueprintReadWrite, Category = NoCategory, meta = (TGType = "TG_Setting", ClampMin = 0, ClampMax = 1, EditConditionHides))
 	float Strength = 0.1f;
 
+	// Style of blurring applied to a texture. 
 	UPROPERTY(EditAnywhere, BlueprintReadWrite, Category = NoCategory, meta = (TGType = "TG_Setting", RegenPinsOnChange))
 	EBlurType BlurType = EBlurType::Gaussian;
 
+	// The output texture having blurred effect
 	UPROPERTY(EditAnywhere, Category = NoCategory, meta = (TGType = "TG_Output", PinDisplayName = ""))
 	FTG_Texture Output;
 
+	// The input texture to apply the blur effect
 	UPROPERTY(meta = (TGType = "TG_Input"))
 	FTG_Texture Input;
 
diff --git a/Engine/Plugins/Experimental/TextureGraph/Source/TextureGraph/Public/Expressions/Filter/TG_Expression_Brightness.h b/Engine/Plugins/Experimental/TextureGraph/Source/TextureGraph/Public/Expressions/Filter/TG_Expression_Brightness.h
--- a/Engine/Plugins/Experimental/TextureGraph/Source/TextureGraph/Public/Expressions/Filter/TG_Expression_Brightness.h
+++ b/Engine/Plugins/Experimental/TextureGraph/Source/TextureGraph/Public/Expressions/Filter/TG_Expression_Brightness.h
@@ -22,11 +22,11 @@ public:
 	TG_DECLARE_EXPRESSION(TG_Category::Adjustment);
 	virtual void							Evaluate(FTG_EvaluationContext* InContext) override;
 
-	// The desired brightness of the image
+	// Adjusts the overall lightness or darkness of an image. Increasing brightness makes it brighter, while decreasing brightness makes it darker.
 	UPROPERTY(EditAnywhere, Category = NoCategory, meta = (TGType = "TG_Setting", UIMin = "-1", ClampMin = "-1", UIMax = "1", ClampMax = "1"))
 		float								Brightness = 0;
 
-	// The desired contrast in the image
+	// Modifies the difference between light and dark areas within an image. Increasing contrast makes light areas lighter and dark areas darker, while decreasing contrast reduces this difference, resulting in a more uniform appearance.
 	UPROPERTY(EditAnywhere, Category = NoCategory, meta = (TGType = "TG_Setting", UIMin = "0", ClampMin = "0", UIMax = "10", ClampMax = "10"))
 		float								Contrast = 1;
 
diff --git a/Engine/Plugins/Experimental/TextureGraph/Source/TextureGraph/Public/Expressions/Input/TG_Expression_Material.h b/Engine/Plugins/Experimental/TextureGraph/Source/TextureGraph/Public/Expressions/Input/TG_Expression_Material.h
--- a/Engine/Plugins/Experimental/TextureGraph/Source/TextureGraph/Public/Expressions/Input/TG_Expression_Material.h
+++ b/Engine/Plugins/Experimental/TextureGraph/Source/TextureGraph/Public/Expressions/Input/TG_Expression_Material.h
@@ -28,8 +28,9 @@ public:
 	
 	void SetMaterial(UMaterialInterface* InMaterial);
 
+	// The attribute identifier among all the attributes of the material that is rendered in the output
 	UPROPERTY(EditAnywhere, BlueprintReadWrite, Category = NoCategory, meta = (TGType = "TG_Setting"))
-	EDrawMaterialAttributeTarget RenderedAttributeId = EDrawMaterialAttributeTarget::BaseColor;	// The attribute identifier among all the attributes of the material that is rendered in the output
+	EDrawMaterialAttributeTarget RenderedAttributeId = EDrawMaterialAttributeTarget::BaseColor;	
 	
 	virtual bool CanHandleAsset(UObject* Asset) override;
 	virtual void SetAsset(UObject* Asset) override;
diff --git a/Engine/Plugins/Experimental/TextureGraph/Source/TextureGraph/Public/Expressions/Input/TG_Expression_Vector.h b/Engine/Plugins/Experimental/TextureGraph/Source/TextureGraph/Public/Expressions/Input/TG_Expression_Vector.h
--- a/Engine/Plugins/Experimental/TextureGraph/Source/TextureGraph/Public/Expressions/Input/TG_Expression_Vector.h
+++ b/Engine/Plugins/Experimental/TextureGraph/Source/TextureGraph/Public/Expressions/Input/TG_Expression_Vector.h
@@ -18,7 +18,7 @@ public:
 
 	virtual void Evaluate(FTG_EvaluationContext* InContext) override;
 
-	// The color value
+	// The 4 dimensional vector X, Y, Z and W.
 	UPROPERTY(EditAnywhere, Category = NoCategory, meta = (TGType = "TG_InputParam", PinDisplayName = ""))
     FVector4f Vector = FVector4f(0);
     
diff --git a/Engine/Plugins/Experimental/TextureGraph/Source/TextureGraph/Public/Expressions/Maths/TG_Expression_Blend.h b/Engine/Plugins/Experimental/TextureGraph/Source/TextureGraph/Public/Expressions/Maths/TG_Expression_Blend.h
--- a/Engine/Plugins/Experimental/TextureGraph/Source/TextureGraph/Public/Expressions/Maths/TG_Expression_Blend.h
+++ b/Engine/Plugins/Experimental/TextureGraph/Source/TextureGraph/Public/Expressions/Maths/TG_Expression_Blend.h
@@ -18,6 +18,7 @@ public:
 	TG_DECLARE_EXPRESSION(TG_Category::Maths)
 	virtual void Evaluate(FTG_EvaluationContext* InContext) override;
 
+	// Blend mode determine how inputs texture mix together.
 	UPROPERTY(EditAnywhere, BlueprintReadWrite, Category = NoCategory, meta = (TGType = "TG_Setting"))
 	TEnumAsByte<EBlendModes> BlendMode = EBlendModes::Normal;
 
diff --git a/Engine/Plugins/Experimental/TextureGraph/Source/TextureGraph/Public/Expressions/Procedural/TG_Expression_Noise.h b/Engine/Plugins/Experimental/TextureGraph/Source/TextureGraph/Public/Expressions/Procedural/TG_Expression_Noise.h
--- a/Engine/Plugins/Experimental/TextureGraph/Source/TextureGraph/Public/Expressions/Procedural/TG_Expression_Noise.h
+++ b/Engine/Plugins/Experimental/TextureGraph/Source/TextureGraph/Public/Expressions/Procedural/TG_Expression_Noise.h
@@ -35,15 +35,15 @@ public:
 	UPROPERTY(EditAnywhere, BlueprintReadWrite, Category = NoCategory, meta = (TGType = "TG_Setting"))
 	TEnumAsByte<ENoiseType>				NoiseType = NOISETYPE_Simplex;
 
-	// The seed value to use
+	// A value used to initialize or "seed" a random number generator when generating procedural noise. Changing the noise seed produces different patterns or variations in the generated noise.
 	UPROPERTY(EditAnywhere, Category = NoCategory, meta = (TGType = "TG_Setting", UIMin = "-100", ClampMin = "-100", UIMax = "100", ClampMax = "100"))
 	int32								Seed = 75;
 
-	// The amplitude of the noise signal
+	// The magnitude or strength of the variations introduced by procedural noise. Adjusting this parameter can control the intensity or impact of procedural noise on textures or effects.
 	UPROPERTY(EditAnywhere, Category = NoCategory, meta = (TGType = "TG_Setting", UIMin = "0", ClampMin = "0", UIMax = "10", ClampMax = "10"))
 	float								Amplitude = 1;
 
-	// The frequency of the noise signal
+	// Determines how quickly or slowly the patterns change within the generated noise
 	UPROPERTY(EditAnywhere, Category = NoCategory, meta = (TGType = "TG_Setting", UIMin = "0", ClampMin = "0", UIMax = "20", ClampMax = "20"))
 	float								Frequency = 5;
 
diff --git a/Engine/Plugins/Experimental/TextureGraph/Source/TextureGraph/Public/Expressions/Procedural/TG_Expression_Shape.h b/Engine/Plugins/Experimental/TextureGraph/Source/TextureGraph/Public/Expressions/Procedural/TG_Expression_Shape.h
--- a/Engine/Plugins/Experimental/TextureGraph/Source/TextureGraph/Public/Expressions/Procedural/TG_Expression_Shape.h
+++ b/Engine/Plugins/Experimental/TextureGraph/Source/TextureGraph/Public/Expressions/Procedural/TG_Expression_Shape.h
@@ -39,35 +39,42 @@ public:
 	virtual bool						CanEditChange(const FProperty* InProperty) const override;
 #endif
 	
-	// The type of shape
+	// Select the geometric primitive shape to generate
 	UPROPERTY(EditAnywhere, BlueprintReadWrite, Category = NoCategory, meta = (TGType = "TG_Setting", RegenPinsOnChange))
 	EShapeType							ShapeType = EShapeType::Circle;
 
+	// The number of sides of the Polygon (only relevant when using the <Polygon> type)
 	UPROPERTY(EditAnywhere, BlueprintReadWrite, Category = NoCategory, meta = (TGType = "TG_Setting", ClampMin = 7, ClampMax = 32, EditConditionHides))
 	int32								PolygonNumSides = 8;
 
+	// Width representing the relative size of the shape in the generated texture. For Circle/Triangle/Pentagon hexagon and Polygon it is the radius of the shape relative to the uv domain of the generated texture
 	UPROPERTY(EditAnywhere, BlueprintReadWrite, Category = NoCategory, meta = (TGType = "TG_Setting", ClampMin = 0, ClampMax =1))
 	float 								Width = 1.0f;
 
+	// Height representing the relative size of the shape in the generate texture. Height is used in Rectangle and ellipse.
 	UPROPERTY(EditAnywhere, BlueprintReadWrite, Category = NoCategory, meta = (TGType = "TG_Setting", ClampMin = 0, ClampMax = 1, EditConditionHides))
 	float 								Height = 1.0f;
 
+	// Orientation is the 2D rotation of the shape around the center
 	UPROPERTY(EditAnywhere, BlueprintReadWrite, Category = NoCategory, meta = (TGType = "TG_Setting", UIMin = "-180", ClampMin = "-180", UIMax = "180", ClampMax = "180", Delta = "1", Units = "Degrees", EditConditionHides))
 	float								Orientation = 0.0f;
 
+	// Rounding ts about the 2D footprint of the shape generated, 0 means the perfect shape, 1 means the most rounded version of the shape
 	UPROPERTY(EditAnywhere, BlueprintReadWrite, Category = NoCategory, meta = (TGType = "TG_Setting", ClampMin = 0, ClampMax = 1, EditConditionHides))
 	float								Rounding = 0.0f;
 
+	// Bevel width is the distance from the outline of the shape on which the bevelling function is applied to the generated mask
 	UPROPERTY(EditAnywhere, BlueprintReadWrite, Category = NoCategory, meta = (TGType = "TG_Setting", ClampMin = 0, ClampMax = 1))
 	float								BevelWidth = 0.0f;
 
+	// Bevel curve is controlling the shape of bevelling varying from concave to linear to convex
 	UPROPERTY(EditAnywhere, BlueprintReadWrite, Category = NoCategory, meta = (TGType = "TG_Setting", ClampMin = -1, ClampMax = 1))
 	float								BevelCurve = 0.0f;
 
+	// Signed distance function the mathematical function defining the distance to the edge of the shape
 	UPROPERTY(EditAnywhere, BlueprintReadWrite, Category = NoCategory, meta = (TGType = "TG_Setting", ClampMin = 0, ClampMax = 1))
 	float								ShowSDF = 0.0f;
-
-
+	
 	// The output generated by the pattern
 	UPROPERTY(EditAnywhere, Category = NoCategory, meta = (TGType = "TG_Output", PinDisplayName = ""))
 	FTG_Texture							Output;
diff --git a/Engine/Plugins/Experimental/TextureGraph/Source/TextureGraph/Public/Expressions/Utilities/TG_Expression_MaterialID.h b/Engine/Plugins/Experimental/TextureGraph/Source/TextureGraph/Public/Expressions/Utilities/TG_Expression_MaterialID.h
--- a/Engine/Plugins/Experimental/TextureGraph/Source/TextureGraph/Public/Expressions/Utilities/TG_Expression_MaterialID.h
+++ b/Engine/Plugins/Experimental/TextureGraph/Source/TextureGraph/Public/Expressions/Utilities/TG_Expression_MaterialID.h
@@ -21,6 +21,7 @@ public:
 	UPROPERTY(meta = (TGType = "TG_Input"))
 	FTG_Texture MaterialIDMap;
 
+	// List of available colors extracted from the input texture. Selected colors will be shown as white in the output texture. 
 	UPROPERTY(EditAnywhere, EditFixedSize, Category = NoCategory, meta = (TGType = "TG_InputParam", NoResetToDefault))
 	TArray<FMaterialIDMaskInfo> MaterialIDMaskInfos;
 
diff --git a/Engine/Plugins/Experimental/TextureGraph/Source/TextureGraph/Public/TG_OutputSettings.h b/Engine/Plugins/Experimental/TextureGraph/Source/TextureGraph/Public/TG_OutputSettings.h
--- a/Engine/Plugins/Experimental/TextureGraph/Source/TextureGraph/Public/TG_OutputSettings.h
+++ b/Engine/Plugins/Experimental/TextureGraph/Source/TextureGraph/Public/TG_OutputSettings.h
@@ -15,34 +15,42 @@ struct TEXTUREGRAPH_API FTG_OutputSettings
 {
 	GENERATED_USTRUCT_BODY()
 
-public:
+	// Export name of the textured asset.
 	UPROPERTY(BlueprintReadWrite,EditAnywhere, Category = "Basic", DisplayName = "File Name", Meta = (NoResetToDefault))
 		FName BaseName;
 
 	UPROPERTY(EditAnywhere, Category = "Basic", Meta = (NoResetToDefault))
 		FName OutputName;
 
+	// Export path for the textured asset.
 	UPROPERTY(BlueprintReadWrite,EditAnywhere, Category = "Basic", DisplayName = "Path", Meta = (NoResetToDefault))
 		FName FolderPath;
 
+	// Width of the texture in pixels. Auto means system will detect automatically based on other images
 	UPROPERTY(BlueprintReadWrite,EditAnywhere, Category = "Advanced", Meta = (NoResetToDefault))
 		EResolution Width = EResolution::Auto;
 
+	// Height of the texture in pixels. Auto means system will detect automatically based on other images
 	UPROPERTY(BlueprintReadWrite,EditAnywhere, Category = "Advanced", Meta = (NoResetToDefault))
 		EResolution Height = EResolution::Auto;
 
+	// List of available texture formats. Auto means system will detect automatically based on the input
 	UPROPERTY(BlueprintReadWrite,EditAnywhere, Category = "Advanced", DisplayName = "Texture Format", Meta = (NoResetToDefault))
 		ETG_TextureFormat TextureFormat = ETG_TextureFormat::BGRA8;
 
+	// List of available texture presets available for export. 
 	UPROPERTY(BlueprintReadWrite, EditAnywhere, Category = "Advanced", DisplayName = "Texture Type", Meta = (NoResetToDefault))
 		ETG_TexturePresetType TexturePresetType = ETG_TexturePresetType::None;
 
+	// The Level of detail group of the texture
 	UPROPERTY(BlueprintReadWrite, EditAnywhere, Category = "Advanced", DisplayName = "LOD Texture Group", Meta = (NoResetToDefault, EditCondition = "TexturePresetType == ETG_TexturePresetType::None"))
 		TEnumAsByte<enum TextureGroup> LODGroup = TextureGroup::TEXTUREGROUP_World;
 
+	// Compression methods available for exporting textured asset.
 	UPROPERTY(BlueprintReadWrite, EditAnywhere, Category = "Advanced", DisplayName = "Compression", Meta = (NoResetToDefault, EditCondition = "TexturePresetType == ETG_TexturePresetType::None") )
 		TEnumAsByte <enum TextureCompressionSettings> Compression = TextureCompressionSettings::TC_Default;
 
+	// Adjust the color space of exporting textured asset. Can be in Linear or Gamma color space. 
 	UPROPERTY(BlueprintReadWrite, EditAnywhere, Category = "Advanced", DisplayName = "sRGB", Meta = (NoResetToDefault, EditCondition = "TexturePresetType == ETG_TexturePresetType::None"))
 		bool bSRGB = false;
 
diff --git a/Engine/Plugins/Experimental/TextureGraph/Source/TextureGraph/Public/TG_Texture.h b/Engine/Plugins/Experimental/TextureGraph/Source/TextureGraph/Public/TG_Texture.h
--- a/Engine/Plugins/Experimental/TextureGraph/Source/TextureGraph/Public/TG_Texture.h
+++ b/Engine/Plugins/Experimental/TextureGraph/Source/TextureGraph/Public/TG_Texture.h
@@ -24,11 +24,13 @@ struct TEXTUREGRAPH_API FTG_TextureDescriptor
 	}
 
 public:
+	// Width of the texture in pixels. Auto means system will detect automatically based on other images
 	UPROPERTY(EditAnywhere, Category = "TextureDescriptor")
-		EResolution Width = EResolution::Auto;		/// Width of the texture in pixels. Auto means system will detect automatically based on other images
+		EResolution Width = EResolution::Auto;		
 
+	// Height of the texture in pixels. Auto means system will detect automatically based on other images
 	UPROPERTY(EditAnywhere, Category = "TextureDescriptor")
-		EResolution Height = EResolution::Auto;	/// Height of the texture in pixels
+		EResolution Height = EResolution::Auto;	
 
 #if WITH_EDITORONLY_DATA
 PRAGMA_DISABLE_DEPRECATION_WARNINGS
diff --git a/Engine/Plugins/Experimental/TextureGraph/Source/TextureGraphEditor/Private/Widgets/STG_RGBAButtons.cpp b/Engine/Plugins/Experimental/TextureGraph/Source/TextureGraphEditor/Private/Widgets/STG_RGBAButtons.cpp
--- a/Engine/Plugins/Experimental/TextureGraph/Source/TextureGraphEditor/Private/Widgets/STG_RGBAButtons.cpp
+++ b/Engine/Plugins/Experimental/TextureGraph/Source/TextureGraphEditor/Private/Widgets/STG_RGBAButtons.cpp
@@ -25,14 +25,14 @@ TSharedRef<SWidget> STG_RGBAButtons::MakeChannelControlWidget()
 		.Padding(1.0f)
 		.AutoWidth()
 		[
-			CreateChannelWidget(ETSChannelButton::Red, "R")
+			CreateChannelWidget(ETSChannelButton::Red, "R", LOCTEXT("RBAButton_R_Channel_ToolTip", "Toggle the preview of R Channel of the texture"))
 		]
 	+ SHorizontalBox::Slot()
 		.VAlign(VAlign_Center)
 		.Padding(1.0f)
 		.AutoWidth()
 		[
-			CreateChannelWidget(ETSChannelButton::Green, "G")
+			CreateChannelWidget(ETSChannelButton::Green, "G", LOCTEXT("RBAButton_G_Channel_ToolTip", "Toggle the preview of G Channel of the texture"))
 		]
 
 	+ SHorizontalBox::Slot()
@@ -40,21 +40,21 @@ TSharedRef<SWidget> STG_RGBAButtons::MakeChannelControlWidget()
 		.Padding(1.0f)
 		.AutoWidth()
 		[
-			CreateChannelWidget(ETSChannelButton::Blue, "B")
+			CreateChannelWidget(ETSChannelButton::Blue, "B", LOCTEXT("RBAButton_B_Channel_ToolTip", "Toggle the preview of B Channel of the texture"))
 		]
 	+ SHorizontalBox::Slot()
 		.VAlign(VAlign_Center)
 		.Padding(1.0f)
 		.AutoWidth()
 		[
-			CreateChannelWidget(ETSChannelButton::Alpha, "A")
+			CreateChannelWidget(ETSChannelButton::Alpha, "A", LOCTEXT("RBAButton_A_Channel_ToolTip", "Toggle the preview of A Channel of the texture"))
 		]
 	;
 
 	return ChannelControl;
 }
 
-TSharedRef<SWidget> STG_RGBAButtons::CreateChannelWidget(ETSChannelButton Type, FString Name)
+TSharedRef<SWidget> STG_RGBAButtons::CreateChannelWidget(ETSChannelButton Type, FString Name, FText TooltipText)
 {
 	auto OnChannelCheckStateChanged = [this](ECheckBoxState NewState, ETSChannelButton Button)
 	{
@@ -84,7 +84,8 @@ TSharedRef<SWidget> STG_RGBAButtons::CreateChannelWidget(ETSChannelButton Type,
 			.OverflowPolicy(ETextOverflowPolicy::Ellipsis)
 			.Text(FText::FromString(Name))
 			]
-		];
+		]
+	.ToolTipText(TooltipText);
 }
 
 void STG_RGBAButtons::OnChannelButtonCheckStateChanged(ETSChannelButton Button)
diff --git a/Engine/Plugins/Experimental/TextureGraph/Source/TextureGraphEditor/Private/Widgets/STG_RGBAButtons.h b/Engine/Plugins/Experimental/TextureGraph/Source/TextureGraphEditor/Private/Widgets/STG_RGBAButtons.h
--- a/Engine/Plugins/Experimental/TextureGraph/Source/TextureGraphEditor/Private/Widgets/STG_RGBAButtons.h
+++ b/Engine/Plugins/Experimental/TextureGraph/Source/TextureGraphEditor/Private/Widgets/STG_RGBAButtons.h
@@ -40,7 +40,7 @@ private:
 	FSlateRoundedBoxBrush* CheckedBrush;
 
 	TSharedRef<SWidget> MakeChannelControlWidget();
-	TSharedRef<SWidget> CreateChannelWidget(ETSChannelButton Type, FString Name);
+	TSharedRef<SWidget> CreateChannelWidget(ETSChannelButton Type, FString Name, FText ToolTipText);
 	void OnChannelButtonCheckStateChanged(ETSChannelButton Button);
 	ECheckBoxState OnGetChannelButtonCheckState(ETSChannelButton Button) const;
 };
