Histoogram crash on tweaking fix.

diff --git a/Engine/Plugins/Experimental/TextureGraph/Source/TextureGraphEngine/FxMat/FxMaterial.h b/Engine/Plugins/Experimental/TextureGraph/Source/TextureGraphEngine/FxMat/FxMaterial.h
--- a/Engine/Plugins/Experimental/TextureGraph/Source/TextureGraphEngine/FxMat/FxMaterial.h
+++ b/Engine/Plugins/Experimental/TextureGraph/Source/TextureGraphEngine/FxMat/FxMaterial.h
@@ -283,18 +283,18 @@ protected:
 																			/// the Blit function will create a one at the time of blitting
 
 public:
-	explicit FxMaterial_Compute(FString outputId, const CmpSHPermutationDomain* permDomain = nullptr, 
-	                            int numThreadsX = GDefaultNumThreadsXY, int numThreadsY = GDefaultNumThreadsXY,int numThreadsZ = 1, 
-	                            FUnorderedAccessViewRHIRef unorderedAccessView = nullptr) 
+	explicit FxMaterial_Compute(FString InOutputId, const CmpSHPermutationDomain* InPermDomain = nullptr, 
+	                            int InNumThreadsX = GDefaultNumThreadsXY, int InNumThreadsY = GDefaultNumThreadsXY, 
+								int InNumThreadsZ= 1, FUnorderedAccessViewRHIRef InUnorderedAccessView = nullptr) 
 		: FxMaterial()
-		, OutputId(outputId)
-		, NumThreadsX(numThreadsX)
-		, NumThreadsY(numThreadsY)
-		, NumThreadsZ(numThreadsZ)
-		, UnorderedAccessView(unorderedAccessView)
+		, OutputId(InOutputId)
+		, NumThreadsX(InNumThreadsX)
+		, NumThreadsY(InNumThreadsY)
+		, NumThreadsZ(InNumThreadsZ)
+		, UnorderedAccessView(InUnorderedAccessView)
 	{
-		if (permDomain)
-			PermutationDomain = *permDomain;
+		if (InPermDomain)
+			PermutationDomain = *InPermDomain;
 
 		TShaderRef<CmpSH_Type> CSH = TShaderMapRef<CmpSH_Type>(GetGlobalShaderMap(GMaxRHIFeatureLevel), PermutationDomain);
 		FString CSHHashStr = CSH->GetHash().ToString();
diff --git a/Engine/Plugins/Experimental/TextureGraph/Source/TextureGraphEngine/Transform/Utility/T_TextureHistogram.cpp b/Engine/Plugins/Experimental/TextureGraph/Source/TextureGraphEngine/Transform/Utility/T_TextureHistogram.cpp
--- a/Engine/Plugins/Experimental/TextureGraph/Source/TextureGraphEngine/Transform/Utility/T_TextureHistogram.cpp
+++ b/Engine/Plugins/Experimental/TextureGraph/Source/TextureGraphEngine/Transform/Utility/T_TextureHistogram.cpp
@@ -47,8 +47,7 @@ TiledBlobPtr T_TextureHistogram::CreateJobAndResult(JobUPtr& OutJob, MixUpdateCy
 	FIntVector4 SrcDimensions(SourceTex->GetWidth(), SourceTex->GetHeight(), 1, 1);
 
 	FString Name = FString::Printf(TEXT("[%s].[%d] Histogram"), *SourceTex->DisplayName(), TargetId);
-	RenderMaterial_FXPtr Transform = T_TextureHistogram::CreateMaterial_Histogram(
-		TEXT("T_Histogram"), TEXT("Result"), PermutationVector, SrcDimensions.X, SrcDimensions.Y, 1);
+	RenderMaterial_FXPtr Transform = T_TextureHistogram::CreateMaterial_Histogram(TEXT("T_Histogram"), TEXT("Result"), PermutationVector, 32, 32, 1);
 
 	OutJob = std::make_unique<Job>(Cycle->GetMix(), TargetId, std::static_pointer_cast<BlobTransform>(Transform));
 	OutJob->AddArg(ARG_BLOB(SourceTex, "SourceTiles"));
