Memory leak fix.

diff --git a/Engine/Plugins/Experimental/TextureGraph/Source/TextureGraphEngine/2D/Tex.cpp b/Engine/Plugins/Experimental/TextureGraph/Source/TextureGraphEngine/2D/Tex.cpp
--- a/Engine/Plugins/Experimental/TextureGraph/Source/TextureGraphEngine/2D/Tex.cpp
+++ b/Engine/Plugins/Experimental/TextureGraph/Source/TextureGraphEngine/2D/Tex.cpp
@@ -219,7 +219,7 @@ Tex::Tex()
 
 Tex::~Tex()
 {
-	//Free();
+	Free();
 }
 
 void Tex::FreeTexture(UTexture2D** Texture)
@@ -235,6 +235,11 @@ void Tex::FreeGenericTexture(UTexture** Texture)
 	if (Texture && *Texture)
 	{
 		UE_LOG(LogTexture, Verbose, TEXT("Deleting Texture: %s [Ptr = 0x%llx]"), *((*Texture)->GetName()), (*Texture));
+		FTextureResource* TextureResource = (*Texture)->GetResource();
+
+		if (!TextureGraphEngine::IsDestroying() && TextureResource && TextureResource->IsInitialized())
+			(*Texture)->ReleaseResource();
+
 		(*Texture) = nullptr;
 	}
 }
@@ -244,12 +249,18 @@ void Tex::FreeRT(UTextureRenderTarget2D** RT)
 	if (RT && *RT)
 	{
 		UE_LOG(LogTexture, Verbose, TEXT("Deleting render target: %s [Ptr = 0x%llx]"), *((*RT)->GetName()), (*RT));
+		FTextureResource* RTResource = (*RT)->GetResource();
+
+		if (!TextureGraphEngine::IsDestroying() && RTResource && RTResource->IsInitialized())
+			(*RT)->ReleaseResource();
+
 		(*RT) = nullptr;
 	}
 }
 
 void Tex::Free()
 {
+	check(IsInGameThread());
 	/// IMPORTANT: Do not free the _image in this function
 	FreeTexture(ToRawPtr(MutableView(Texture)));
 	FreeRT(ToRawPtr(MutableView(RT)));
diff --git a/Engine/Plugins/Experimental/TextureGraph/Source/TextureGraphEngine/Device/FX/Device_FX.cpp b/Engine/Plugins/Experimental/TextureGraph/Source/TextureGraphEngine/Device/FX/Device_FX.cpp
--- a/Engine/Plugins/Experimental/TextureGraph/Source/TextureGraphEngine/Device/FX/Device_FX.cpp
+++ b/Engine/Plugins/Experimental/TextureGraph/Source/TextureGraphEngine/Device/FX/Device_FX.cpp
@@ -56,10 +56,10 @@ void Device_FX::AddReferencedObjects(FReferenceCollector& Collector)
 {
 	for (auto Iter : RTCache)
 	{
-		RTList* list = Iter.second;
-		if (list)
+		RTList* List = Iter.second;
+		if (List)
 		{
-			for (auto& RT : *list)
+			for (auto& RT : *List)
 			{
 				Collector.AddReferencedObject(RT);
 			}
@@ -68,10 +68,10 @@ void Device_FX::AddReferencedObjects(FReferenceCollector& Collector)
 
 	for (auto Iter : RTArrayCache)
 	{
-		RTList* list = Iter.second;
-		if (list)
+		RTList* List = Iter.second;
+		if (List)
 		{
-			for (auto& RT : *list)
+			for (auto& RT : *List)
 			{
 				Collector.AddReferencedObject(RT);
 			}
@@ -107,7 +107,7 @@ void Device_FX::FreeCacheInternal(RenderTargetCache& TargetRTCache)
 			delete List;
 		}
 	}
-	
+
 	TargetRTCache.clear();
 }
 
@@ -130,13 +130,13 @@ cti::continuable<int32> Device_FX::Use() const
 	}
 
 	return cti::make_continuable<int32>([this](auto&& Promise) mutable
-	{
-		ENQUEUE_RENDER_COMMAND(SceneDrawCompletion)([this, Promise = std::forward<decltype(Promise)>(Promise)](FRHICommandListImmediate& DevRHI) mutable
 		{
-			const_cast<Device_FX*>(this)->CallDeviceUpdate(); // OK as long as the type object isn't const
-			Promise.set_value(0);
+			ENQUEUE_RENDER_COMMAND(SceneDrawCompletion)([this, Promise = std::forward<decltype(Promise)>(Promise)](FRHICommandListImmediate& DevRHI) mutable
+				{
+					const_cast<Device_FX*>(this)->CallDeviceUpdate(); // OK as long as the type object isn't const
+					Promise.set_value(0);
+				});
 		});
-	});
 }
 
 void Device_FX::GetStatArray(TArray<FString>& ResourceArrayTiled,
@@ -145,30 +145,31 @@ void Device_FX::GetStatArray(TArray<FString>& ResourceArrayTiled,
 	TArray<FString>& TooltipListUnTiled,
 	FString& TotalStats)
 {
-	size_t rtMemUnused = 0;
-	size_t rtCountUnused = 0;
-	size_t rtMemUsed = 0;
+	size_t RTMemUnused = 0;
+	size_t RTCountUnused = 0;
+	size_t RTMemUsed = 0;
 
 	for (auto Iter = RTCache.begin(); Iter != RTCache.end(); Iter++)
 	{
-		RTList* rts = Iter->second;
-		if(rts)
-		for (UTextureRenderTarget2D* RT : *rts)
+		RTList* RenderTargets = Iter->second;
+		if (RenderTargets)
 		{
-			auto RTFormat = RT->RenderTargetFormat;
-			auto PixelFormat = TextureHelper::GetPixelFormatFromRenderTargetFormat(RTFormat);
+			for (UTextureRenderTarget2D* RT : *RenderTargets)
+			{
+				auto RTFormat = RT->RenderTargetFormat;
+				auto PixelFormat = TextureHelper::GetPixelFormatFromRenderTargetFormat(RTFormat);
 
-			
-			size_t size = RT->GetResourceSizeBytes(EResourceSizeMode::EstimatedTotal);
+				size_t Size = RT->GetResourceSizeBytes(EResourceSizeMode::EstimatedTotal);
 
-			FString info(FString::Printf(TEXT("[Unused] Name : %s "), *RT->GetName()));
-			FString tooltip(FString::Printf(TEXT("Ptr : %x , Hash : %llu, Memory Size : %0.2fKB"), RT, Iter->first, (float)size / (1024.0f)));
+				FString Info(FString::Printf(TEXT("[Unused] Name : %s "), *RT->GetName()));
+				FString Tooltip(FString::Printf(TEXT("Ptr : %x , Hash : %llu, Memory Size : %0.2fKB"), RT, Iter->first, (float)Size / (1024.0f)));
 
-			ResourceArrayUnTiled.Add(info);
-			TooltipListUnTiled.Add(tooltip);
+				ResourceArrayUnTiled.Add(Info);
+				TooltipListUnTiled.Add(Tooltip);
 
-			rtMemUnused += size;
-			rtCountUnused++;
+				RTMemUnused += Size;
+				RTCountUnused++;
+			}
 		}
 	}
 
@@ -176,23 +177,23 @@ void Device_FX::GetStatArray(TArray<FString>& ResourceArrayTiled,
 	{
 		auto RTFormat = RT->RenderTargetFormat;
 		auto PixelFormat = TextureHelper::GetPixelFormatFromRenderTargetFormat(RTFormat);
-		size_t size = RT->GetResourceSizeBytes(EResourceSizeMode::EstimatedTotal);
-		FString info(FString::Printf(TEXT("[Used] Name : %s"), *RT->GetName() ));
-		FString tooltip(FString::Printf(TEXT("Ptr : %x , Memory Size : %0.2fKB"), RT.Get(), (float)size / (1024.0f)));
-		
-		ResourceArrayUnTiled.Add(info);
-		TooltipListUnTiled.Add(tooltip);
-		
-		rtMemUsed += size;
+		size_t Size = RT->GetResourceSizeBytes(EResourceSizeMode::EstimatedTotal);
+		FString Info(FString::Printf(TEXT("[Used] Name : %s"), *RT->GetName()));
+		FString Tooltip(FString::Printf(TEXT("Ptr : %x , Memory Size : %0.2fKB"), RT.Get(), (float)Size / (1024.0f)));
+
+		ResourceArrayUnTiled.Add(Info);
+		TooltipListUnTiled.Add(Tooltip);
+
+		RTMemUsed += Size;
 	}
 
-	float usedRTMemory = (float)rtMemUsed / (1024.0f * 1024.0f);
-	float unusedRTMemory = (float)rtMemUnused / (1024.0f * 1024.0f);
+	float usedRTMemory = (float)RTMemUsed / (1024.0f * 1024.0f);
+	float unusedRTMemory = (float)RTMemUnused / (1024.0f * 1024.0f);
 
-	TotalStats = FString::Printf(TEXT("Total Resources [UNUSED] : %llu, Total Memory [UNUSED] : %0.2fMB, Total Resources [USED] : %llu, Total Memory [USED] : %0.2f MB"),rtCountUnused,unusedRTMemory,RTUsed.size(),usedRTMemory);
+	TotalStats = FString::Printf(TEXT("Total Resources [UNUSED] : %llu, Total Memory [UNUSED] : %0.2fMB, Total Resources [USED] : %llu, Total Memory [USED] : %0.2f MB"), RTCountUnused, unusedRTMemory, RTUsed.size(), usedRTMemory);
 }
 
-AsyncDeviceBufferRef Device_FX::CombineFromTiles(const CombineSplitArgs& CombineArgs) 
+AsyncDeviceBufferRef Device_FX::CombineFromTiles(const CombineSplitArgs& CombineArgs)
 {
 	check(IsInGameThread());
 
@@ -205,18 +206,18 @@ AsyncDeviceBufferRef Device_FX::CombineFromTiles(const CombineSplitArgs& Combine
 	check(Tiles.Rows() > 1 && Tiles.Cols() > 1);
 
 	/// check whether the buffers are compatible
-	bool IsCompatible = true;
+	bool bIsCompatible = true;
 	BufferDescriptor Desc = Tiles[0][0]->Descriptor();
 	CHashPtrVec TileHashes(Tiles.Rows() * Tiles.Cols());
 
-	for (int32 TileX = 0; TileX < Tiles.Rows() && IsCompatible; TileX++)
+	for (int32 TileX = 0; TileX < Tiles.Rows() && bIsCompatible; TileX++)
 	{
-		for (int32 TileY = 0; TileY < Tiles.Cols() && IsCompatible; TileY++)
+		for (int32 TileY = 0; TileY < Tiles.Cols() && bIsCompatible; TileY++)
 		{
 			DeviceBufferRef Tile = Tiles[TileX][TileY];
 			CHashPtr TileHash = Tile->Hash(false);
 			check(TileHash && TileHash->IsValid());
-			IsCompatible &= Tile->IsCompatible(this);
+			bIsCompatible &= Tile->IsCompatible(this);
 			TileHashes[TileY * Tiles.Cols() + TileX] = TileHash;
 		}
 	}
@@ -225,14 +226,14 @@ AsyncDeviceBufferRef Device_FX::CombineFromTiles(const CombineSplitArgs& Combine
 	if (Buffer)
 		Buffer->SetHash(BufferHash);
 
-	Desc.Width	= Desc.Width * Tiles.Cols();
+	Desc.Width = Desc.Width * Tiles.Cols();
 	Desc.Height = Desc.Height * Tiles.Rows();
 
 	/// Tiles must have been made compatible beforehand
-	check(IsCompatible);
+	check(bIsCompatible);
 
 	/// If any of the buffers is not compatible, then we fallback to the default implementation
-	if (!IsCompatible)
+	if (!bIsCompatible)
 		return Device::CombineFromTiles(CombineArgs);
 
 	if (!Buffer)
@@ -242,63 +243,63 @@ AsyncDeviceBufferRef Device_FX::CombineFromTiles(const CombineSplitArgs& Combine
 
 	AsyncBufferResultPtr Promise = cti::make_ready_continuable(std::make_shared<BufferResult>());
 
-	if(!IsArray)
+	if (!IsArray)
 		Promise = FXBuffer->AllocateRenderTarget();
 
 	return (std::move(Promise)).then([this, Tiles, Buffer, IsArray]()
-	{
-		check(IsInGameThread());
-		if (IsArray)
-			return FillTextureArray_Deferred(Buffer, Tiles);
+		{
+			check(IsInGameThread());
+			if (IsArray)
+				return FillTextureArray_Deferred(Buffer, Tiles);
 
-		return DrawTilesToBuffer_Deferred(Buffer, Tiles);
-	});
+			return DrawTilesToBuffer_Deferred(Buffer, Tiles);
+		});
 }
 
 AsyncDeviceBufferRef Device_FX::FillTextureArray_Deferred(DeviceBufferRef Buffer, const T_Tiles<DeviceBufferRef>& Tiles)
 {
 	DeviceNativeTaskPtr Task = DeviceNativeTask_Lambda::Create(this, (int32)E_Priority::kSystem, TEXT("FillTextureArray"), [this, Buffer, Tiles]() mutable -> int32
-	{
-		SCOPE_CYCLE_COUNTER(STAT_Device_FX_FillTextureArray);
-		check(IsInRenderingThread());
+		{
+			SCOPE_CYCLE_COUNTER(STAT_Device_FX_FillTextureArray);
+			check(IsInRenderingThread());
 
-		DeviceBuffer_FX* FXBuffer = static_cast<DeviceBuffer_FX*>(Buffer.get());
+			DeviceBuffer_FX* FXBuffer = static_cast<DeviceBuffer_FX*>(Buffer.get());
 
-		/// This texture must have been created beforehand AND it must be a render target
-		TexPtr DstTex = FXBuffer->GetTexture();
-		check(DstTex);
+			/// This texture must have been created beforehand AND it must be a render target
+			TexPtr DstTex = FXBuffer->GetTexture();
+			check(DstTex);
 
-		UTextureRenderTarget2D* RTDest = (UTextureRenderTarget2D*)DstTex->GetTexture();
-		FTexture2DRHIRef RTResDest = ((FTextureRenderTarget2DResource*)RTDest->GetResource())->GetTextureRHI();
+			UTextureRenderTarget2D* RTDest = (UTextureRenderTarget2D*)DstTex->GetTexture();
+			FTexture2DRHIRef RTResDest = ((FTextureRenderTarget2DResource*)RTDest->GetResource())->GetTextureRHI();
 
-		FRHICommandListImmediate& DevRHI = RHI();
+			FRHICommandListImmediate& DevRHI = RHI();
 
-		int32 index = 0;
-		for (int32 TileX = 0; TileX < Tiles.Rows(); TileX++)
-		{
-			for (int32 TileY = 0; TileY < Tiles.Cols(); TileY++, index++)
+			int32 Index = 0;
+			for (int32 TileX = 0; TileX < Tiles.Rows(); TileX++)
 			{
-				DeviceBuffer_FX* TileBuffer = static_cast<DeviceBuffer_FX*>(Tiles[TileX][TileY].get());
-				TexPtr TileTex = TileBuffer->GetTexture();
+				for (int32 TileY = 0; TileY < Tiles.Cols(); TileY++, Index++)
+				{
+					DeviceBuffer_FX* TileBuffer = static_cast<DeviceBuffer_FX*>(Tiles[TileX][TileY].get());
+					TexPtr TileTex = TileBuffer->GetTexture();
 
-				check(TileTex);
+					check(TileTex);
 
-				int32 TileWidth = TileTex->GetWidth();
-				int32 TileHeight = TileTex->GetHeight();
+					int32 TileWidth = TileTex->GetWidth();
+					int32 TileHeight = TileTex->GetHeight();
 
-				UTextureRenderTarget2D* TileRT = TileTex->GetRenderTarget();
-				FRHITexture2D* TileResource = TileTex->GetRHITexture();
+					UTextureRenderTarget2D* TileRT = TileTex->GetRenderTarget();
+					FRHITexture2D* TileResource = TileTex->GetRHITexture();
 
-				FRHICopyTextureInfo CopyInfo;
-				CopyInfo.Size = TileResource->GetSizeXYZ();
-				CopyInfo.DestSliceIndex = index;
+					FRHICopyTextureInfo CopyInfo;
+					CopyInfo.Size = TileResource->GetSizeXYZ();
+					CopyInfo.DestSliceIndex = Index;
 
-				DevRHI.CopyTexture(TileResource, RTResDest, CopyInfo);
+					DevRHI.CopyTexture(TileResource, RTResDest, CopyInfo);
+				}
 			}
-		}
 
-		return 0;
-	});
+			return 0;
+		});
 
 	return cti::make_ready_continuable(Buffer);
 }
@@ -306,88 +307,88 @@ AsyncDeviceBufferRef Device_FX::FillTextureArray_Deferred(DeviceBufferRef Buffer
 AsyncDeviceBufferRef Device_FX::DrawTilesToBuffer_Deferred(DeviceBufferRef Buffer, const T_Tiles<DeviceBufferRef>& Tiles)
 {
 	DeviceNativeTaskPtr Task = DeviceNativeTask_Lambda::Create(this, (int32)E_Priority::kSystem, TEXT("DrawTilesToBuffer"), [this, Buffer, Tiles]() mutable -> int32
-	{
-		SCOPE_CYCLE_COUNTER(STAT_Device_FX_DrawTilesToBuffer);
-		check(IsInRenderingThread());
+		{
+			SCOPE_CYCLE_COUNTER(STAT_Device_FX_DrawTilesToBuffer);
+			check(IsInRenderingThread());
 
-		DeviceBuffer_FX* FXBuffer = static_cast<DeviceBuffer_FX*>(Buffer.get());
+			DeviceBuffer_FX* FXBuffer = static_cast<DeviceBuffer_FX*>(Buffer.get());
 
-		/// This texture must have been created beforehand AND it must be a render target
-		TexPtr DstTex = FXBuffer->GetTexture();
-		check(DstTex);
+			/// This texture must have been created beforehand AND it must be a render target
+			TexPtr DstTex = FXBuffer->GetTexture();
+			check(DstTex);
 
-		UTextureRenderTarget2D* RTDest = DstTex->GetRenderTarget();
-		FTexture2DRHIRef RTResDest = ((FTextureRenderTarget2DResource*)RTDest->GetResource())->GetTextureRHI();
+			UTextureRenderTarget2D* RTDest = DstTex->GetRenderTarget();
+			FTexture2DRHIRef RTResDest = ((FTextureRenderTarget2DResource*)RTDest->GetResource())->GetTextureRHI();
 
-		FRHICommandListImmediate& DevRHI = RHI();
+			FRHICommandListImmediate& DevRHI = RHI();
 
-		const uint32 ViewportWidth = DstTex->GetWidth();
-		const uint32 ViewportHeight = DstTex->GetHeight();
-		const FIntPoint TargetSize(ViewportWidth, ViewportHeight);
+			const uint32 ViewportWidth = DstTex->GetWidth();
+			const uint32 ViewportHeight = DstTex->GetHeight();
+			const FIntPoint TargetSize(ViewportWidth, ViewportHeight);
 
-		FRHIRenderPassInfo RenderPassInfo(RTResDest, ERenderTargetActions::Load_Store);
-		DevRHI.BeginRenderPass(RenderPassInfo, TEXT("CopyTexture"));
-		{
-			FGraphicsPipelineStateInitializer GraphicsPSOInit;
-			DevRHI.ApplyCachedRenderTargets(GraphicsPSOInit);
+			FRHIRenderPassInfo RenderPassInfo(RTResDest, ERenderTargetActions::Load_Store);
+			DevRHI.BeginRenderPass(RenderPassInfo, TEXT("CopyTexture"));
+			{
+				FGraphicsPipelineStateInitializer GraphicsPSOInit;
+				DevRHI.ApplyCachedRenderTargets(GraphicsPSOInit);
 
-			const auto FeatureLevel = GMaxRHIFeatureLevel;
-			auto ShaderMap = GetGlobalShaderMap(FeatureLevel);
-			TShaderMapRef<FScreenVS> VertexShader(ShaderMap);
-			TShaderMapRef<FScreenPS> PixelShader(ShaderMap);
-			FxMaterial::InitPSO_Default(GraphicsPSOInit, VertexShader.GetVertexShader(), PixelShader.GetPixelShader());
+				const auto FeatureLevel = GMaxRHIFeatureLevel;
+				auto ShaderMap = GetGlobalShaderMap(FeatureLevel);
+				TShaderMapRef<FScreenVS> VertexShader(ShaderMap);
+				TShaderMapRef<FScreenPS> PixelShader(ShaderMap);
+				FxMaterial::InitPSO_Default(GraphicsPSOInit, VertexShader.GetVertexShader(), PixelShader.GetPixelShader());
 
-			GraphicsPSOInit.PrimitiveType = PT_TriangleList;
-			GraphicsPSOInit.BoundShaderState.VertexDeclarationRHI = GFilterVertexDeclaration.VertexDeclarationRHI;
-		
-			FRHISamplerState* PixelSampler = TStaticSamplerState<SF_Bilinear>::GetRHI();
+				GraphicsPSOInit.PrimitiveType = PT_TriangleList;
+				GraphicsPSOInit.BoundShaderState.VertexDeclarationRHI = GFilterVertexDeclaration.VertexDeclarationRHI;
 
-			SetGraphicsPipelineState(DevRHI, GraphicsPSOInit, 0);
+				FRHISamplerState* PixelSampler = TStaticSamplerState<SF_Bilinear>::GetRHI();
 
-			for (int32 TileX = 0; TileX < Tiles.Rows(); TileX++)
-			{
-				for (int32 TileY = 0; TileY < Tiles.Cols(); TileY++)
+				SetGraphicsPipelineState(DevRHI, GraphicsPSOInit, 0);
+
+				for (int32 TileX = 0; TileX < Tiles.Rows(); TileX++)
 				{
-					DeviceBuffer_FX* TileBuffer = static_cast<DeviceBuffer_FX*>(Tiles[TileX][TileY].get());
-					TexPtr TileTex = TileBuffer->GetTexture();
+					for (int32 TileY = 0; TileY < Tiles.Cols(); TileY++)
+					{
+						DeviceBuffer_FX* TileBuffer = static_cast<DeviceBuffer_FX*>(Tiles[TileX][TileY].get());
+						TexPtr TileTex = TileBuffer->GetTexture();
 
-					UE_LOG(LogDevice, Verbose, TEXT("Combine tiles: %s [%d, %d] => %llu [%s]"), *FXBuffer->GetName(), TileX, TileY, TileBuffer->Hash(false)->Value(), *TileBuffer->GetName());
+						UE_LOG(LogDevice, Verbose, TEXT("Combine tiles: %s [%d, %d] => %llu [%s]"), *FXBuffer->GetName(), TileX, TileY, TileBuffer->Hash(false)->Value(), *TileBuffer->GetName());
 
-					check(TileTex);
+						check(TileTex);
 
-					UTextureRenderTarget2D* TileRT = TileTex->GetRenderTarget();
-					FRHITexture2D* TileResource = TileTex->GetRHITexture();
+						UTextureRenderTarget2D* TileRT = TileTex->GetRenderTarget();
+						FRHITexture2D* TileResource = TileTex->GetRHITexture();
+
+						const float SrcTextureWidth = TileTex->GetWidth();
+						const float SrcTextureHeight = TileTex->GetHeight();
 
-					const float SrcTextureWidth = TileTex->GetWidth();
-					const float SrcTextureHeight = TileTex->GetHeight();
+						DevRHI.Transition(FRHITransitionInfo(TileResource, ERHIAccess::Unknown, ERHIAccess::SRVGraphics));
 
-					DevRHI.Transition(FRHITransitionInfo(TileResource, ERHIAccess::Unknown, ERHIAccess::SRVGraphics));
+						int32 X1 = TileX * SrcTextureWidth;
+						int32 Y1 = TileY * SrcTextureHeight;
 
-					int32 X1 = TileX * SrcTextureWidth;
-					int32 Y1 = TileY * SrcTextureHeight;
+						DevRHI.SetViewport(X1, Y1, 0, X1 + ViewportWidth, Y1 + ViewportHeight, 0);
 
-					DevRHI.SetViewport(X1, Y1, 0, X1 + ViewportWidth, Y1 + ViewportHeight, 0);
+						SetShaderParametersLegacyPS(DevRHI, PixelShader, PixelSampler, TileResource);
 
-					SetShaderParametersLegacyPS(DevRHI, PixelShader, PixelSampler, TileResource);
+						RendererModule->DrawRectangle(
+							DevRHI,
+							0, 0,
+							SrcTextureWidth, SrcTextureHeight,
+							0.0f, 0.0f,
+							1.0f, 1.0f,
+							TargetSize,
+							FIntPoint(1, 1),
+							VertexShader,
+							EDRF_Default);
 
-					RendererModule->DrawRectangle(
-						DevRHI,
-						0, 0,
-						SrcTextureWidth, SrcTextureHeight,
-						0.0f, 0.0f,
-						1.0f, 1.0f,
-						TargetSize,
-						FIntPoint(1, 1),
-						VertexShader,
-						EDRF_Default);
-					
+					}
 				}
 			}
-		}
-		DevRHI.EndRenderPass();
+			DevRHI.EndRenderPass();
 
-		return 0;
-	});
+			return 0;
+		});
 
 	return cti::make_ready_continuable(Buffer);
 }
@@ -399,37 +400,37 @@ AsyncDeviceBufferRef Device_FX::SplitToTiles_Internal(const CombineSplitArgs& Sp
 	auto IsArray = SplitArgs.IsArray;
 
 	DeviceNativeTask_Lambda::Create(this, (int32)E_Priority::kSystem, TEXT("SplitToTiles_Internal"), [this, Buffer, Tiles, IsArray]() -> int32
-	{
-		DeviceBuffer_FX* FXBuffer = static_cast<DeviceBuffer_FX*>(Buffer.get());
+		{
+			DeviceBuffer_FX* FXBuffer = static_cast<DeviceBuffer_FX*>(Buffer.get());
 
-		/// This texture must have been created beforehand AND it must be a render target
-		TexPtr SrcTex = FXBuffer->GetTexture();
-		FTexture2DRHIRef sourceResource = ((UTextureRenderTarget2D*)SrcTex->GetTexture())->GetRenderTargetResource()->GetRenderTargetTexture();
+			/// This texture must have been created beforehand AND it must be a render target
+			TexPtr SrcTex = FXBuffer->GetTexture();
+			FTexture2DRHIRef sourceResource = ((UTextureRenderTarget2D*)SrcTex->GetTexture())->GetRenderTargetResource()->GetRenderTargetTexture();
 
-		for (int32 TileX = 0; TileX < Tiles.Rows(); TileX++)
-		{
-			for (int32 TileY = 0; TileY < Tiles.Cols(); TileY++)
+			for (int32 TileX = 0; TileX < Tiles.Rows(); TileX++)
 			{
-				/// Using RHI to split Tiles to Buffer
-				DeviceBuffer_FX* TileBuffer = static_cast<DeviceBuffer_FX*>(Tiles[TileX][TileY].get());
-				TexPtr TileTex = TileBuffer->GetTexture();
-				check(TileTex);
-				check(TileTex->GetRenderTarget());
-
-				FTexture2DRHIRef texture2Dres = ((FTextureRenderTarget2DResource*)TileTex->GetRenderTarget()->GetResource())->GetTextureRHI();
-				
-				FRHICopyTextureInfo CopyInfo;
-				CopyInfo.SourcePosition = FIntVector(TileX * TileTex->GetWidth(), TileY * TileTex->GetHeight(), 0);
-				CopyInfo.DestPosition = FIntVector(0, 0, 0);
-				CopyInfo.Size = FIntVector(TileTex->GetWidth(), TileTex->GetHeight(), 0);
-
-				FRHICommandListImmediate& DevRHI = GRHICommandList.GetImmediateCommandList();
-				DevRHI.CopyTexture(sourceResource, texture2Dres, CopyInfo);
+				for (int32 TileY = 0; TileY < Tiles.Cols(); TileY++)
+				{
+					/// Using RHI to split Tiles to Buffer
+					DeviceBuffer_FX* TileBuffer = static_cast<DeviceBuffer_FX*>(Tiles[TileX][TileY].get());
+					TexPtr TileTex = TileBuffer->GetTexture();
+					check(TileTex);
+					check(TileTex->GetRenderTarget());
+
+					FTexture2DRHIRef texture2Dres = ((FTextureRenderTarget2DResource*)TileTex->GetRenderTarget()->GetResource())->GetTextureRHI();
+
+					FRHICopyTextureInfo CopyInfo;
+					CopyInfo.SourcePosition = FIntVector(TileX * TileTex->GetWidth(), TileY * TileTex->GetHeight(), 0);
+					CopyInfo.DestPosition = FIntVector(0, 0, 0);
+					CopyInfo.Size = FIntVector(TileTex->GetWidth(), TileTex->GetHeight(), 0);
+
+					FRHICommandListImmediate& DevRHI = GRHICommandList.GetImmediateCommandList();
+					DevRHI.CopyTexture(sourceResource, texture2Dres, CopyInfo);
+				}
 			}
-		}
 
-		return 0;
-	});
+			return 0;
+		});
 
 	return cti::make_ready_continuable(Buffer);
 }
@@ -437,10 +438,10 @@ AsyncDeviceBufferRef Device_FX::SplitToTiles_Internal(const CombineSplitArgs& Sp
 AsyncDeviceBufferRef Device_FX::SplitToTiles(const CombineSplitArgs& SplitArgs)
 {
 	SCOPE_CYCLE_COUNTER(STAT_Device_FX_SplitToTiles)
-	//check(IsInRenderingThread());
+		//check(IsInRenderingThread());
 
 #define NATIVE_SPLIT_TILES
-	const T_Tiles<DeviceBufferRef>& Tiles = SplitArgs.Tiles;
+		const T_Tiles<DeviceBufferRef>& Tiles = SplitArgs.Tiles;
 	auto Buffer = SplitArgs.Buffer;
 	auto IsArray = SplitArgs.IsArray;
 
@@ -481,27 +482,27 @@ AsyncDeviceBufferRef Device_FX::SplitToTiles(const CombineSplitArgs& SplitArgs)
 
 	/// Otherwise, we wait ...
 	return cti::when_all(promises.begin(), promises.end()).then([this, SplitArgs]()
-	{
-		return cti::make_continuable<DeviceBufferRef>([this, SplitArgs](auto&& Promise)
 		{
-			//Util::OnRenderingThread([this, Buffer, &Tiles, Promise = std::forward<decltype(Promise)>(Promise)](FRHICommandListImmediate& DevRHI) mutable
-			//{
-			/// And then, once we're ready, we split
-			SplitToTiles_Internal(SplitArgs).then([this, Promise = std::forward<decltype(Promise)>(Promise)](DeviceBufferRef Result) mutable
-			{
-				Promise.set_value(Result);
-			});
-			//});
+			return cti::make_continuable<DeviceBufferRef>([this, SplitArgs](auto&& Promise)
+				{
+					//Util::OnRenderingThread([this, Buffer, &Tiles, Promise = std::forward<decltype(Promise)>(Promise)](FRHICommandListImmediate& DevRHI) mutable
+					//{
+					/// And then, once we're ready, we split
+					SplitToTiles_Internal(SplitArgs).then([this, Promise = std::forward<decltype(Promise)>(Promise)](DeviceBufferRef Result) mutable
+						{
+							Promise.set_value(Result);
+						});
+					//});
+				});
 		});
-	});
 #else 
 	return cti::make_continuable<DeviceBufferRef>([this, Buffer, &Tiles](auto&& Promise)
-	{
-		Device::SplitToTiles_Generic(Buffer, Tiles).then([this, Promise = std::forward<decltype(Promise)>(Promise)](DeviceBufferRef Result) mutable
 		{
-			Promise.set_value(Result);
+			Device::SplitToTiles_Generic(Buffer, Tiles).then([this, Promise = std::forward<decltype(Promise)>(Promise)](DeviceBufferRef Result) mutable
+				{
+					Promise.set_value(Result);
+				});
 		});
-	});
 #endif /// NATIVE_SPLIT_TILES
 }
 
@@ -543,16 +544,16 @@ void Device_FX::FreeRenderTarget(HashType HashValue, UTextureRenderTarget2D* RT)
 	auto Iter = RTCache.find(HashValue);
 
 	if (Iter == RTCache.end())
-	{	
-		/// If the list wasnt previously found, that means RT was not created through AllocateRenderTarget
+	{
+		/// If the List wasnt previously found, that means RT was not created through AllocateRenderTarget
 		/// Hence there is no need to save this RT , so we simply release the resources
 		/// This is possible since FreeRenderTarget is called on deallocation, and deallocation can occur on shared ptr operator overloading as well
 		RT->ReleaseResource();
 		return;
 	}
 
-	RTList* list = Iter->second;
-	list->push_back(RT);
+	RTList* List = Iter->second;
+	List->push_back(RT);
 }
 
 TexPtr Device_FX::AllocateRenderTarget(const BufferDescriptor& Desc)
@@ -562,15 +563,15 @@ TexPtr Device_FX::AllocateRenderTarget(const BufferDescriptor& Desc)
 	uint8        NumMips = 1;
 	uint8        NumSamples = 1;
 	uint32       ExtData = 0;
-	
+
 	FRHITextureCreateInfo rhiDesc = FRHITextureDesc(ETextureDimension::Texture2D, ETextureCreateFlags::RenderTargetable, Desc.PixelFormat(),
-													FClearValueBinding(Desc.DefaultValue), FIntPoint(Desc.Width, Desc.Height), Depth, ArraySize, NumMips, NumSamples, ExtData);
-	
+		FClearValueBinding(Desc.DefaultValue), FIntPoint(Desc.Width, Desc.Height), Depth, ArraySize, NumMips, NumSamples, ExtData);
+
 	check(IsInGameThread());
 
 	SCOPE_CYCLE_COUNTER(STAT_Device_FX_AllocateRendertarget)
 
-	RTList* list = nullptr;
+		RTList* List = nullptr;
 	HashType Hash = Desc.FormatHashValue();
 
 	{
@@ -579,19 +580,19 @@ TexPtr Device_FX::AllocateRenderTarget(const BufferDescriptor& Desc)
 
 		if (Iter == RTCache.end())
 		{
-			list = new RTList();
-			RTCache[Hash] = list;
+			List = new RTList();
+			RTCache[Hash] = List;
 		}
 		else
-			list = Iter->second;
+			List = Iter->second;
 	}
 
-	if (!list->empty())
+	if (!List->empty())
 	{
-		/// otherwise, we can just pull it from the list
-		UTextureRenderTarget2D* RT = list->back();
+		/// otherwise, we can just pull it from the List
+		UTextureRenderTarget2D* RT = List->back();
 
-		list->pop_back();
+		List->pop_back();
 
 		FString existingName = RT->GetName();
 		//UE_LOG(LogDevice, Verbose, TEXT("Rename: %s => %s"), *existingName, *Desc.Name);
@@ -609,16 +610,16 @@ TexPtr Device_FX::AllocateRenderTarget(const BufferDescriptor& Desc)
 	}
 
 	check(IsInGameThread());
-	
+
 	TexDescriptor TextureDesc(Desc);
 	TexPtr TextureObj = std::make_shared<Tex>(TextureDesc);
 	TextureObj->InitRT();
-	
+
 	check(TextureObj->GetRenderTarget());
 
 	RTUsed.push_back(TextureObj->GetRenderTarget());
 
-	DeviceNativeTask_Lambda_Func allocateFunction = [this, TextureObj]() { return AllocateRTResource(TextureObj); }; 
+	DeviceNativeTask_Lambda_Func allocateFunction = [this, TextureObj]() { return AllocateRTResource(TextureObj); };
 	auto Task = DeviceNativeTask_Lambda::Create(this, (int32)E_Priority::kSystem, TEXT("AllocateRenderTarget"), allocateFunction);
 
 	return TextureObj;
@@ -628,9 +629,9 @@ TexPtr Device_FX::AllocateRenderTargetArray(const BufferDescriptor& Desc, int32
 {
 	check(IsInGameThread());
 
-SCOPE_CYCLE_COUNTER(STAT_Device_FX_AllocateRendertarget)
+	SCOPE_CYCLE_COUNTER(STAT_Device_FX_AllocateRendertarget)
 
-	RTList* list = nullptr;
+		RTList* List = nullptr;
 
 	HashType Hash = Desc.FormatHashValue();
 	{
@@ -639,11 +640,11 @@ SCOPE_CYCLE_COUNTER(STAT_Device_FX_AllocateRendertarget)
 
 		if (Iter == RTArrayCache.end())
 		{
-			list = new RTList();
-			RTArrayCache[Hash] = list;
+			List = new RTList();
+			RTArrayCache[Hash] = List;
 		}
 		else
-			list = Iter->second;
+			List = Iter->second;
 	}
 
 	check(IsInGameThread());
@@ -656,10 +657,10 @@ SCOPE_CYCLE_COUNTER(STAT_Device_FX_AllocateRendertarget)
 	RTArray->SetResource(RTArray->CreateResource());
 
 	/// The rest of it can stay in RenderingThread
-	DeviceNativeTask_Lambda_Func AllocateFunction = [this, TexArrayObj]() { return AllocateRTArrayResource(TexArrayObj); }; 
+	DeviceNativeTask_Lambda_Func AllocateFunction = [this, TexArrayObj]() { return AllocateRTArrayResource(TexArrayObj); };
 	auto Task = DeviceNativeTask_Lambda::Create(this, (int32)E_Priority::kSystem, TEXT("AllocateRenderTargetArray"), AllocateFunction);
 
-	list->push_back((UTextureRenderTarget2D*)TexArrayObj->GetRenderTargetArray());
+	List->push_back((UTextureRenderTarget2D*)TexArrayObj->GetRenderTargetArray());
 
 	return TexArrayObj;
 }
@@ -675,7 +676,7 @@ int32 Device_FX::InitRTResource(TexPtr TextureObj, UTextureRenderTarget* RT)
 	FTextureRenderTarget2DResource* RTRes = (FTextureRenderTarget2DResource*)RT->GetRenderTargetResource();
 	check(RTRes);
 
-	UE_LOG(LogDevice, VeryVerbose, TEXT("New RT Array allocation: %s %llu [Ptr: 0x%x, Size: %dx%d]"), *TextureObj->GetDescriptor().Name, 
+	UE_LOG(LogDevice, VeryVerbose, TEXT("New RT Array allocation: %s %llu [Ptr: 0x%x, Size: %dx%d]"), *TextureObj->GetDescriptor().Name,
 		TextureObj->GetDescriptor().Format_HashValue(), RT, RTRes->GetSizeX(), RTRes->GetSizeY());
 
 	FTexture2DRHIRef rhiTexture = RTRes->GetTextureRHI();
@@ -705,7 +706,7 @@ void Device_FX::MarkForCollection(TexPtr TextureObj)
 	{
 		UTextureRenderTarget2D* RT = TextureObj->RenderTarget();
 
-		/// Add to RT Display list
+		/// Add to RT Display List
 		HashType Hash = TextureObj->Descriptor().ToBufferDescriptor().Format_HashValue();
 		FreeRenderTarget(Hash, RT);
 		TextureObj->ReleaseRT();
@@ -728,7 +729,7 @@ DeviceBufferRef Device_FX::CreateFromTex(TexPtr TextureObj, bool bInitRaw)
 
 	if (bInitRaw)
 	{
-		RawBufferPtr Raw = TextureObj->Raw(); 
+		RawBufferPtr Raw = TextureObj->Raw();
 		FXBuffer = new DeviceBuffer_FX(this, TextureObj, Raw);
 	}
 	else
@@ -838,33 +839,33 @@ void Device_FX::PrintStats()
 
 	for (auto Iter = RTCache.begin(); Iter != RTCache.end(); Iter++)
 	{
-		RTList* rts = Iter->second;
-		if (rts != nullptr)
+		RTList* RenderTargets = Iter->second;
+		if (RenderTargets != nullptr)
 		{
-			for (UTextureRenderTarget2D* RT : *rts)
+			for (UTextureRenderTarget2D* RT : *RenderTargets)
 			{
 				auto RTFormat = RT->RenderTargetFormat;
 				auto PixelFormat = TextureHelper::GetPixelFormatFromRenderTargetFormat(RTFormat);
-				size_t size = RT->SizeX * RT->SizeY * TextureHelper::GetBppFromPixelFormat(PixelFormat) / 8;
-				RTMemUnused += size;
+				size_t Size = RT->SizeX * RT->SizeY * TextureHelper::GetBppFromPixelFormat(PixelFormat) / 8;
+				RTMemUnused += Size;
 				RTCountUnused++;
 			}
 		}
 	}
 
-	size_t rtMemUsed = 0;
+	size_t RTMemUsed = 0;
 	for (auto RT : RTUsed)
 	{
 		auto RTFormat = RT->RenderTargetFormat;
 		auto PixelFormat = TextureHelper::GetPixelFormatFromRenderTargetFormat(RTFormat);
-		size_t size = RT->SizeX * RT->SizeY * TextureHelper::GetBppFromPixelFormat(PixelFormat) / 8;
-		rtMemUsed += size;
+		size_t Size = RT->SizeX * RT->SizeY * TextureHelper::GetBppFromPixelFormat(PixelFormat) / 8;
+		RTMemUsed += Size;
 	}
 
 	UE_LOG(LogDevice, Log, TEXT("===== BEGIN Device: FX STATS (Native) ====="));
 
 	UE_LOG(LogDevice, Log, TEXT("[USED] RT Count   : %llu"), RTUsed.size());
-	UE_LOG(LogDevice, Log, TEXT("[USED] RT Mem     : %0.2f MB"), (float)rtMemUsed / (1024.0f * 1024.0f));
+	UE_LOG(LogDevice, Log, TEXT("[USED] RT Mem     : %0.2f MB"), (float)RTMemUsed / (1024.0f * 1024.0f));
 
 	UE_LOG(LogDevice, Log, TEXT("[Unused] RT Count : %llu"), RTCountUnused);
 	UE_LOG(LogDevice, Log, TEXT("[Unused] RT Mem   : %0.2f MB"), (float)RTMemUnused / (1024.0f * 1024.0f));
@@ -889,7 +890,7 @@ void Device_FX::InitTiles_Texture(BlobUPtr* Tiles, size_t NumRows, size_t NumCol
 {
 	BufferDescriptor TileDesc = InTileDesc;
 	TileDesc.Format = TextureHelper::FindOptimalSupportedFormat(InTileDesc.Format);
-	
+
 	for (int32 TileX = 0; TileX < NumRows; TileX++)
 	{
 		for (int32 TileY = 0; TileY < NumCols; TileY++)
@@ -916,7 +917,7 @@ void Device_FX::InitTiles_RenderTargets(BlobUPtr* Tiles, size_t NumRows, size_t
 		{
 			TexDescriptor TileDesc(InTileDesc);
 			TileDesc.Name = TextureHelper::CreateTileName(InTileDesc.Name, TileX, TileY);
-		
+
 			TexPtr TileTex = std::make_shared<Tex>(TileDesc);
 			TileTex->InitRT();
 
