FIxing the logic around how the Pin widget displays the current evaluated pin value
Introducing a getter function from the Pin to provide the current evaluated value as a String.
the getter is symetrical to the already existing Setter (renammed) from string

Clean up of the api of the Pin.

diff --git a/Engine/Plugins/Experimental/TextureGraph/Source/TextureGraph/Private/Expressions/Input/TG_Expression_Scalar.cpp b/Engine/Plugins/Experimental/TextureGraph/Source/TextureGraph/Private/Expressions/Input/TG_Expression_Scalar.cpp
--- a/Engine/Plugins/Experimental/TextureGraph/Source/TextureGraph/Private/Expressions/Input/TG_Expression_Scalar.cpp
+++ b/Engine/Plugins/Experimental/TextureGraph/Source/TextureGraph/Private/Expressions/Input/TG_Expression_Scalar.cpp
@@ -29,7 +29,6 @@ void UTG_Expression_Scalar::Evaluate(FTG_EvaluationContext* InContext)
 
 	// The Value is updated either as an input or as a setting and then becomes the output for this expression
 	// The pin out is named "ValueOut"
-	//InContext->Outputs.GetVar("ValueOut")->EditAs<float>() = InputScalar;
 	ValueOut = Scalar;
 }
 
diff --git a/Engine/Plugins/Experimental/TextureGraph/Source/TextureGraph/Private/Expressions/Input/TG_Expression_Vector.cpp b/Engine/Plugins/Experimental/TextureGraph/Source/TextureGraph/Private/Expressions/Input/TG_Expression_Vector.cpp
--- a/Engine/Plugins/Experimental/TextureGraph/Source/TextureGraph/Private/Expressions/Input/TG_Expression_Vector.cpp
+++ b/Engine/Plugins/Experimental/TextureGraph/Source/TextureGraph/Private/Expressions/Input/TG_Expression_Vector.cpp
@@ -30,7 +30,6 @@ void UTG_Expression_Vector::Evaluate(FTG_EvaluationContext* InContext)
 
 	// The Value is updated either as an input or as a setting and then becomes the output for this expression
 	// The pin out is named "ValueOut"
-	// InContext->Outputs.GetVar("ValueOut")->EditAs<FLinearColor>() = Value;
 	ValueOut = Vector;
 }
 
diff --git a/Engine/Plugins/Experimental/TextureGraph/Source/TextureGraph/Private/Expressions/TG_Expression.cpp b/Engine/Plugins/Experimental/TextureGraph/Source/TextureGraph/Private/Expressions/TG_Expression.cpp
--- a/Engine/Plugins/Experimental/TextureGraph/Source/TextureGraph/Private/Expressions/TG_Expression.cpp
+++ b/Engine/Plugins/Experimental/TextureGraph/Source/TextureGraph/Private/Expressions/TG_Expression.cpp
@@ -63,15 +63,22 @@ FTG_Signature::FInit UTG_Expression::GetSignatureInitArgsFromClass() const
 	for (TFieldIterator<FProperty> Prop(GetClass()); Prop; ++Prop) {
 		FName PropName = FName(Prop->GetNameCPP());
 		FName PropTypeName = FName(Prop->GetCPPType());
-		uint8 MaskNotConnectable = 0;
+		bool bIsNotConnectable = false; // by default all arguments are connectable
 #if WITH_EDITORONLY_DATA
 		
 		auto PropertyTGType = Prop->GetMetaData(TEXT("TGType"));
-		bool PropertyPinNotConnectable = Prop->HasMetaData(TEXT("TGPinNotConnectable"));
-		MaskNotConnectable = PropertyPinNotConnectable ? static_cast<uint8>(ETG_Access::NotConnectableFlag) : 0;
+		bIsNotConnectable = Prop->HasMetaData(TEXT("TGPinNotConnectable"));
 #else
 		FString PropertyTGType;
 #endif
+		// Add NotConnectable for certain types:
+		{
+			FByteProperty* ByteProperty = CastField<FByteProperty>(*Prop);
+			bIsNotConnectable |= (ByteProperty != nullptr);
+			FEnumProperty* EnumProperty = CastField<FEnumProperty>(*Prop);
+			bIsNotConnectable |= (EnumProperty != nullptr);
+		}
+		uint8 MaskNotConnectable = bIsNotConnectable ? static_cast<uint8>(ETG_Access::NotConnectableFlag) : 0;
 		
 		if (PropertyTGType.Compare(TEXT("TG_InputParam")) == 0)
 		{
diff --git a/Engine/Plugins/Experimental/TextureGraph/Source/TextureGraph/Private/TG_Pin.cpp b/Engine/Plugins/Experimental/TextureGraph/Source/TextureGraph/Private/TG_Pin.cpp
--- a/Engine/Plugins/Experimental/TextureGraph/Source/TextureGraph/Private/TG_Pin.cpp
+++ b/Engine/Plugins/Experimental/TextureGraph/Source/TextureGraph/Private/TG_Pin.cpp
@@ -225,26 +225,74 @@ FProperty* UTG_Pin::GetExpressionProperty() const
 	return Property;
 }
 
-void UTG_Pin::SetSelfVarValueFromString(const FString& InValueStr, bool bIsTweaking /*= false*/)
+
+FString UTG_Pin::GetEvaluatedVarValue() const
 {
-	Modify();
+	// First acces the true Var used from this Pin
+	const FTG_Var* CurrentVar = GetSelfVar();
+	if (IsInput() && IsConnected())
+	{
+		if (ConnectionNeedsConversion())
+		{
+			CurrentVar = &ConvertedVar;
+		}
+		else
+		{
+			CurrentVar = GetNodePtr()->GetGraph()->GetVar(Edges[0]);
+		}
+	}
 
-	// Pin->DefaultValue
-	FString DefaultValue = InValueStr;
+	FString CurrentValue = CurrentVar->LogValue();
 
 	// Enum requires special-case handling
 	// Needs help from the FProperty as the Var does not know the Enum type
 	FProperty* Property = GetExpressionProperty();
-	FEnumProperty* EnumProperty = CastField<FEnumProperty>(Property);
 	FByteProperty* ByteProperty = CastField<FByteProperty>(Property);
-	if (ByteProperty && ByteProperty->Enum)
+	if (ByteProperty)
 	{
-		DefaultValue = FString::FromInt(ByteProperty->Enum->GetValueByNameString(DefaultValue));
+		UEnum* Enum = ByteProperty->GetIntPropertyEnum();
+		if (Enum)
+		{
+			CurrentValue = Enum->GetNameByValue(FCString::Atoi(*CurrentValue)).ToString();
+		}
+	}
+	FEnumProperty* EnumProperty = CastField<FEnumProperty>(Property);
+	if (EnumProperty)
+	{
+		UEnum* Enum = EnumProperty->GetEnum();
+		if (Enum)
+		{
+			CurrentValue = Enum->GetNameByValue(FCString::Atoi(*CurrentValue)).ToString();
+		}
 	}
-	else if (EnumProperty && EnumProperty->GetEnum())
+
+	return CurrentValue;
+}
+
+void UTG_Pin::SetValue(const FString& InValueStr, bool bIsTweaking /*= false*/)
+{
+	Modify();
+
+	// Pin->DefaultValue
+	FString DefaultValue = InValueStr;
+
+	// Enum requires special-case handling
+	// Needs help from the FProperty as the Var does not know the Enum type
+	FProperty* Property = GetExpressionProperty();
+	if (Property)
 	{
-		DefaultValue = FString::FromInt(EnumProperty->GetEnum()->GetValueByNameString(DefaultValue));
+		FEnumProperty* EnumProperty = CastField<FEnumProperty>(Property);
+		FByteProperty* ByteProperty = CastField<FByteProperty>(Property);
+		if (ByteProperty && ByteProperty->Enum)
+		{
+			DefaultValue = FString::FromInt(ByteProperty->Enum->GetValueByNameString(DefaultValue));
+		}
+		else if (EnumProperty && EnumProperty->GetEnum())
+		{
+			DefaultValue = FString::FromInt(EnumProperty->GetEnum()->GetValueByNameString(DefaultValue));
+		}
 	}
+
 	EditSelfVar()->SetValueFromStr(DefaultValue);
 
 	NotifyPinSelfVarChanged(bIsTweaking);
diff --git a/Engine/Plugins/Experimental/TextureGraph/Source/TextureGraph/Public/TG_Pin.h b/Engine/Plugins/Experimental/TextureGraph/Source/TextureGraph/Public/TG_Pin.h
--- a/Engine/Plugins/Experimental/TextureGraph/Source/TextureGraph/Public/TG_Pin.h
+++ b/Engine/Plugins/Experimental/TextureGraph/Source/TextureGraph/Public/TG_Pin.h
@@ -79,6 +79,7 @@ class TEXTUREGRAPH_API UTG_Pin : public UObject {
 
 	// When an InputVarConverterKey is installed, then we host the result of the conversion as var value here
 	// this member is not saved and updated on every evaluation
+	// It is only useful during the evaluation to capture the value passed in and converted from a connected pin
 	UPROPERTY(Transient)
 		FTG_Var				ConvertedVar;
 
@@ -97,7 +98,7 @@ public:
 
 	bool				IsValid() const { return Id != FTG_Id::INVALID; }
 
-	UTG_Node* GetNodePtr() const;
+	UTG_Node*			GetNodePtr() const;
 	FTG_Id				GetNodeId() const;
 
 	FTG_Id				GetId() const { return Id; }
@@ -122,6 +123,9 @@ public:
 	const TArray<FTG_Id>& GetEdges() const { return Edges; }
 	bool				IsConnected() const { return !Edges.IsEmpty(); }
 
+	// Access the SelfVar of the Pin containing the currrent default value for the pin
+	// if an output pin, SelfVar contains the current result value produced by the node's expression by last evaluation
+	// if an input pin, SelfVar contains the current 'default' value of the pin used as expression input IF the pin is not connected.		
 	bool				IsValidSelfVar() const { return SelfVar.IsValid(); }
 	const FTG_Var*		GetSelfVar() const { return &SelfVar; }
 	FTG_Var*			EditSelfVar() { return &SelfVar; }
@@ -134,30 +138,23 @@ public:
 	//		if connected, it is the SelfVar Id of the Connected Pin
 	FTG_Id				GetVarId() const { return ((IsOutput() || Edges.IsEmpty()) ? GetId() : Edges[0]); }
 
-
 	// In the case the pin is an input, and is connected and the feeding pin var needs a conversion
 	// This implies that during the evaluation, the feeding Var is Converted into the SelfVar which is then input in the expression
 	bool				ConnectionNeedsConversion() const { return ((IsInput() || IsSetting()) && IsConnected() && !InputVarConverterKey.IsNone()); }
 	FName				GetInputVarConverterKey() const { return InputVarConverterKey; }
 
-	// GetConvertedInputVarId returns the VarId used to input its value in the pin's expression AFTER conversion
-	// This returns the self var id if the pin's ConnectionNeedsConversion()
-	// else it returns the GetVarId();
-	FTG_Id				GetConvertedInputVarId() const { return (ConnectionNeedsConversion() ? GetId() : GetVarId()); }
+	// EditConvertedVar returns the internal Var used to input its value in the pin's expression AFTER conversion
+	// Should only be accessed by the Graph Evaluation
 	FTG_Var*			EditConvertedVar() { return &ConvertedVar; }
 
 	FString				LogHead() const;
 	FString				Log(FString Tab) const;
-
 	FString				LogTooltip() const;
 
 	// Access the FProperty associated to this pin in the Expression::Class of the node
 	// A valid FProperty is returned if it exists, null otherwise
 	FProperty*			GetExpressionProperty() const;
 
-	void				SetSelfVarValueFromString(const FString& InValueStr, bool bIsTweaking = false);
-
-
 	FTG_Evaluation::VarConformer ConformerFunctor;
 	bool				NeedsConformance() const { return ConformerFunctor != nullptr; }
 
@@ -165,6 +162,7 @@ public:
 	static FTG_Hash		Hash(const UTG_Pin& Pin);
 	FTG_Hash			Hash() const { return Hash(*this); }
 
+
 	// Check the type of the argument's Var expected to be contained in the Pin's Var
 	bool IsArgScalar() const { return Argument.IsScalar(); }
 	bool IsArgColor() const { return Argument.IsColor(); }
@@ -191,7 +189,22 @@ public:
 	// After getting the value, the next step is to check the type contained in the variant.
 	bool GetValue(FTG_Variant& OutValue) const;
 
-	
+	// Useful for the UI:
+	// access the current value of this pin as passed to the evaluation if it is an input
+	// or received from the evaluation if it is an output 
+	// if output pin: the selfVar
+	// else if input pin:
+	//		if NOT connected: the selfVar
+	//		else if connected with a conversion: the convertedVar
+	//		else if connected without conversion: the Var from the connected pin
+	// This is MEANT TO BE USED for UI purpose !!!!
+	FString				GetEvaluatedVarValue() const;
+
+	// Set the self var value from string
+	// Used by UI to assign pin default value
+	// This is the same as the selfvar setter
+	void				SetValue(const FString& InValueStr, bool bIsTweaking = false);
+
 	// Setters for Scalar, Color, Vector arguments
 	bool SetValue(float Value);
 	bool SetValue(double Value) { return SetValue((float)Value); }
diff --git a/Engine/Plugins/Experimental/TextureGraph/Source/TextureGraphEditor/Private/EdGraph/TG_EdGraphNode.cpp b/Engine/Plugins/Experimental/TextureGraph/Source/TextureGraphEditor/Private/EdGraph/TG_EdGraphNode.cpp
--- a/Engine/Plugins/Experimental/TextureGraph/Source/TextureGraphEditor/Private/EdGraph/TG_EdGraphNode.cpp
+++ b/Engine/Plugins/Experimental/TextureGraph/Source/TextureGraphEditor/Private/EdGraph/TG_EdGraphNode.cpp
@@ -187,34 +187,20 @@ void UTG_EdGraphNode::AllocateDefaultPins()
 
 			NewPin->bNotConnectable = Pin->IsNotConnectable() || Pin->IsParam();
 
-			FString DefaultValue = Pin->GetSelfVar()->LogValue();
-			FProperty* Property = Pin->GetExpressionProperty(); //// BEWARE Pin do NOT all have a Property, it is potentially NULL
-			FByteProperty* ByteProperty = CastField<FByteProperty>(Property);
-			if (ByteProperty)
-			{
-				UEnum* Enum = ByteProperty->GetIntPropertyEnum();
-				if (Enum)
-				{
-					DefaultValue = Enum->GetNameByValue(FCString::Atoi(*DefaultValue)).ToString();
-				}
-				NewPin->bNotConnectable = true;
-			}
-			FEnumProperty* EnumProperty = CastField<FEnumProperty>(Property);
-			if (EnumProperty)
-			{
-				UEnum* Enum = EnumProperty->GetEnum();
-				if (Enum)
-				{
-					DefaultValue = Enum->GetNameByValue(FCString::Atoi(*DefaultValue)).ToString();
-				}
-				NewPin->bNotConnectable = true;
-			}
+			// Fetch the current pin value as a string and assign as the current value
+			FString DefaultValue = Pin->GetEvaluatedVarValue();
 			NewPin->DefaultValue = DefaultValue;
+
+			// Also grab the Property associated if exist
+			// BEWARE Pin do NOT all have a Property, it is potentially NULL
+			FProperty* Property = Pin->GetExpressionProperty();
+
 			// This updates the UObject (Texture/Material) picker UI in the Node to get updated
 			if (Property && Property->GetClass()->IsChildOf(FObjectProperty::StaticClass()))
 			{
 				NewPin->DefaultObject = Pin->EditSelfVar()->GetAs<TObjectPtr<UObject>>();
 			}
+
 			NewPin->PinFriendlyName = FText::FromName(Pin->GetAliasName());
 			NewPin->PinToolTip = Pin->LogTooltip();
 
@@ -491,7 +477,8 @@ void UTG_EdGraphNode::PinDefaultValueChangedWithTweaking(UEdGraphPin* Pin, bool
 {
 	const UTG_EdGraphSchema* Schema = Cast<const UTG_EdGraphSchema>(GetSchema());
 	UTG_Pin* TGPin = Schema->GetTGPinFromEdPin(Pin);
-	TGPin->SetSelfVarValueFromString(Pin->DefaultValue, bIsTweaking);
+
+	TGPin->SetValue(Pin->DefaultValue, bIsTweaking);
 
 	// This updates the UObject (Texture/Material) picker UI in the Node to get updated
 	FProperty* Property = TGPin->GetExpressionProperty();
@@ -609,35 +596,13 @@ bool UTG_EdGraphNode::UpdateEdPinDefaultValue(UEdGraphPin* EdPin, const UTG_EdGr
 	if (TSPin)
 	{
 		UTG_EdGraph* EdGraph = CastChecked<UTG_EdGraph>(GetGraph());
-		UTG_Graph* Graph = EdGraph->TextureGraph->Graph();
 		
-		FTG_Var* Var = Graph->GetVar(TSPin->GetVarId());
+		// Access the current pin evaluated value as a string
+		FString DefaultValue = TSPin->GetEvaluatedVarValue();
 
-		bool ShowSelfValue = TSPin->NeedsConformance() && TSPin->IsConnected();
-		FString DefaultValue =  ShowSelfValue ? TSPin->GetSelfVar()->LogValue() : Var->LogValue();
-		
-		// Enum requires special-case handling
-		// Needs help from the FProperty as the Var does not know the Enum type
+		// Also grab the Property associated if exist
 		FProperty* Property = TSPin->GetExpressionProperty();
-		FByteProperty* ByteProperty = CastField<FByteProperty>(Property);
-		if (ByteProperty)
-		{
-			UEnum* Enum = ByteProperty->GetIntPropertyEnum();
-			if (Enum)
-			{
-				DefaultValue = Enum->GetNameByValue(FCString::Atoi(*DefaultValue)).ToString();
-			}
-		}
-		FEnumProperty* EnumProperty = CastField<FEnumProperty>(Property);
-		if (EnumProperty)
-		{
-			UEnum* Enum = EnumProperty->GetEnum();
-			if (Enum)
-			{
-				DefaultValue = Enum->GetNameByValue(FCString::Atoi(*DefaultValue)).ToString();
-			}
-		}
-		
+
 #if WITH_EDITOR
 		// check here if the property has a tag for regen pins and value has changed
 		bShouldUpdatePinsVisibility = (Property != nullptr && Property->HasMetaData("RegenPinsOnChange"))
