Fixing bug with SelfVar value for a UObject  Property of an Expression

If the Property has a setter, then when we copy the value from Var to Expression's Property, then we now assign back the true value located in the expression's property back in the Var
This allow the case when the setter CHANGE the expected result, so that the Var is fedback the true current value.

Fixed the consequence of this with regard to undo/redo to get proper behavior.

Also moved the TextureGRaph to the input category

diff --git a/Engine/Plugins/Experimental/TextureGraph/Source/TextureGraph/Private/Expressions/Input/TG_Expression_Graph.cpp b/Engine/Plugins/Experimental/TextureGraph/Source/TextureGraph/Private/Expressions/Input/TG_Expression_Graph.cpp
--- a/Engine/Plugins/Experimental/TextureGraph/Source/TextureGraph/Private/Expressions/Input/TG_Expression_Graph.cpp
+++ b/Engine/Plugins/Experimental/TextureGraph/Source/TextureGraph/Private/Expressions/Input/TG_Expression_Graph.cpp
@@ -95,14 +95,61 @@ void UTG_Expression_Graph::Evaluate(FTG_EvaluationContext* InContext)
 //////////////////////////////////////////////////////////////////////////
 
 #if WITH_EDITOR
+void UTG_Expression_TextureGraph::PreEditChange(FProperty* PropertyAboutToChange)
+{
+	// if Graph changes catch it first
+	if (PropertyAboutToChange->GetName() == GET_MEMBER_NAME_CHECKED(UTG_Expression_TextureGraph, TextureGraph))
+	{	
+		bool bShouldMarkAsDirty;
+		if (GIsTransacting)
+		{
+			// Don't mark the outer package as dirty during an undo/redo operation.
+			bShouldMarkAsDirty = false;
+		}
+		else if (PropertyAboutToChange && PropertyAboutToChange->HasAnyPropertyFlags(CPF_SkipSerialization))
+		{
+			// Don't mark the outer package as dirty if we're about to change a non-serializable property.
+			bShouldMarkAsDirty = false;
+		}
+		else
+		{
+			PreEditTextureGraph = TextureGraph;
+
+			// Don't mark the outer package as dirty if annotated to be deferred (e.g. during propagation).
+			bShouldMarkAsDirty = false;
+		}
+
+		Modify(bShouldMarkAsDirty);
+
+	}
+	else
+	{
+		Super::PreEditChange(PropertyAboutToChange);
+	}
+}
+
 
 void UTG_Expression_TextureGraph::PostEditChangeProperty(FPropertyChangedEvent& PropertyChangedEvent)
 {
 	// if Graph changes catch it first
 	if (PropertyChangedEvent.GetPropertyName() == GET_MEMBER_NAME_CHECKED(UTG_Expression_TextureGraph, TextureGraph))
 	{
-		UE_LOG(LogTextureGraph, Log, TEXT("TextureGraph Expression PostEditChangeProperty."));
-		SetTextureGraphInternal(TextureGraph);
+		// In the case of TextureGraph changed from detail panel
+		// Validate the TextureGraph assigned here
+		// if it is not passing then fall back to the previous version which should have been captured in the PreEditChange()
+		bool IsValidTextureGraph = CheckDependencies(TextureGraph);
+		UTextureGraph* PreEditTextureGraphPtr = PreEditTextureGraph.Get();
+		PreEditTextureGraph.Reset();
+		if (!IsValidTextureGraph)
+		{
+			// Early exit without changing anything
+			// Restore the TextureGRaph to its pre edit change values
+			TextureGraph = PreEditTextureGraphPtr;
+		}
+		else
+		{
+			SetTextureGraphInternal(TextureGraph);
+		}		
 	}
 
 	Super::PostEditChangeProperty(PropertyChangedEvent);
@@ -185,17 +232,16 @@ bool UTG_Expression_TextureGraph::CheckDependencies(const UTextureGraph* InTextu
 void UTG_Expression_TextureGraph::SetTextureGraphInternal(UTextureGraph* InTextureGraph)
 {
 	// validate that we're allowed to use InTextureGraph as input here
-	if (!CheckDependencies(InTextureGraph))
+	// if dependencies exist, we don't want to use the InTextureGraph so we ll treat it just like a nullptr.
+	// TODO: At the moment we reset TextureGraph to nullptr when a dependency check fails.
+	// Ideally we'd like to revert to the previous valid value. The problem here is that
+	// TextureGraph at this point comes pre-filled with the problematic value 
+	bool IsValidTextureGraph = CheckDependencies(InTextureGraph);
+	if (!IsValidTextureGraph)
 	{
-		// if dependencies exist, we don't want to keep using the problematic TextureGraph
-		// TODO: At the moment we reset TextureGraph to nullptr when a dependency check fails.
-		// Ideally we'd like to revert to the previous valid value. The problem here is that
-		// TextureGraph at this point comes pre-filled with the problematic value 
-		TextureGraph = nullptr;
-		RuntimeGraph = nullptr;
 		return;
 	}
-	
+
 	Modify();
 	if (!InTextureGraph)
 	{
diff --git a/Engine/Plugins/Experimental/TextureGraph/Source/TextureGraph/Private/TG_Var.cpp b/Engine/Plugins/Experimental/TextureGraph/Source/TextureGraph/Private/TG_Var.cpp
--- a/Engine/Plugins/Experimental/TextureGraph/Source/TextureGraph/Private/TG_Var.cpp
+++ b/Engine/Plugins/Experimental/TextureGraph/Source/TextureGraph/Private/TG_Var.cpp
@@ -244,6 +244,12 @@ void VarPropertySerializer_ObjectProperty(FTG_Var::VarPropertySerialInfo& Info)
 		{
 			TObjectPtr<UObject> ObjectPtr = Info.Var->GetAs<TObjectPtr<UObject>>();
 			ObjectPtrProperty->SetObjectPropertyValue_InContainer(Info.Owner, ObjectPtr.Get(), Info.ClampedIndex());
+
+			// If UObject is going through a setter then make sure to feedback the true end value in the var
+			if (Property->HasSetter())
+			{
+				Info.Var->EditAs<TObjectPtr<UObject>>() = ObjectPtrProperty->GetObjectPropertyValue_InContainer(Info.Owner, Info.ClampedIndex());
+			}
 		}
 		else
 		{
diff --git a/Engine/Plugins/Experimental/TextureGraph/Source/TextureGraph/Public/Expressions/Input/TG_Expression_Graph.h b/Engine/Plugins/Experimental/TextureGraph/Source/TextureGraph/Public/Expressions/Input/TG_Expression_Graph.h
--- a/Engine/Plugins/Experimental/TextureGraph/Source/TextureGraph/Public/Expressions/Input/TG_Expression_Graph.h
+++ b/Engine/Plugins/Experimental/TextureGraph/Source/TextureGraph/Public/Expressions/Input/TG_Expression_Graph.h
@@ -19,7 +19,7 @@ class TEXTUREGRAPH_API UTG_Expression_Graph : public UTG_Expression
 {
 	GENERATED_BODY()
 public:
-	TG_DECLARE_DYNAMIC_EXPRESSION(TG_Category::Utilities);
+	TG_DECLARE_DYNAMIC_EXPRESSION(TG_Category::Input);
 	virtual void Evaluate(FTG_EvaluationContext* InContext) override;
 
 protected:
@@ -46,7 +46,13 @@ public:
 	UTG_Expression_TextureGraph();
 	virtual ~UTG_Expression_TextureGraph();
 	
+
 #if WITH_EDITOR
+protected:
+	// Only in editor mode, in the case the TextureGraph is assigned from Details panel we need to remember the previous version
+	TWeakObjectPtr<UTextureGraph>	PreEditTextureGraph; 
+public:
+	virtual void PreEditChange(FProperty* PropertyAboutToChange) override;
 	virtual void PostEditChangeProperty(FPropertyChangedEvent& PropertyChangedEvent) override;
 	virtual void PostEditUndo() override;
 #endif
