Prevent histogram asserting in case of a transient texture making it to the histogram service.

diff --git a/Engine/Plugins/Experimental/TextureGraph/Source/TextureGraph/Public/Expressions/Maths/TG_Expression_Maths_ThreeInputs.h b/Engine/Plugins/Experimental/TextureGraph/Source/TextureGraph/Public/Expressions/Maths/TG_Expression_Maths_ThreeInputs.h
--- a/Engine/Plugins/Experimental/TextureGraph/Source/TextureGraph/Public/Expressions/Maths/TG_Expression_Maths_ThreeInputs.h
+++ b/Engine/Plugins/Experimental/TextureGraph/Source/TextureGraph/Public/Expressions/Maths/TG_Expression_Maths_ThreeInputs.h
@@ -57,8 +57,8 @@ public:
 	UPROPERTY(meta = (TGType = "TG_Input", PinDisplayName = "Y"))
 	FTG_Variant							Input2 = 0;
 
-	// A value that linearly interpolates between X and Y
-	UPROPERTY(EditAnywhere, Category = NoCategory, meta = (TGType = "TG_Setting", PinDisplayName = "S"))
+	// A value that linearly interpolates between X and Y. This can be an image or a Scalar node connected for easier control
+	UPROPERTY(meta = (TGType = "TG_Input", PinDisplayName = "S"))
 	FTG_Variant							LerpValue = 0;
 
 public:
diff --git a/Engine/Plugins/Experimental/TextureGraph/Source/TextureGraphEngine/Transform/Utility/T_TextureHistogram.cpp b/Engine/Plugins/Experimental/TextureGraph/Source/TextureGraphEngine/Transform/Utility/T_TextureHistogram.cpp
--- a/Engine/Plugins/Experimental/TextureGraph/Source/TextureGraphEngine/Transform/Utility/T_TextureHistogram.cpp
+++ b/Engine/Plugins/Experimental/TextureGraph/Source/TextureGraphEngine/Transform/Utility/T_TextureHistogram.cpp
@@ -66,7 +66,6 @@ TiledBlobPtr T_TextureHistogram::CreateJobAndResult(JobUPtr& OutJob, MixUpdateCy
 
 	Result->MakeSingleBlob();
 
-
 	if (!SourceTex->HasHistogram())
 	{
 		//setting it as the histogram of source so it is retained untill the life cycle of source blob.
@@ -101,7 +100,11 @@ TiledBlobPtr T_TextureHistogram::CreateOnService(UMixInterface* InMix, TiledBlob
 
 	check(InMix);
 	check(SourceTex);
-	check(!SourceTex->IsTransient());
+
+	/// If the source texture turns out to be transient at this point, we just return a black histogram. 
+	/// We don't want to calculate anything for transient buffers
+	if (SourceTex->IsTransient())
+		return TextureHelper::GetBlack();
 
 	HistogramServicePtr Service = TextureGraphEngine::GetScheduler()->GetHistogramService().lock();
 	check(Service);
