Texture Graph : Issues with renaming output
Removed TG_outputSettingSet and its usage as it is no longer needed
Fixed, Output node doesn't expose in the Export window
Linking an Output node reverts its Display name back to Output
Crash after deleting a renamed and unlinked Output Node
Name discrepancy between the Output Node header, the "Outputs to export" name and the "Output Name" property of the node
Fixed output file name and path set to none when add another output

diff --git a/Engine/Plugins/Experimental/TextureGraph/Source/TextureGraph/Private/Expressions/Output/TG_Expression_Output.cpp b/Engine/Plugins/Experimental/TextureGraph/Source/TextureGraph/Private/Expressions/Output/TG_Expression_Output.cpp
--- a/Engine/Plugins/Experimental/TextureGraph/Source/TextureGraph/Private/Expressions/Output/TG_Expression_Output.cpp
+++ b/Engine/Plugins/Experimental/TextureGraph/Source/TextureGraph/Private/Expressions/Output/TG_Expression_Output.cpp
@@ -6,6 +6,7 @@
 #include "FxMat/MaterialManager.h"
 #include "Job/Job.h"
 #include "Job/JobArgs.h"
+#include "TG_Graph.h"
 #include "Transform/Expressions/T_FlatColorTexture.h"
 
 void UTG_Expression_Output::Evaluate(FTG_EvaluationContext* InContext)
@@ -97,6 +98,7 @@ void UTG_Expression_Output::Evaluate(FTG_EvaluationContext* InContext)
 void UTG_Expression_Output::SetTitleName(FName NewName)
 {
 	GetParentNode()->GetOutputPin(GET_MEMBER_NAME_CHECKED(UTG_Expression_Output, Output))->SetAliasName(NewName);
+	OutputSettings.OutputName = GetParentNode()->GetOutputPin(GET_MEMBER_NAME_CHECKED(UTG_Expression_Output, Output))->GetAliasName();
 }
 
 FName UTG_Expression_Output::GetTitleName() const
@@ -111,4 +113,20 @@ void UTG_Expression_Output::UpdateBufferDescriptorValues()
 	OutputTexture.Descriptor.Width = OutputSettings.Width;
 	OutputTexture.Descriptor.Height = OutputSettings.Height;
 	OutputTexture.Descriptor.TextureFormat = OutputSettings.TextureFormat;
+}
+
+void UTG_Expression_Output::InitializeOutputSettings()
+{
+	SetTitleName(GetDefaultName());
+	OutputSettings.Initialize(GetParentNode()->GetGraph()->GetPathName(),GetTitleName());
+
+	UTG_Pin* Settings = GetParentNode()->GetPin(GET_MEMBER_NAME_CHECKED(UTG_Expression_Output, OutputSettings));
+	Settings->SetValue(OutputSettings.ToString());
+}
+
+void UTG_Expression_Output::SetExport(bool bExport)
+{
+	OutputSettings.bExport = bExport;
+	UTG_Pin* Settings = GetParentNode()->GetPin(GET_MEMBER_NAME_CHECKED(UTG_Expression_Output, OutputSettings));
+	Settings->SetValue(OutputSettings.ToString());
 }
\ No newline at end of file
diff --git a/Engine/Plugins/Experimental/TextureGraph/Source/TextureGraph/Private/TG_Graph.cpp b/Engine/Plugins/Experimental/TextureGraph/Source/TextureGraph/Private/TG_Graph.cpp
--- a/Engine/Plugins/Experimental/TextureGraph/Source/TextureGraph/Private/TG_Graph.cpp
+++ b/Engine/Plugins/Experimental/TextureGraph/Source/TextureGraph/Private/TG_Graph.cpp
@@ -3,6 +3,7 @@
 // ReSharper disable All
 #include "TG_Graph.h"
 #include "Expressions/TG_Expression.h"
+#include "Expressions/Output/TG_Expression_Output.h"
 #include "TG_CustomVersion.h"
 
 #include "Algo/Reverse.h"
@@ -164,10 +165,20 @@ void UTG_Graph::RegenerateNode(UTG_Node* InNode)
 	// the element is Invalid if there is no matching existing Pin
 	FTG_Indices NewIdxToOld = FTG_Signature::GenerateMappingArgIdxTable((*OldSignature), (*NewSignature));
 
+	//Save the edited alias name so we can reaply them
+	//Alias names can be edited when used as a title for node
+	TArray<FName> EditedAliasNames;
+	for (int32 i = 0; i < InNode->Pins.Num(); ++i)
+	{
+		FName EditedAlias = InNode->Pins[i]->HasAliasName() ? InNode->Pins[i]->GetAliasName() : NAME_None;
+		EditedAliasNames.Add(EditedAlias);
+	}
+
 	// Go through the allocated pins and reuse the matching arguments
 	// Store the future "Pins" array of the node in the "NewPinArray"
 	FTG_Index InNodeIdx = InNode->GetId().NodeIdx();
 	TArray<TObjectPtr<UTG_Pin>> NewPinArray;
+	
 	for (int32 i = 0; i < NewIdxToOld.Num(); ++i)
 	{
 		FTG_Index OldIdx = NewIdxToOld[i];
@@ -230,6 +241,16 @@ void UTG_Graph::RegenerateNode(UTG_Node* InNode)
 		{
 			auto& Arg = NewSignature->GetArgument(i);
 			AllocatePin(InNode, Arg, i);
+
+			if (i >= 0 && i < EditedAliasNames.Num())
+			{
+				//if the pin reallocated ever had an edited Name we should keep it
+				FName EditedAlias = EditedAliasNames[i];
+				if (!EditedAlias.IsNone())
+				{
+					InNode->Pins[i]->SetAliasName(EditedAlias);
+				}
+			}
 		}
 		// If pin is reused, reassign the new signature argument to update potential changes in the ArgFlags
 		else
@@ -758,6 +779,18 @@ void UTG_Graph::ForEachEdges(std::function<void(const UTG_Pin* /*pinFrom*/, cons
 	);
 }
 
+void UTG_Graph::ForEachOutputSettings( std::function<void(const FTG_OutputSettings& /*settings*/)> visitor)
+{
+		ForEachNodes([visitor](const UTG_Node* Node, uint32 Index)
+		{
+			UTG_Expression_Output* TargetExpression = Cast<UTG_Expression_Output>(Node->GetExpression());
+			if (TargetExpression)
+			{
+				visitor(TargetExpression->OutputSettings);
+			}
+		});
+}
+
 int UTG_Graph::GetOutputParamTextures(TArray<FName>& OutNames, FTG_Ids& OutPinIds) const
 {
 	int NumFounds = 0;
@@ -955,7 +988,7 @@ void UTG_Graph::EvalTraverseOrder() const
 	while (nodeReservoirs[reservoirIndex].Num())
 	{
 		GatherOuterNodes(sourceNodesArray, nodeReservoirs[reservoirIndex], nodeReservoirs[(pass + 1) % 2]);
-		if (nodeReservoirs[pass].IsEmpty())
+		if (nodeReservoirs[reservoirIndex].IsEmpty())
 			break; /// THis should not happen, means a 
 
 		// Add all the outer nodes in the traverse order
diff --git a/Engine/Plugins/Experimental/TextureGraph/Source/TextureGraph/Private/TG_HelperFunctions.cpp b/Engine/Plugins/Experimental/TextureGraph/Source/TextureGraph/Private/TG_HelperFunctions.cpp
--- a/Engine/Plugins/Experimental/TextureGraph/Source/TextureGraph/Private/TG_HelperFunctions.cpp
+++ b/Engine/Plugins/Experimental/TextureGraph/Source/TextureGraph/Private/TG_HelperFunctions.cpp
@@ -101,66 +101,70 @@ AsyncBool FTG_HelperFunctions::ExportAsync(UTextureGraph* InTextureGraph, FStrin
 	MixUpdateCyclePtr Cycle = Batch->GetCycle();
 	InTextureGraph->Update(Cycle);
 	
-	auto SettingsSet = InTextureGraph->GetOutputSettingsSet();
+	//auto SettingsSet = InTextureGraph->GetOutputSettingsSet();
 
-	for (const auto& Info : SettingsSet->OutputExpressionInfos)
+	InTextureGraph->Graph()->ForEachNodes([=,&TargetExportSettings](const UTG_Node* Node, uint32 Index)
 	{
-		if (Info.bExport || ExportAllOutputs)
+		UTG_Expression_Output* TargetExpression = Cast<UTG_Expression_Output>(Node->GetExpression());
+		if (TargetExpression)
 		{
-			auto OutputSetting = Info.OutputPtr->OutputSettings;
-			auto Node = Cast<UTG_Node>(Info.OutputPtr->GetOuter());
+			FTG_OutputSettings& OutputSetting = TargetExpression->OutputSettings;
 
-			FTG_EvaluationContext EvaluationContext;
-			EvaluationContext.Cycle = Batch->GetCycle();
-			
-			auto ExportBlobs = FTG_HelperFunctions::GetTexturedOutputs(Node, &EvaluationContext);
-
-			FString Path = OutputSetting.FolderPath.ToString();
-			if (OverrideExportPath)
+			if (OutputSetting.bExport || ExportAllOutputs)
 			{
-				Path = ExportPath;
-			}
+				FTG_EvaluationContext EvaluationContext;
+				EvaluationContext.Cycle = Batch->GetCycle();
 
-			FString Name = AssetName.IsEmpty() ? OutputSetting.GetFullOutputName() : AssetName;
-			
-			bool IsNameValid = TextureExporter::IsPackageNameValid(Path, Name);
+				auto ExportBlobs = FTG_HelperFunctions::GetTexturedOutputs(Node, &EvaluationContext);
 
-			bool HasOutputs = ExportBlobs.Num() > 0;
+				FString Path = OutputSetting.FolderPath.ToString();
+				if (OverrideExportPath)
+				{
+					Path = ExportPath;
+				}
 
-			if (HasOutputs && IsNameValid)
-			{
-				auto Output = std::static_pointer_cast<TiledBlob>(ExportBlobs[0]);//Dealing with one output per Node for now
-				FExportMapSettings MapSettings = TextureExporter::GetExportSettingsForTarget(TargetExportSettings, std::static_pointer_cast<TiledBlob>(Output), *Name);
-				MapSettings.Name = FName(*Name);
-				MapSettings.Path = Path;
-				MapSettings.UseOverridePath = OverrideExportPath;
-				MapSettings.OverwriteTextures = OverwriteTextures;
-				MapSettings.LODGroup = OutputSetting.LODGroup;
-				MapSettings.Compression = OutputSetting.Compression;
-				MapSettings.IsSRGB = OutputSetting.bSRGB;
-				TargetExportSettings.ExportPreset.push_back(std::pair<FName, FExportMapSettings>{ MapSettings.Name, MapSettings });
-			}
-			else
-			{
-				//Log Error to Error System
-				if (!HasOutputs)
+				FString Name = AssetName.IsEmpty() ? OutputSetting.GetFullOutputName() : AssetName;
+
+				bool IsNameValid = TextureExporter::IsPackageNameValid(Path, Name);
+
+				bool HasOutputs = ExportBlobs.Num() > 0;
+
+				if (HasOutputs && IsNameValid)
 				{
-					auto ErrorType = static_cast<int32>(ETextureGraphErrorType::UNSUPPORTED_TYPE);
-					TextureGraphEngine::GetErrorReporter(InTextureGraph)->ReportError(ErrorType, FString::Format(TEXT("Texture Export Error : No valid output found for OutputSetting {0}"), { OutputSetting.OutputName.ToString() }), nullptr);
+					TiledBlobPtr Output = std::static_pointer_cast<TiledBlob>(ExportBlobs[0]);//Dealing with one output per Node for now
+					FExportMapSettings MapSettings = TextureExporter::GetExportSettingsForTarget(TargetExportSettings, std::static_pointer_cast<TiledBlob>(Output), *Name);
+					MapSettings.Name = FName(*Name);
+					MapSettings.Path = Path;
+					MapSettings.UseOverridePath = OverrideExportPath;
+					MapSettings.OverwriteTextures = OverwriteTextures;
+					MapSettings.LODGroup = OutputSetting.LODGroup;
+					MapSettings.Compression = OutputSetting.Compression;
+					MapSettings.IsSRGB = OutputSetting.bSRGB;
+					TargetExportSettings.ExportPreset.push_back(std::pair<FName, FExportMapSettings>{ MapSettings.Name, MapSettings });
 				}
-				if (!IsNameValid)
+				else
 				{
-					auto ErrorType = static_cast<int32>(ETextureGraphErrorType::UNSUPPORTED_TYPE);
-					TextureGraphEngine::GetErrorReporter(InTextureGraph)->ReportError(ErrorType, FString::Format(TEXT("Texture Export Error : Invalid path set for OutputSetting {0}"), { OutputSetting.OutputName.ToString() }), nullptr);
+					//Log Error to Error System
+					if (!HasOutputs)
+					{
+						auto ErrorType = static_cast<int32>(ETextureGraphErrorType::UNSUPPORTED_TYPE);
+						TextureGraphEngine::GetErrorReporter(InTextureGraph)->ReportError(ErrorType, FString::Format(TEXT("Texture Export Error : No valid output found for OutputSetting {0}"), { OutputSetting.OutputName.ToString() }), nullptr);
+					}
+					if (!IsNameValid)
+					{
+						auto ErrorType = static_cast<int32>(ETextureGraphErrorType::UNSUPPORTED_TYPE);
+						TextureGraphEngine::GetErrorReporter(InTextureGraph)->ReportError(ErrorType, FString::Format(TEXT("Texture Export Error : Invalid path set for OutputSetting {0}"), { OutputSetting.OutputName.ToString() }), nullptr);
+					}
 				}
 			}
 		}
-	}
+	});
+
 	return RenderAsync(InTextureGraph, Batch)
-			.then([InTextureGraph, TargetExportSettings, ExportPath](auto result) 
-			{
-				return TextureExporter::ExportAsUAsset(InTextureGraph, TargetExportSettings, ExportPath);
-			});
+	.then([InTextureGraph, TargetExportSettings, ExportPath](auto result) 
+	{
+		return TextureExporter::ExportAsUAsset(InTextureGraph, TargetExportSettings, ExportPath);
+	});
 }
 
 void FTG_HelperFunctions::InitTargets(UTextureGraph* InTextureGraph)
@@ -292,4 +296,4 @@ AsyncBool FTG_HelperFunctions::RenderAsync(UTextureGraph* InTextureGraph, JobBat
 				});
 		}
 	}); 
-}
+}
\ No newline at end of file
diff --git a/Engine/Plugins/Experimental/TextureGraph/Source/TextureGraph/Private/TG_OutputSettings.cpp b/Engine/Plugins/Experimental/TextureGraph/Source/TextureGraph/Private/TG_OutputSettings.cpp
--- a/Engine/Plugins/Experimental/TextureGraph/Source/TextureGraph/Private/TG_OutputSettings.cpp
+++ b/Engine/Plugins/Experimental/TextureGraph/Source/TextureGraph/Private/TG_OutputSettings.cpp
@@ -7,12 +7,6 @@
 #include "TG_HelperFunctions.h"
 #include "Misc/Paths.h"
 
-FTG_OutputSettings* FTG_OutputExpressionInfo::GetOutputSettings()
-{
-	check(OutputPtr);
-	return &OutputPtr->OutputSettings;
-}
-
 void FTG_OutputSettings::Set(int InWidth, int InHeight, FString Name /*= "None"*/, FString Path /*= "None"*/, ETG_TextureFormat Format /*= ETG_TextureFormat::BGRA8*/, ETG_TexturePresetType InTextureType /*= ETG_TexturePresetType::None*/,
 	TextureCompressionSettings InCompression /*= TextureCompressionSettings::TC_Default*/, TextureGroup InLodGroup /*= TextureGroup::TEXTUREGROUP_World*/, bool InbSRGB /*= false*/)
 {
@@ -79,124 +73,13 @@ void FTG_OutputSettings::OnSetTexturePresetType(ETG_TexturePresetType Type)
 	}
 }
 
-void UTG_OutputSettingsSet::InitOutputSettings()
-{
-	OutputExpressionInfos.Empty();
-
-}
-
-FTG_OutputExpressionInfo* UTG_OutputSettingsSet::GetOutputExpressionInfo(FName OutputName)
-{
-	auto predicate = [OutputName](const FTG_OutputExpressionInfo& OutputPtr) { return OutputPtr.OutputName == OutputName; };
-	auto bExists = OutputExpressionInfos.ContainsByPredicate(predicate);
-	check(bExists);
-	return OutputExpressionInfos.FindByPredicate(predicate);
-}
-
-FTG_OutputSettings* UTG_OutputSettingsSet::GetOutputSetting(FName OutputName)
-{
-	return GetOutputExpressionInfo(OutputName)->GetOutputSettings();
-}
-
-void UTG_OutputSettingsSet::AddOutputSetting(UTextureGraph* InTextureGraph, FName OutputName, UTG_Expression_Output* Output)
+void FTG_OutputSettings::Initialize(FString PathName,FName InName /*= "Output"*/)
 {
-	Modify();
-
-	FTG_OutputSettings Settings;
-	Settings.OutputName = OutputName;
-	Settings.BaseName = OutputName;
-	FString AssetPath = InTextureGraph->GetPathName();
+	OutputName = InName;
+	BaseName = InName;
+	FString AssetPath = PathName;
 	FString DefaultDirectory = FPaths::GetPath(AssetPath);
 
-	Settings.FolderPath = FName(DefaultDirectory);
-	Settings.TextureFormat = ETG_TextureFormat::BGRA8;
-
-	Output->OutputSettings = Settings;
-	
-	FTG_OutputExpressionInfo SettingsPtr;
-	SettingsPtr.OutputName = OutputName;
-	SettingsPtr.OutputPtr = Output;
-	OutputExpressionInfos.Add(SettingsPtr);
-
-}
-
-void UTG_OutputSettingsSet::RemoveOutputSetting(FName OutputName)
-{
-	Modify();
-
-	FTG_OutputExpressionInfo Info = *GetOutputExpressionInfo(OutputName);
-	OutputExpressionInfos.Remove(Info);
-}
-
-void UTG_OutputSettingsSet::RenameOutputSetting(FName OldName, FName NewName)
-{
-	Modify();
-
-	auto* Info = GetOutputExpressionInfo(OldName);
-	Info->GetOutputSettings()->OutputName = NewName;
-	Info->OutputName = NewName;
-}
-
-void UTG_OutputSettingsSet::UpdateTitle(UTG_Node* Output)
-{
-	if (Output)
-	{
-		for (auto& item : OutputExpressionInfos)
-		{
-			if (item.OutputPtr == Output->GetExpression() && Output->GetNodeName() != item.OutputName)
-			{
-				Modify();
-				item.OutputName = Output->GetNodeName();
-				item.GetOutputSettings()->OutputName = item.OutputName;
-			}
-		}
-	}
-}
-
-EResolution UTG_OutputSettingsSet::GetMaxWidth()
-{
-	EResolution MaxWidth = EResolution::Auto;
-	for (auto& Item : OutputExpressionInfos)
-	{
-		auto ItemWidth = Item.OutputPtr->OutputSettings.Width;
-		MaxWidth = static_cast<EResolution>(FMath::Max(static_cast<int32>(MaxWidth), static_cast<int32>(ItemWidth)));
-	}
-	return MaxWidth;
-}
-
-EResolution UTG_OutputSettingsSet::GetMaxHeight()
-{
-	EResolution MaxHeight = EResolution::Auto;
-	for (auto& Item : OutputExpressionInfos)
-	{
-		auto ItemHeight = Item.OutputPtr->OutputSettings.Height;
-		MaxHeight = static_cast<EResolution>(FMath::Max(static_cast<int32>(MaxHeight), static_cast<int32>(ItemHeight)));
-	}
-	return MaxHeight;
-}
-
-int32 UTG_OutputSettingsSet::GetMaxBufferChannels()
-{
-	uint32 MaxBufferChannels = 0;
-	for (auto& Item : OutputExpressionInfos)
-	{
-		uint32 Channels = 0;
-		BufferFormat Format = BufferFormat::Auto;
-		TextureHelper::GetBufferFormatAndChannelsFromTGTextureFormat(Item.OutputPtr->OutputSettings.TextureFormat, Format, Channels);
-		MaxBufferChannels = FMath::Max(MaxBufferChannels, Channels);
-	}
-	return MaxBufferChannels;
-}
-
-BufferFormat UTG_OutputSettingsSet::GetMaxBufferFormat()
-{
-	BufferFormat MaxBufferFormat = BufferFormat::Auto;
-	for (auto& Item : OutputExpressionInfos)
-	{
-		uint32 Channels = 0;
-		BufferFormat Format = BufferFormat::Auto;
-		TextureHelper::GetBufferFormatAndChannelsFromTGTextureFormat(Item.OutputPtr->OutputSettings.TextureFormat, Format, Channels);
-		MaxBufferFormat = static_cast<BufferFormat>(FMath::Max(static_cast<int32>(MaxBufferFormat), static_cast<int32>(Format)));
-	}
-	return MaxBufferFormat;
+	FolderPath = FName(DefaultDirectory);
+	TextureFormat = ETG_TextureFormat::BGRA8;
 }
diff --git a/Engine/Plugins/Experimental/TextureGraph/Source/TextureGraph/Private/TextureGraph.cpp b/Engine/Plugins/Experimental/TextureGraph/Source/TextureGraph/Private/TextureGraph.cpp
--- a/Engine/Plugins/Experimental/TextureGraph/Source/TextureGraph/Private/TextureGraph.cpp
+++ b/Engine/Plugins/Experimental/TextureGraph/Source/TextureGraph/Private/TextureGraph.cpp
@@ -93,10 +93,10 @@ void UTextureGraph::Construct(FString InName)
 
 	const UTG_Node* OutputNode = TextureGraph->CreateExpressionNode(UTG_Expression_Output::StaticClass());
 
-	Settings->GetViewportSettings().InitDefaultSettings(OutputNode->GetNodeName());
+	UTG_Expression_Output* OutputExpression = Cast<UTG_Expression_Output>(OutputNode->GetExpression());
+	OutputExpression->InitializeOutputSettings();
 
-	OutputSettingsSet = NewObject<UTG_OutputSettingsSet>(this, NAME_None, RF_Transactional);
-	OutputSettingsSet->AddOutputSetting(this, OutputNode->GetNodeName(), Cast<UTG_Expression_Output>(OutputNode->GetExpression()));
+	Settings->GetViewportSettings().InitDefaultSettings(OutputNode->GetNodeName());
 }
 
 
@@ -121,22 +121,6 @@ void UTextureGraph::PostLoad()
 	Super::PostLoad();
 	bInvalidateTextures = false;
 
-	// Output Settings Set must exist in case it wasn't saved properly
-	if (!OutputSettingsSet)
-	{
-		OutputSettingsSet = NewObject<UTG_OutputSettingsSet>(this, NAME_None, RF_Transactional);
-		TextureGraph->ForEachParams([&](const UTG_Pin* ParamPin, uint32 Index)
-		{
-			if (ParamPin->IsOutput())
-			{
-				UTG_Node* OutputNode = ParamPin->GetNodePtr();
-				UTG_Expression_Output* OutputExpression = Cast<UTG_Expression_Output>(OutputNode->GetExpression());
-				if (OutputExpression)
-					OutputSettingsSet->AddOutputSetting(this, OutputNode->GetNodeName(), Cast<UTG_Expression_Output>(OutputNode->GetExpression()));
-			}
-		});
-	}
-
 	// Settings must exist in case it wasn't saved properly
 	if (!Settings)
 	{
@@ -225,6 +209,66 @@ void UTextureGraph::InvalidateAll()
 	TextureGraphEngine::GetMixManager()->InvalidateMix(this, Details);
 }
 
+void UTextureGraph::UpdateGlobalTGSettings()
+{
+	check(GetSettings());
+	GetSettings()->SetWidth(GetMaxWidth());
+	GetSettings()->SetHeight(GetMaxHeight());
+
+	int32 Channels = GetMaxBufferChannels();
+	BufferFormat Format = GetMaxBufferFormat();
+	const ETG_TextureFormat TextureFormat = TextureHelper::GetTGTextureFormatFromChannelsAndFormat(Channels, Format);
+	GetSettings()->SetTextureFormat(TextureFormat);
+}
+
+EResolution UTextureGraph::GetMaxWidth()
+{
+	EResolution MaxWidth = EResolution::Auto;
+	Graph()->ForEachOutputSettings( [MaxWidth](const FTG_OutputSettings& OutSettings) mutable
+	{
+		EResolution ItemWidth = OutSettings.Width;
+		MaxWidth = static_cast<EResolution>(FMath::Max(static_cast<int32>(MaxWidth), static_cast<int32>(ItemWidth)));
+	});
+	return MaxWidth;
+}
+
+EResolution UTextureGraph::GetMaxHeight()
+{
+	EResolution MaxHeight = EResolution::Auto;
+	Graph()->ForEachOutputSettings( [MaxHeight](const FTG_OutputSettings& OutSettings) mutable
+	{
+		EResolution ItemHeight = OutSettings.Height;
+		MaxHeight = static_cast<EResolution>(FMath::Max(static_cast<int32>(MaxHeight), static_cast<int32>(ItemHeight)));
+	});
+	return MaxHeight;
+}
+
+int32 UTextureGraph::GetMaxBufferChannels()
+{
+	uint32 MaxBufferChannels = 0;
+	Graph()->ForEachOutputSettings( [MaxBufferChannels](const FTG_OutputSettings& OutSettings) mutable
+	{
+		uint32 Channels = 0;
+		BufferFormat Format = BufferFormat::Auto;
+		TextureHelper::GetBufferFormatAndChannelsFromTGTextureFormat(OutSettings.TextureFormat, Format, Channels);
+		MaxBufferChannels = FMath::Max(MaxBufferChannels, Channels);
+	});
+	return MaxBufferChannels;
+}
+
+BufferFormat UTextureGraph::GetMaxBufferFormat()
+{
+	BufferFormat MaxBufferFormat = BufferFormat::Auto;
+	Graph()->ForEachOutputSettings( [MaxBufferFormat](const FTG_OutputSettings& OutSettings) mutable
+	{
+		uint32 Channels = 0;
+		BufferFormat Format = BufferFormat::Auto;
+		TextureHelper::GetBufferFormatAndChannelsFromTGTextureFormat(OutSettings.TextureFormat, Format, Channels);
+		MaxBufferFormat = static_cast<BufferFormat>(FMath::Max(static_cast<int32>(MaxBufferFormat), static_cast<int32>(Format)));
+	});
+	return MaxBufferFormat;
+}
+
 void UTextureGraph::Log() const
 {
 	TextureGraph->Log();
diff --git a/Engine/Plugins/Experimental/TextureGraph/Source/TextureGraph/Public/Expressions/Output/TG_Expression_Output.h b/Engine/Plugins/Experimental/TextureGraph/Source/TextureGraph/Public/Expressions/Output/TG_Expression_Output.h
--- a/Engine/Plugins/Experimental/TextureGraph/Source/TextureGraph/Public/Expressions/Output/TG_Expression_Output.h
+++ b/Engine/Plugins/Experimental/TextureGraph/Source/TextureGraph/Public/Expressions/Output/TG_Expression_Output.h
@@ -37,4 +37,8 @@ public:
 	virtual FText GetTooltipText() const override { return FText::FromString(TEXT("Uses the node input as a graph output. It can be exported as a texture asset and is automatically exposed as a graph output parameter. ")); } 
 
 	void UpdateBufferDescriptorValues();
+	void InitializeOutputSettings();
+
+	//Sets the Export bool in OutputSettings pin
+	void SetExport(bool bExport);
 };
diff --git a/Engine/Plugins/Experimental/TextureGraph/Source/TextureGraph/Public/TG_Graph.h b/Engine/Plugins/Experimental/TextureGraph/Source/TextureGraph/Public/TG_Graph.h
--- a/Engine/Plugins/Experimental/TextureGraph/Source/TextureGraph/Public/TG_Graph.h
+++ b/Engine/Plugins/Experimental/TextureGraph/Source/TextureGraph/Public/TG_Graph.h
@@ -220,8 +220,11 @@ public:
 	// Iterate through all the VALID Param pins
 	void ForEachParams(std::function<void(const UTG_Pin* /*pin*/, uint32 /*index*/)> visitor) const;
 
-    // Iterate through all the VALID edges
-    void ForEachEdges(std::function<void(const UTG_Pin* /*pinFrom*/, const UTG_Pin* /*pinTo*/)> visitor) const;
+	// Iterate through all the VALID edges
+	void ForEachEdges(std::function<void(const UTG_Pin* /*pinFrom*/, const UTG_Pin* /*pinTo*/)> visitor) const;
+
+	// Iterate through all the output settings
+	void ForEachOutputSettings( std::function<void(const FTG_OutputSettings& /*settings*/)> visitor);
 
 	//////////////////////////////////////////////////////////////////////////
 	// Accessors for output param values after evaluation
diff --git a/Engine/Plugins/Experimental/TextureGraph/Source/TextureGraph/Public/TG_OutputSettings.h b/Engine/Plugins/Experimental/TextureGraph/Source/TextureGraph/Public/TG_OutputSettings.h
--- a/Engine/Plugins/Experimental/TextureGraph/Source/TextureGraph/Public/TG_OutputSettings.h
+++ b/Engine/Plugins/Experimental/TextureGraph/Source/TextureGraph/Public/TG_OutputSettings.h
@@ -19,7 +19,7 @@ struct TEXTUREGRAPH_API FTG_OutputSettings
 	UPROPERTY(BlueprintReadWrite,EditAnywhere, Category = "Basic", DisplayName = "File Name", Meta = (NoResetToDefault))
 		FName BaseName;
 
-	UPROPERTY(EditAnywhere, Category = "Basic", Meta = (NoResetToDefault))
+	UPROPERTY()
 		FName OutputName;
 
 	// Export path for the textured asset.
@@ -54,6 +54,9 @@ struct TEXTUREGRAPH_API FTG_OutputSettings
 	UPROPERTY(BlueprintReadWrite, EditAnywhere, Category = "Advanced", DisplayName = "sRGB", Meta = (NoResetToDefault, EditCondition = "TexturePresetType == ETG_TexturePresetType::None"))
 		bool bSRGB = false;
 
+	UPROPERTY()
+	bool bExport = true;
+
 	FString GetFullOutputName() { return  FString::Format(TEXT("{0}"), { BaseName.ToString()});}
 
 	bool operator==(const FTG_OutputSettings& Other) const
@@ -61,6 +64,8 @@ struct TEXTUREGRAPH_API FTG_OutputSettings
 		return OutputName == Other.OutputName && BaseName == Other.BaseName;
 	}
 
+	void Initialize(FString PathName, FName InName = "Output");
+
 	void InitFromString(const FString& StrVal)
 	{
 		FOutputDeviceNull NullOut;
@@ -79,58 +84,3 @@ struct TEXTUREGRAPH_API FTG_OutputSettings
 
 	void OnSetTexturePresetType(ETG_TexturePresetType Type);
 };
-
-USTRUCT()
-struct TEXTUREGRAPH_API FTG_OutputExpressionInfo
-{
-	GENERATED_BODY()
-
-public:
-	UPROPERTY()
-	FName OutputName;
-
-	UPROPERTY()
-	TObjectPtr<UTG_Expression_Output> OutputPtr;
-
-	UPROPERTY()
-	bool bExport = true;
-
-	FTG_OutputSettings* GetOutputSettings();
-
-	bool operator==(const FTG_OutputExpressionInfo& Other) const
-	{
-		return OutputName == Other.OutputName;
-	}
-};
-
-UCLASS()
-class TEXTUREGRAPH_API UTG_OutputSettingsSet : public UObject
-{
-	GENERATED_BODY()
-
-public:
-	UPROPERTY(EditAnywhere, EditFixedSize, DisplayName = "Settings" , Category = NoCategory, meta = (HideItemCount, NoResetToDefault, EditFixedOrder, ShowOnlyInnerProperties, FullyExpand, TitleProperty = "OutputName"))
-	TArray<FTG_OutputExpressionInfo>	OutputExpressionInfos;
-
-	void InitOutputSettings();
-	
-	FTG_OutputExpressionInfo* GetOutputExpressionInfo(FName OutputName);
-
-	FTG_OutputSettings* GetOutputSetting(FName OutputName);
-
-	void AddOutputSetting(UTextureGraph* InTextureGraph, FName OutputName,UTG_Expression_Output* Output);
-
-	void RemoveOutputSetting(FName OutputName);
-
-	void RenameOutputSetting(FName OldName,FName NewName);
-
-	void UpdateTitle(UTG_Node* Output);
-
-	EResolution GetMaxWidth();
-
-	EResolution GetMaxHeight();
-
-	int32 GetMaxBufferChannels();
-
-	BufferFormat GetMaxBufferFormat();
-};
diff --git a/Engine/Plugins/Experimental/TextureGraph/Source/TextureGraph/Public/TextureGraph.h b/Engine/Plugins/Experimental/TextureGraph/Source/TextureGraph/Public/TextureGraph.h
--- a/Engine/Plugins/Experimental/TextureGraph/Source/TextureGraph/Public/TextureGraph.h
+++ b/Engine/Plugins/Experimental/TextureGraph/Source/TextureGraph/Public/TextureGraph.h
@@ -18,9 +18,6 @@ protected:
 	UPROPERTY(EditAnywhere, Category="TG_GraphParams")
 	TObjectPtr<class UTG_Graph> TextureGraph;
 
-	UPROPERTY(EditAnywhere, Category = NoCategory)
-	TObjectPtr<UTG_OutputSettingsSet> OutputSettingsSet;
-
 	bool CheckCyclicDependency(const UTextureGraph* InTextureGraph) const;
 	void GatherAllDependentGraphs(TArray<UTextureGraph*>& DependentGraphs) const;
 	
@@ -40,16 +37,25 @@ public:
 	// Override PreSave method of UObject
 	virtual void PreSave(FObjectPreSaveContext SaveContext) override;
 
-	virtual const UTG_Graph*			Graph() const { return TextureGraph; }
-	virtual UTG_Graph*					Graph() { return TextureGraph; }
-	UTG_OutputSettingsSet*				GetOutputSettingsSet() { return OutputSettingsSet; }
+	virtual const UTG_Graph* Graph() const { return TextureGraph; }
+	virtual UTG_Graph* Graph() { return TextureGraph; }
+
+	void InvalidateAll() override;
+	void Update(MixUpdateCyclePtr InCycle) override;
+
+	void PostMeshLoad() override;
+
+	void TriggerUpdate(bool Tweaking);
+
+	void UpdateGlobalTGSettings();
+
+	EResolution GetMaxWidth();
 
-	void								InvalidateAll() override;
-	void								Update(MixUpdateCyclePtr InCycle) override;
+	EResolution GetMaxHeight();
 
-	void								PostMeshLoad() override;
+	int32 GetMaxBufferChannels();
 
-	void								TriggerUpdate(bool Tweaking);
+	BufferFormat GetMaxBufferFormat();
 
-	void								Log() const;
+	void Log() const;
 };
diff --git a/Engine/Plugins/Experimental/TextureGraph/Source/TextureGraphEditor/Private/Customizations/TG_OutputSettingsCustomization.h b/Engine/Plugins/Experimental/TextureGraph/Source/TextureGraphEditor/Private/Customizations/TG_OutputSettingsCustomization.h
--- a/Engine/Plugins/Experimental/TextureGraph/Source/TextureGraphEditor/Private/Customizations/TG_OutputSettingsCustomization.h
+++ b/Engine/Plugins/Experimental/TextureGraph/Source/TextureGraphEditor/Private/Customizations/TG_OutputSettingsCustomization.h
@@ -214,32 +214,4 @@ public:
 	}
 };
 
-class FTG_OutputExpressionInfoCustomization : public IPropertyTypeCustomization
-{
-public:
-	static TSharedRef<IPropertyTypeCustomization> Create()
-	{
-		return MakeShareable(new FTG_OutputExpressionInfoCustomization);
-	}
-
-	virtual void CustomizeHeader(TSharedRef<IPropertyHandle> PropertyHandle, FDetailWidgetRow& HeaderRow, IPropertyTypeCustomizationUtils& CustomizationUtils) override
-	{
-		
-	}
-
-	virtual void CustomizeChildren(TSharedRef<IPropertyHandle> PropertyHandle, IDetailChildrenBuilder& ChildBuilder, IPropertyTypeCustomizationUtils& CustomizationUtils) override
-	{
-		FPropertyEditorModule& PropertyModule = FModuleManager::LoadModuleChecked<FPropertyEditorModule>("PropertyEditor");
-
-		TArray<void*> RawData;
-		PropertyHandle->AccessRawData(RawData);
-
-		FTG_OutputExpressionInfo* OutputExpressionInfo = reinterpret_cast<FTG_OutputExpressionInfo*>(RawData[0]);
-		TArray<UObject*> ExternalObjects;
-		ExternalObjects.Add(OutputExpressionInfo->OutputPtr);
-
-		ChildBuilder.AddExternalObjectProperty(ExternalObjects, FName("OutputSettings"), FAddPropertyParams().HideRootObjectNode(true));
-	}
-};
-
 #undef LOCTEXT_NAMESPACE
diff --git a/Engine/Plugins/Experimental/TextureGraph/Source/TextureGraphEditor/Private/STG_OutputSelectionDlg.cpp b/Engine/Plugins/Experimental/TextureGraph/Source/TextureGraphEditor/Private/STG_OutputSelectionDlg.cpp
--- a/Engine/Plugins/Experimental/TextureGraph/Source/TextureGraphEditor/Private/STG_OutputSelectionDlg.cpp
+++ b/Engine/Plugins/Experimental/TextureGraph/Source/TextureGraphEditor/Private/STG_OutputSelectionDlg.cpp
@@ -16,7 +16,6 @@
 #define LOCTEXT_NAMESPACE "STG_OutputSelectionDlg"
 void STG_OutputSelectionDlg::Construct(const FArguments& InArgs)
 {
-	OutputSettingsSet = InArgs._OutputSettingsSet;
 	EdGraph = InArgs._EdGraph;
 	SWindow::Construct(SWindow::FArguments()
 		.Title(InArgs._Title)
@@ -70,71 +69,78 @@ void STG_OutputSelectionDlg::Construct(const FArguments& InArgs)
 void STG_OutputSelectionDlg::AddExportItems()
 {
 	ScrollBox->ClearChildren();
-	for (const auto& Info : OutputSettingsSet->OutputExpressionInfos)
+	UTG_EdGraph* TGEdGraph = Cast<UTG_EdGraph>(EdGraph);
+	TGEdGraph->TextureGraph->Graph()->ForEachNodes([this](const UTG_Node* Node, uint32 Index)
 	{
-		auto OutputSetting = Info.OutputPtr->OutputSettings;
-		auto Node = Cast<UTG_Node>(Info.OutputPtr->GetOuter());
-		auto OutPinIds = Node->GetOutputPinIds();
-		for (auto Id : OutPinIds)
+		UTG_Expression_Output* TargetExpression = Cast<UTG_Expression_Output>(Node->GetExpression());
+		if (TargetExpression)
 		{
-			//This is a work around for checking the type of the output
-			//Probably we need to have a better solution for checking output type
-			auto Pin = Node->GetGraph()->GetPin(Id);
-			FTG_Variant Variant;
-			Pin->GetValue(Variant);
-			TSharedPtr<SWidget> ThumbnailWidget;
-			if (Variant.IsTexture())
+			FTG_OutputSettings& OutputSetting = TargetExpression->OutputSettings;
+			auto OutPinIds = Node->GetOutputPinIds();
+			for (auto Id : OutPinIds)
 			{
-				UTG_EdGraphNode* EdNode = EdGraph->GetViewModelNode(Node->GetId());
-				TiledBlobPtr ThumbBlob = EdGraph->GetCachedThumbBlob(Id);
+				//This is a work around for checking the type of the output
+				//Probably we need to have a better solution for checking output type
+				auto Pin = Node->GetGraph()->GetPin(Id);
 
-				if (!ThumbBlob)
+				FName OutputName = Pin->GetAliasName();
+
+				FTG_Variant Variant;
+				Pin->GetValue(Variant);
+				TSharedPtr<SWidget> ThumbnailWidget;
+				if (Variant.IsTexture())
 				{
-					ThumbBlob = TextureHelper::GetBlack();
-				}
+					TiledBlobPtr ThumbBlob = EdGraph->GetCachedThumbBlob(Id);
 
-				TSharedPtr<STG_NodeThumbnail> NodeThumbnail = SNew(STG_NodeThumbnail);
-				ThumbBlob->OnFinalise()
-					.then([ThumbBlob, NodeThumbnail]
+					if (!ThumbBlob)
 					{
-						// NOTE: If later, "this" were to be captured here, we should check DoesSharedInstanceExist()
-						// as there is a chance this might be invoked when the slate widgets have already been destroyed
-						if (NodeThumbnail.IsValid())
+						ThumbBlob = TextureHelper::GetBlack();
+					}
+
+					TSharedPtr<STG_NodeThumbnail> NodeThumbnail = SNew(STG_NodeThumbnail);
+					ThumbBlob->OnFinalise()
+					.then([ThumbBlob, NodeThumbnail]
 						{
-							NodeThumbnail->UpdateBlob(ThumbBlob);
-						}
-					});
+							// NOTE: If later, "this" were to be captured here, we should check DoesSharedInstanceExist()
+							// as there is a chance this might be invoked when the slate widgets have already been destroyed
+							if (NodeThumbnail.IsValid())
+							{
+								NodeThumbnail->UpdateBlob(ThumbBlob);
+							}
+						});
 
-				ThumbnailWidget = NodeThumbnail;
-			}
-			// else if color do something
-			else if (Variant.IsColor())
-			{
-				FLinearColor ColorValue;
-				Pin->GetValue(ColorValue);
-				ThumbnailWidget = SNew(SColorBlock)
-									.Color(ColorValue); 
-			}
-			else
-			{
-				continue;
+					ThumbnailWidget = NodeThumbnail;
+				}
+				// else if color do something
+				else if (Variant.IsColor())
+				{
+					FLinearColor ColorValue;
+					Pin->GetValue(ColorValue);
+					ThumbnailWidget = SNew(SColorBlock)
+					.Color(ColorValue);
+				}
+				else
+				{
+					TargetExpression->SetExport(false);
+					continue;
+				}
+				ScrollBox->AddSlot()
+				.Padding(5)
+				[
+					SNew(STG_OutputSelector)
+					.Name(FText::FromString(OutputName.ToString()))
+					.ThumbnailWidget(ThumbnailWidget)
+					.OnOutputSelectionChanged(this, &STG_OutputSelectionDlg::OnOutputSelectionChanged)
+					.bIsSelected(OutputSetting.bExport)
+				];
+				ScrollBox->AddSlot()
+				[
+					SNew(SSeparator)
+					.Thickness(1)
+				];
 			}
-			ScrollBox->AddSlot()
-			.Padding(5)
-			[
-				SNew(STG_OutputSelector)
-				.Name(FText::FromString(Info.OutputName.ToString()))
-				.ThumbnailWidget(ThumbnailWidget)
-				.OnOutputSelectionChanged(this,&STG_OutputSelectionDlg::OnOutputSelectionChanged)
-				.bIsSelected(Info.bExport)
-			];
-			ScrollBox->AddSlot()
-			[
-				SNew(SSeparator)
-				.Thickness(1)
-			];
 		}
-	}
+	});
 }
 FReply STG_OutputSelectionDlg::OnButtonClick(EAppReturnType::Type ButtonID)
 {
@@ -158,6 +164,17 @@ EAppReturnType::Type STG_OutputSelectionDlg::ShowModal()
 }
 void STG_OutputSelectionDlg::OnOutputSelectionChanged(const FString ItemName, ECheckBoxState NewState)
 {
-	OutputSettingsSet->GetOutputExpressionInfo(*ItemName)->bExport = NewState == ECheckBoxState::Checked;
+	UTG_EdGraph* TGEdGraph = Cast<UTG_EdGraph>(EdGraph);
+	TGEdGraph->TextureGraph->Graph()->ForEachNodes([=](const UTG_Node* Node, uint32 Index)
+	{
+		UTG_Expression_Output* TargetExpression = Cast<UTG_Expression_Output>(Node->GetExpression());
+		if (TargetExpression)
+		{
+			if (TargetExpression->OutputSettings.OutputName == ItemName)
+			{
+				TargetExpression->SetExport( NewState == ECheckBoxState::Checked);
+			}
+		}
+	});
 }
 #undef LOCTEXT_NAMESPACE
\ No newline at end of file
diff --git a/Engine/Plugins/Experimental/TextureGraph/Source/TextureGraphEditor/Private/STG_OutputSelectionDlg.h b/Engine/Plugins/Experimental/TextureGraph/Source/TextureGraphEditor/Private/STG_OutputSelectionDlg.h
--- a/Engine/Plugins/Experimental/TextureGraph/Source/TextureGraphEditor/Private/STG_OutputSelectionDlg.h
+++ b/Engine/Plugins/Experimental/TextureGraph/Source/TextureGraphEditor/Private/STG_OutputSelectionDlg.h
@@ -16,7 +16,6 @@ class STG_OutputSelectionDlg : public SWindow
 	SLATE_ARGUMENT(FText, Title)
 	SLATE_ARGUMENT(TObjectPtr<UTG_EdGraph>, EdGraph)
 	SLATE_ARGUMENT(TArray<FText>, ExportItems)
-	SLATE_ARGUMENT(TObjectPtr<UTG_OutputSettingsSet>, OutputSettingsSet)
 	SLATE_END_ARGS()
 
 public:
@@ -31,7 +30,6 @@ protected:
 
 private:
 	TSharedPtr<SScrollBox> ScrollBox;
-	TObjectPtr<UTG_OutputSettingsSet> OutputSettingsSet;
 	TObjectPtr<UTG_EdGraph> EdGraph;
 
 	void AddExportItems();
diff --git a/Engine/Plugins/Experimental/TextureGraph/Source/TextureGraphEditor/Private/TG_Editor.cpp b/Engine/Plugins/Experimental/TextureGraph/Source/TextureGraphEditor/Private/TG_Editor.cpp
--- a/Engine/Plugins/Experimental/TextureGraph/Source/TextureGraphEditor/Private/TG_Editor.cpp
+++ b/Engine/Plugins/Experimental/TextureGraph/Source/TextureGraphEditor/Private/TG_Editor.cpp
@@ -1125,8 +1125,6 @@ TSharedRef<SDockTab> FTG_Editor::SpawnTab_Output(const FSpawnTabArgs& Args)
 				]
 		];
 
-	GetOutputView()->SetObject(EditedTextureGraph->GetOutputSettingsSet(), true);
-
 	return OutputTab.ToSharedRef();
 }
 
@@ -1140,8 +1138,7 @@ FReply FTG_Editor::OnExportClick()
 	TSharedPtr<STG_OutputSelectionDlg> ExportSelectionWidget =
 		SNew(STG_OutputSelectionDlg)
 		.Title(LOCTEXT("Export Selection", "Choose outputs to export"))
-		.EdGraph(TG_EdGraph)
-		.OutputSettingsSet(EditedTextureGraph->GetOutputSettingsSet());
+		.EdGraph(TG_EdGraph);
 
 	if (ExportSelectionWidget->ShowModal() == EAppReturnType::Ok)
 	{
@@ -1803,14 +1800,7 @@ void FTG_Editor::OnGraphChanged(UTG_Graph* InGraph, UTG_Node* InNode, bool Tweak
 void FTG_Editor::UpdateMixSettings()
 {
 	check(EditedTextureGraph);
-	check(EditedTextureGraph->GetSettings());
-	EditedTextureGraph->GetSettings()->SetWidth(EditedTextureGraph->GetOutputSettingsSet()->GetMaxWidth());
-	EditedTextureGraph->GetSettings()->SetHeight(EditedTextureGraph->GetOutputSettingsSet()->GetMaxHeight());
-
-	int32 Channels = EditedTextureGraph->GetOutputSettingsSet()->GetMaxBufferChannels();
-	BufferFormat Format = EditedTextureGraph->GetOutputSettingsSet()->GetMaxBufferFormat();
-	const ETG_TextureFormat TextureFormat = TextureHelper::GetTGTextureFormatFromChannelsAndFormat(Channels, Format);
-	EditedTextureGraph->GetSettings()->SetTextureFormat(TextureFormat);
+	EditedTextureGraph->UpdateGlobalTGSettings();
 }
 
 void FTG_Editor::OnNodeAdded(UTG_Node* InNode)
@@ -1819,7 +1809,7 @@ void FTG_Editor::OnNodeAdded(UTG_Node* InNode)
 	auto TargetExpression = Cast<UTG_Expression_Output>(InNode->GetExpression());
 	if (TargetExpression != nullptr)
 	{
-		EditedTextureGraph->GetOutputSettingsSet()->AddOutputSetting(EditedTextureGraph, InNode->GetNodeName(), TargetExpression);
+		TargetExpression->InitializeOutputSettings();
 	}
 }
 
@@ -1828,8 +1818,6 @@ void FTG_Editor::OnNodeRemoved(UTG_Node* InNode, FName Title)
 	check(EditedTextureGraph);
 	if (IsOutputNode(InNode))
 	{
-		EditedTextureGraph->GetOutputSettingsSet()->RemoveOutputSetting(Title);
-
 		// Remove output node texture.
 		for (int i = 0; i < EditedTextureGraph->GetSettings()->NumTargets(); i++)
 		{
@@ -1869,8 +1857,6 @@ void FTG_Editor::OnNodeRenamed(UTG_Node* InNode, FName OldName)
 		EditedTextureGraph->GetSettings()->GetViewportSettings().OnTargetRename(OldName, InNode->GetNodeName());
 
 		// update output 
-		EditedTextureGraph->GetOutputSettingsSet()->UpdateTitle(InNode);
-		GetOutputView()->SetObject(EditedTextureGraph->GetOutputSettingsSet(), true);
 	}
 
 	// check on the parameters view, it could have changed
diff --git a/Engine/Plugins/Experimental/TextureGraph/Source/TextureGraphEditor/Private/TextureGraphEditorModule.cpp b/Engine/Plugins/Experimental/TextureGraph/Source/TextureGraphEditor/Private/TextureGraphEditorModule.cpp
--- a/Engine/Plugins/Experimental/TextureGraph/Source/TextureGraphEditor/Private/TextureGraphEditorModule.cpp
+++ b/Engine/Plugins/Experimental/TextureGraph/Source/TextureGraphEditor/Private/TextureGraphEditorModule.cpp
@@ -36,7 +36,6 @@ void FTextureGraphEditorModule::StartupModule()
 	PropertyEditorModule.RegisterCustomPropertyTypeLayout("FloatProperty", FOnGetPropertyTypeCustomizationInstance::CreateStatic(&FTG_ScalarCustomization::Create), ScalarIdentifier);
 	PropertyEditorModule.RegisterCustomPropertyTypeLayout("MaterialMappingInfo", FOnGetPropertyTypeCustomizationInstance::CreateStatic(&FTG_MaterialMappingInfoCustomization::Create));
 	PropertyEditorModule.RegisterCustomPropertyTypeLayout("OutputSettings", FOnGetPropertyTypeCustomizationInstance::CreateStatic(&FTG_OutputSettingsCustomization::Create));
-	PropertyEditorModule.RegisterCustomPropertyTypeLayout("OutputExpressionInfo", FOnGetPropertyTypeCustomizationInstance::CreateStatic(&FTG_OutputExpressionInfoCustomization::Create));
 	PropertyEditorModule.RegisterCustomPropertyTypeLayout("TG_LevelsSettings", FOnGetPropertyTypeCustomizationInstance::CreateStatic(&FTG_LevelsSettingsCustomization::Create));
 
 	// Register slate style overrides
@@ -66,7 +65,6 @@ void FTextureGraphEditorModule::ShutdownModule()
 	PropertyEditorModule.UnregisterCustomPropertyTypeLayout("FloatProperty", ScalarIdentifier);
 	PropertyEditorModule.UnregisterCustomPropertyTypeLayout("MaterialMappingInfo");
 	PropertyEditorModule.UnregisterCustomPropertyTypeLayout("OutputSettings");
-	PropertyEditorModule.UnregisterCustomPropertyTypeLayout("OutputExpressionInfo");
 	PropertyEditorModule.UnregisterCustomPropertyTypeLayout("TG_LevelsSettings");
 
 	FEdGraphUtilities::UnregisterVisualPinFactory(GraphPanelPinFactory);
