Fixed a crash when auto-reloading a TextureGraph asset when UE is run.

diff --git a/Engine/Plugins/Experimental/TextureGraph/Source/TextureGraphEngine/3D/RenderMesh_Editor.cpp b/Engine/Plugins/Experimental/TextureGraph/Source/TextureGraphEngine/3D/RenderMesh_Editor.cpp
--- a/Engine/Plugins/Experimental/TextureGraph/Source/TextureGraphEngine/3D/RenderMesh_Editor.cpp
+++ b/Engine/Plugins/Experimental/TextureGraph/Source/TextureGraphEngine/3D/RenderMesh_Editor.cpp
@@ -63,9 +63,11 @@ void RenderMesh_Editor::SetMaterial(UMaterialInterface* material)
 	for (int meshIndex = 0; meshIndex < _meshComponents.Num(); meshIndex++)
 	{
 		UStaticMeshComponent* meshComponent = _meshComponents[meshIndex];
-
-		MeshInfoPtr meshInfo = _meshes[meshIndex];
-		meshComponent->SetMaterial(meshInfo->GetMaterialIndex(), material);
+		if (meshComponent->IsValidLowLevel())
+		{
+			MeshInfoPtr meshInfo = _meshes[meshIndex];
+			meshComponent->SetMaterial(meshInfo->GetMaterialIndex(), material);
+		}
 	}
 }
 
