Fixed issue: The dynamic settings of the node breaks when a node has the input node connected to them

diff --git a/Engine/Plugins/Experimental/TextureGraph/Source/TextureGraph/Public/Expressions/Filter/TG_Expression_Blur.h b/Engine/Plugins/Experimental/TextureGraph/Source/TextureGraph/Public/Expressions/Filter/TG_Expression_Blur.h
--- a/Engine/Plugins/Experimental/TextureGraph/Source/TextureGraph/Public/Expressions/Filter/TG_Expression_Blur.h
+++ b/Engine/Plugins/Experimental/TextureGraph/Source/TextureGraph/Public/Expressions/Filter/TG_Expression_Blur.h
@@ -29,11 +29,11 @@ public:
 #endif
 
 	// Controls how much an texture is blurred
-	UPROPERTY(EditAnywhere, BlueprintReadWrite, Category = NoCategory, meta = (TGType = "TG_Setting", ClampMin = 0, ClampMax = 20, EditConditionHides))
+	UPROPERTY(EditAnywhere, BlueprintReadWrite, Category = NoCategory, meta = (TGType = "TG_Setting", ClampMin = 0, ClampMax = 100, EditConditionHides))
 	int32 Radius = 1;
 
 	// Adjusting the angle changes the direction at which the blur effect is applied." 
-	UPROPERTY(EditAnywhere, BlueprintReadWrite, Category = NoCategory, meta = (TGType = "TG_Setting", ClampMin = -180, ClampMax = 180, EditConditionHides))
+	UPROPERTY(EditAnywhere, BlueprintReadWrite, Category = NoCategory, meta = (TGType = "TG_Setting", ClampMin = -180, ClampMax = 180,EditConditionHides))
 	float Angle = 0.0f;
 
 	// Determines how intense the blurring effect is applied. Higher values result in stronger blur, while lower values produce a milder effect.
diff --git a/Engine/Plugins/Experimental/TextureGraph/Source/TextureGraphEditor/Private/EdGraph/TG_EdGraphNode.cpp b/Engine/Plugins/Experimental/TextureGraph/Source/TextureGraphEditor/Private/EdGraph/TG_EdGraphNode.cpp
--- a/Engine/Plugins/Experimental/TextureGraph/Source/TextureGraphEditor/Private/EdGraph/TG_EdGraphNode.cpp
+++ b/Engine/Plugins/Experimental/TextureGraph/Source/TextureGraphEditor/Private/EdGraph/TG_EdGraphNode.cpp
@@ -342,12 +342,13 @@ void UTG_EdGraphNode::ReconstructNode()
 		if (NewPin)
 		{
 			// And also check that the new pin at that name is connected to anything, if so grab the UI connections
-			UTG_Pin* TSPin = Node->GetPin(Node->GetPinId(OldPinName));
-			if (TSPin->IsConnected())
+			UTG_Pin* TGPin = Node->GetPin(Node->GetPinId(OldPinName));
+			if (TGPin->IsConnected())
+			{
 				(*NewPin)->MovePersistentDataFromOldPin(*OldPin);
+			}
 		}
 	}
-
 	
 	// Remove old pins
 	for (UEdGraphPin* OldPin : OldPins)
@@ -448,11 +449,15 @@ void UTG_EdGraphNode::UpdatePinVisibility(UEdGraphPin* Pin, UTG_Pin* TGPin) cons
 		// update Metadata with EditCondition to also update details view
 		FName NAME_EditCondition(TEXT("EditCondition"));
 		Property->SetMetaData(NAME_EditCondition, bCanEditChange ? TEXT("true") : TEXT("false"));
-
 	}
 
 	Pin->bHidden = !bCanEditChange && bEditConditionHides;
 	Pin->bDefaultValueIsReadOnly = bCanEditChange;
+	
+	if(Pin->bHidden) 
+	{
+		Pin->GetSchema()->BreakPinLinks(*Pin, false);
+	}
 }
 
 void UTG_EdGraphNode::UpdateInputPinsVisibility() const
