Texture Graph Bug Fixes
Removed inconsistencies from get and set output settings parameter node for blueprint.
Fixed drag not working in the list palette view

diff --git a/Engine/Plugins/Experimental/TextureGraph/Source/TextureGraph/Private/Blueprint/TG_BlueprintFunctionLibrary.cpp b/Engine/Plugins/Experimental/TextureGraph/Source/TextureGraph/Private/Blueprint/TG_BlueprintFunctionLibrary.cpp
--- a/Engine/Plugins/Experimental/TextureGraph/Source/TextureGraph/Private/Blueprint/TG_BlueprintFunctionLibrary.cpp
+++ b/Engine/Plugins/Experimental/TextureGraph/Source/TextureGraph/Private/Blueprint/TG_BlueprintFunctionLibrary.cpp
@@ -234,8 +234,8 @@ FLinearColor UTG_BlueprintFunctionLibrary::GetColorParameterValue(UObject* World
 }
 
 void UTG_BlueprintFunctionLibrary::SetSettingsParameterValue(UObject* WorldContextObject, UTextureGraph* InTextureGraph, FName ParameterName, int Width, int Height, 
-	FString FileName /*= "None"*/, FString FolderPath /*= "None"*/, ETG_TextureFormat Format /*= ETG_TextureFormat::BGRA8*/, ETG_TexturePresetType TextureType /*= ETG_TexturePresetType::None*/,
-	TextureCompressionSettings Compression /*= TextureCompressionSettings::TC_Default*/, TextureGroup LodGroup /*= TextureGroup::TEXTUREGROUP_World*/, bool bSRGB /*= false*/)
+	FName FileName /*= "None"*/, FName Path /*= "None"*/, ETG_TextureFormat Format /*= ETG_TextureFormat::BGRA8*/, ETG_TexturePresetType TextureType /*= ETG_TexturePresetType::None*/,
+	TextureGroup LODTextureGroup /*= TextureGroup::TEXTUREGROUP_World*/, TextureCompressionSettings Compression /*= TextureCompressionSettings::TC_Default*/, bool SRGB /*= false*/)
 {
 	if (UWorld* World = GEngine->GetWorldFromContextObject(WorldContextObject, EGetWorldErrorMode::LogAndReturnNull))
 	{
@@ -246,15 +246,11 @@ void UTG_BlueprintFunctionLibrary::SetSettingsParameterValue(UObject* WorldConte
 
 		const FString FunctionName = "SetSettingsParameterValue";
 
-		if (Width <= 0 || Height <= 0)
-		{
-			bSizeError = true;
-		}
-		if (FolderPath.IsEmpty() || FolderPath == "None")
+		if (Path.ToString().IsEmpty() || Path == "None")
 		{
 			bPathError = true;
 		}
-		if (FileName.IsEmpty() || FolderPath == "None")
+		if (FileName.ToString().IsEmpty() || FileName == "None")
 		{
 			bNameError = true;
 		}
@@ -269,7 +265,7 @@ void UTG_BlueprintFunctionLibrary::SetSettingsParameterValue(UObject* WorldConte
 				if (ExpressionPtr)
 				{
 					FTG_OutputSettings ParameterValue;
-					ParameterValue.Set(Width, Height, FileName, FolderPath, Format, TextureType, Compression, LodGroup, bSRGB);
+					ParameterValue.Set(Width, Height, FileName, Path, Format, TextureType, Compression, LODTextureGroup, SRGB);
 					ExpressionPtr->Settings = ParameterValue;
 					PinParam->EditSelfVar()->EditAs<FTG_OutputSettings>() = ParameterValue;
 					bFoundParameter = true;
@@ -282,11 +278,6 @@ void UTG_BlueprintFunctionLibrary::SetSettingsParameterValue(UObject* WorldConte
 			}
 		}
 
-		if (bSizeError)
-		{
-			AddError(InTextureGraph, FunctionName, "Invalid size try to set non zero size");
-		}
-
 		if (bPathError)
 		{
 			AddError(InTextureGraph, FunctionName, "Invalid path try to set a valid path , path connot be empty or none");
@@ -299,7 +290,7 @@ void UTG_BlueprintFunctionLibrary::SetSettingsParameterValue(UObject* WorldConte
 	}
 }
 
-FTG_OutputSettings UTG_BlueprintFunctionLibrary::GetSettingsParameterValue(UObject* WorldContextObject, UTextureGraph* InTextureGraph, FName ParameterName)
+FTG_OutputSettings UTG_BlueprintFunctionLibrary::GetSettingsParameterValue(UObject* WorldContextObject, UTextureGraph* InTextureGraph, FName ParameterName, int& Width, int& Height)
 {
 	FTG_OutputSettings ParameterValue;
 
@@ -312,6 +303,8 @@ FTG_OutputSettings UTG_BlueprintFunctionLibrary::GetSettingsParameterValue(UObje
 			if (PinParam)
 			{
 				ParameterValue = PinParam->GetSelfVar()->GetAs<FTG_OutputSettings>();
+				Width = (int)ParameterValue.Width;
+				Height = (int)ParameterValue.Height;
 				bFoundParameter = true;
 			}
 		}
diff --git a/Engine/Plugins/Experimental/TextureGraph/Source/TextureGraph/Private/TG_OutputSettings.cpp b/Engine/Plugins/Experimental/TextureGraph/Source/TextureGraph/Private/TG_OutputSettings.cpp
--- a/Engine/Plugins/Experimental/TextureGraph/Source/TextureGraph/Private/TG_OutputSettings.cpp
+++ b/Engine/Plugins/Experimental/TextureGraph/Source/TextureGraph/Private/TG_OutputSettings.cpp
@@ -7,11 +7,11 @@
 #include "TG_HelperFunctions.h"
 #include "Misc/Paths.h"
 
-void FTG_OutputSettings::Set(int InWidth, int InHeight, FString Name /*= "None"*/, FString Path /*= "None"*/, ETG_TextureFormat Format /*= ETG_TextureFormat::BGRA8*/, ETG_TexturePresetType InTextureType /*= ETG_TexturePresetType::None*/,
+void FTG_OutputSettings::Set(int InWidth, int InHeight, FName Name /*= "None"*/, FName Path /*= "None"*/, ETG_TextureFormat Format /*= ETG_TextureFormat::BGRA8*/, ETG_TexturePresetType InTextureType /*= ETG_TexturePresetType::None*/,
 	TextureCompressionSettings InCompression /*= TextureCompressionSettings::TC_Default*/, TextureGroup InLodGroup /*= TextureGroup::TEXTUREGROUP_World*/, bool InbSRGB /*= false*/)
 {
-	BaseName = *Name;
-	FolderPath = *Path;
+	BaseName = Name;
+	FolderPath = Path;
 	Width = (EResolution)InWidth;
 	Height = (EResolution)InHeight;
 	TextureFormat = Format;
diff --git a/Engine/Plugins/Experimental/TextureGraph/Source/TextureGraph/Public/Blueprint/TG_BlueprintFunctionLibrary.h b/Engine/Plugins/Experimental/TextureGraph/Source/TextureGraph/Public/Blueprint/TG_BlueprintFunctionLibrary.h
--- a/Engine/Plugins/Experimental/TextureGraph/Source/TextureGraph/Public/Blueprint/TG_BlueprintFunctionLibrary.h
+++ b/Engine/Plugins/Experimental/TextureGraph/Source/TextureGraph/Public/Blueprint/TG_BlueprintFunctionLibrary.h
@@ -51,13 +51,13 @@ public:
 
 	/** Sets a FTG_OutputSettings parameter value on the TextureGraph instance. Logs if ParameterName is invalid. */
 	UFUNCTION(BlueprintCallable, Category = "TextureGraph", meta = (Keywords = "SetSettingsParameterValue", WorldContext = "WorldContextObject"))
-	static void SetSettingsParameterValue(UObject* WorldContextObject, UTextureGraph* TextureScript, FName ParameterName, int Width,int Height, FString FileName = "None",
-		FString FolderPath = "None", ETG_TextureFormat Format = ETG_TextureFormat::BGRA8 , ETG_TexturePresetType TextureType = ETG_TexturePresetType::None,
-		TextureCompressionSettings Compression = TextureCompressionSettings::TC_Default, TextureGroup LodGroup = TextureGroup::TEXTUREGROUP_World, bool bSRGB = false);
+	static void SetSettingsParameterValue(UObject* WorldContextObject, UTextureGraph* TextureScript, FName ParameterName, int Width,int Height, FName FileName = "None",
+		FName Path = "None", ETG_TextureFormat Format = ETG_TextureFormat::BGRA8 , ETG_TexturePresetType TextureType = ETG_TexturePresetType::None,
+		TextureGroup LODTextureGroup = TextureGroup::TEXTUREGROUP_World, TextureCompressionSettings Compression = TextureCompressionSettings::TC_Default, bool SRGB = false);
 
-	/** Gets a FOutputSettings parameter value from the TextureGraph instance. Logs if ParameterName is invalid. */
+	/** Gets a FTG_OutputSettings parameter value from the TextureGraph instance. Logs if ParameterName is invalid. */
 	UFUNCTION(BlueprintCallable, Category = "TextureGraph", meta = (Keywords = "GetOutputSettingsParameterValue", WorldContext = "WorldContextObject"))
-	static FTG_OutputSettings GetSettingsParameterValue(UObject* WorldContextObject, UTextureGraph* TextureScript, FName ParameterName);
+	static FTG_OutputSettings GetSettingsParameterValue(UObject* WorldContextObject, UTextureGraph* TextureScript, FName ParameterName , int& Width, int& Height);
 
 	static void AddParamWarning(FName ParamName, UObject* ObjectPtr, FString FunctionName);
 	static void AddError(UObject* ObjectPtr, FString FunctionName, FString Error);
diff --git a/Engine/Plugins/Experimental/TextureGraph/Source/TextureGraph/Public/TG_OutputSettings.h b/Engine/Plugins/Experimental/TextureGraph/Source/TextureGraph/Public/TG_OutputSettings.h
--- a/Engine/Plugins/Experimental/TextureGraph/Source/TextureGraph/Public/TG_OutputSettings.h
+++ b/Engine/Plugins/Experimental/TextureGraph/Source/TextureGraph/Public/TG_OutputSettings.h
@@ -79,7 +79,7 @@ struct TEXTUREGRAPH_API FTG_OutputSettings
 		return ExportString;
 	}
 
-	void Set(int InWidth, int InHeight, FString Name = "None", FString Path = "None", ETG_TextureFormat Format = ETG_TextureFormat::BGRA8, ETG_TexturePresetType InTextureType = ETG_TexturePresetType::None,
+	void Set(int InWidth, int InHeight, FName Name = "None", FName Path = "None", ETG_TextureFormat Format = ETG_TextureFormat::BGRA8, ETG_TexturePresetType InTextureType = ETG_TexturePresetType::None,
 		TextureCompressionSettings InCompression = TextureCompressionSettings::TC_Default, TextureGroup InLodGroup = TextureGroup::TEXTUREGROUP_World, bool InbSRGB = false);
 
 	void OnSetTexturePresetType(ETG_TexturePresetType Type);
diff --git a/Engine/Plugins/Experimental/TextureGraph/Source/TextureGraphEditor/Private/Widgets/ActionMenuWidgets/STG_GraphActionMenu.cpp b/Engine/Plugins/Experimental/TextureGraph/Source/TextureGraphEditor/Private/Widgets/ActionMenuWidgets/STG_GraphActionMenu.cpp
--- a/Engine/Plugins/Experimental/TextureGraph/Source/TextureGraphEditor/Private/Widgets/ActionMenuWidgets/STG_GraphActionMenu.cpp
+++ b/Engine/Plugins/Experimental/TextureGraph/Source/TextureGraphEditor/Private/Widgets/ActionMenuWidgets/STG_GraphActionMenu.cpp
@@ -89,6 +89,12 @@ void STG_GraphActionMenu::CollectAllActions(FGraphActionListBuilderBase& OutAllA
 
 TSharedRef<SWidget> STG_GraphActionMenu::OnCreateWidgetForAction(FCreateWidgetForActionData* const InCreateData)
 {
+	if (!bSpwanOnSelect)
+	{
+		//Dont need the mouse down event when not spawning on selection
+		InCreateData->MouseButtonDownDelegate.Unbind();
+	}
+	
 	return	SNew(STG_PaletteItem, InCreateData);
 }
 
