Texture Histogram Compute Thread group size changed

diff --git a/Engine/Plugins/Experimental/TextureGraph/Shaders/Utils/Histogram_comp.usf b/Engine/Plugins/Experimental/TextureGraph/Shaders/Utils/Histogram_comp.usf
--- a/Engine/Plugins/Experimental/TextureGraph/Shaders/Utils/Histogram_comp.usf
+++ b/Engine/Plugins/Experimental/TextureGraph/Shaders/Utils/Histogram_comp.usf
@@ -11,8 +11,8 @@ RWTexture2D<float4> Result;
 RWStructuredBuffer<uint4> Histogram : register(u0);
 
 #ifndef THREADGROUPSIZE_X
-#define THREADGROUPSIZE_X 1
-#define THREADGROUPSIZE_Y 1
+#define THREADGROUPSIZE_X 16
+#define THREADGROUPSIZE_Y 16
 #define THREADGROUPSIZE_Z 1
 #endif
 
diff --git a/Engine/Plugins/Experimental/TextureGraph/Source/TextureGraphEngine/Transform/Utility/T_TextureHistogram.cpp b/Engine/Plugins/Experimental/TextureGraph/Source/TextureGraphEngine/Transform/Utility/T_TextureHistogram.cpp
--- a/Engine/Plugins/Experimental/TextureGraph/Source/TextureGraphEngine/Transform/Utility/T_TextureHistogram.cpp
+++ b/Engine/Plugins/Experimental/TextureGraph/Source/TextureGraphEngine/Transform/Utility/T_TextureHistogram.cpp
@@ -47,7 +47,7 @@ TiledBlobPtr T_TextureHistogram::CreateJobAndResult(JobUPtr& OutJob, MixUpdateCy
 	FIntVector4 SrcDimensions(SourceTex->GetWidth(), SourceTex->GetHeight(), 1, 1);
 
 	FString Name = FString::Printf(TEXT("[%s].[%d] Histogram"), *SourceTex->DisplayName(), TargetId);
-	RenderMaterial_FXPtr Transform = T_TextureHistogram::CreateMaterial_Histogram(TEXT("T_Histogram"), TEXT("Result"), PermutationVector, 32, 32, 1);
+	RenderMaterial_FXPtr Transform = T_TextureHistogram::CreateMaterial_Histogram(TEXT("T_Histogram"), TEXT("Result"), PermutationVector, SrcDimensions.X, SrcDimensions.Y, 1);
 
 	OutJob = std::make_unique<Job>(Cycle->GetMix(), TargetId, std::static_pointer_cast<BlobTransform>(Transform));
 	OutJob->AddArg(ARG_BLOB(SourceTex, "SourceTiles"));
diff --git a/Engine/Plugins/Experimental/TextureGraph/Source/TextureGraphEngine/Transform/Utility/T_TextureHistogram.h b/Engine/Plugins/Experimental/TextureGraph/Source/TextureGraphEngine/Transform/Utility/T_TextureHistogram.h
--- a/Engine/Plugins/Experimental/TextureGraph/Source/TextureGraphEngine/Transform/Utility/T_TextureHistogram.h
+++ b/Engine/Plugins/Experimental/TextureGraph/Source/TextureGraphEngine/Transform/Utility/T_TextureHistogram.h
@@ -8,7 +8,7 @@
 #include "Model/Mix/MixUpdateCycle.h"
 #include <DataDrivenShaderPlatformInfo.h>
 
-class CSH_Histogram : public CmpSH_Base<1, 1, 1>
+class CSH_Histogram : public CmpSH_Base<16, 16, 1>
 {
 public:
 	DECLARE_GLOBAL_SHADER(CSH_Histogram);
