Fixed issue: Texture Type, LOD Group, Compression Settings and sRGB properties do not transfer to chained Settings nodes

diff --git a/Engine/Plugins/Experimental/TextureGraph/Source/TextureGraph/Private/Expressions/Input/TG_Expression_OutputSettings.cpp b/Engine/Plugins/Experimental/TextureGraph/Source/TextureGraph/Private/Expressions/Input/TG_Expression_OutputSettings.cpp
--- a/Engine/Plugins/Experimental/TextureGraph/Source/TextureGraph/Private/Expressions/Input/TG_Expression_OutputSettings.cpp
+++ b/Engine/Plugins/Experimental/TextureGraph/Source/TextureGraph/Private/Expressions/Input/TG_Expression_OutputSettings.cpp
@@ -34,6 +34,22 @@ void UTG_Expression_OutputSettings::Evaluate(FTG_EvaluationContext* InContext)
 		{
 			Settings.TextureFormat = Input.TextureFormat;
 		}
+		if (Settings.LODGroup < Input.LODGroup || PreviousInput.LODGroup == Settings.LODGroup)
+		{
+			Settings.LODGroup = Input.LODGroup;
+		}
+		if (Settings.Compression < Input.Compression || PreviousInput.Compression == Settings.Compression)
+		{
+			Settings.Compression = Input.Compression;
+		}
+		if (Settings.bSRGB < Input.bSRGB || PreviousInput.bSRGB == Settings.bSRGB)
+		{
+			Settings.bSRGB = Input.bSRGB;
+		}
+		if (Settings.TexturePresetType < Input.TexturePresetType || PreviousInput.TexturePresetType == Settings.TexturePresetType)
+		{
+			Settings.TexturePresetType = Input.TexturePresetType;
+		}
 		
 		SettingsPin->EditSelfVar()->EditAs<FTG_OutputSettings>() = Settings;
 
