Fixed issue: Adding R16, RFloat and RGBA Float textures in the texture graph triggers assert

diff --git a/Engine/Plugins/Experimental/TextureGraph/Source/TextureGraphEngine/2D/TextureHelper.cpp b/Engine/Plugins/Experimental/TextureGraph/Source/TextureGraphEngine/2D/TextureHelper.cpp
--- a/Engine/Plugins/Experimental/TextureGraph/Source/TextureGraphEngine/2D/TextureHelper.cpp
+++ b/Engine/Plugins/Experimental/TextureGraph/Source/TextureGraphEngine/2D/TextureHelper.cpp
@@ -1471,10 +1471,13 @@ bool TextureHelper::GetPixelFormatFromTextureSourceFormat(ETextureSourceFormat S
 	switch (SourceFormat)
 	{
 	// Currently supported formats : 
+	case ETextureSourceFormat::TSF_RGBA32F:
 	case ETextureSourceFormat::TSF_RGBA16F:
 	case ETextureSourceFormat::TSF_RGBA16:
 	case ETextureSourceFormat::TSF_BGRA8:
 	case ETextureSourceFormat::TSF_BGRE8:
+	case ETextureSourceFormat::TSF_R32F:
+	case ETextureSourceFormat::TSF_R16F:
 	case ETextureSourceFormat::TSF_G16:
 	case ETextureSourceFormat::TSF_G8:
 	// Commenting as 'TSF_RGBA8': Legacy ETextureSourceFormat not supported, 
