Fixed issue: Adding a NormaltoHeightMap node in a texture graph causes the texture to flip upside down

diff --git a/Engine/Plugins/Experimental/TextureGraph/Shaders/Expressions/Expression_NormalFromHeightMap.usf b/Engine/Plugins/Experimental/TextureGraph/Shaders/Expressions/Expression_NormalFromHeightMap.usf
--- a/Engine/Plugins/Experimental/TextureGraph/Shaders/Expressions/Expression_NormalFromHeightMap.usf
+++ b/Engine/Plugins/Experimental/TextureGraph/Shaders/Expressions/Expression_NormalFromHeightMap.usf
@@ -12,12 +12,10 @@ float4 FSH_NormalFromHeightMap(float2 UV : TEXCOORD0) : SV_Target0
 {
 	float2 FullTexPos = TileInfo_fromCurrentTileToLayer(UV);
 
-	FullTexPos.y = 1.0 - FullTexPos.y; /// Mirrored vertically
-	
 	float2 UVLeft  = FullTexPos + (Offset * float2(-1,0));
 	float2 UVRight = FullTexPos + (Offset * float2(1,0));
-	float2 UVUp    = FullTexPos + (Offset * float2(0,1));
-	float2 UVDown  = FullTexPos + (Offset * float2(0,-1));
+	float2 UVUp    = FullTexPos + (Offset * float2(0,-1));
+	float2 UVDown  = FullTexPos + (Offset * float2(0,1));
 	
 	float SampleLeft  = SourceTexture.Sample(SamplerStates_Wrap, UVLeft).r;
 	float SampleRight = SourceTexture.Sample(SamplerStates_Wrap, UVRight).r;
