FTG_Variant serialization and detail customization

diff --git a/Engine/Plugins/Experimental/TextureGraph/Source/TextureGraph/Private/Expressions/TG_Expression.cpp b/Engine/Plugins/Experimental/TextureGraph/Source/TextureGraph/Private/Expressions/TG_Expression.cpp
--- a/Engine/Plugins/Experimental/TextureGraph/Source/TextureGraph/Private/Expressions/TG_Expression.cpp
+++ b/Engine/Plugins/Experimental/TextureGraph/Source/TextureGraph/Private/Expressions/TG_Expression.cpp
@@ -34,7 +34,7 @@ UTG_Node* UTG_Expression::GetParentNode() const
 	return Cast<UTG_Node>(GetOuter());
 }
 
-void UTG_Expression::NotifyExpressionChanged(FPropertyChangedEvent& PropertyChangedEvent)const
+void UTG_Expression::NotifyExpressionChanged(const FPropertyChangedEvent& PropertyChangedEvent)const
 {
 	UTG_Node* ParentNode = GetParentNode();
 	if (ParentNode)
diff --git a/Engine/Plugins/Experimental/TextureGraph/Source/TextureGraph/Private/Expressions/TG_Expression_Variant.cpp b/Engine/Plugins/Experimental/TextureGraph/Source/TextureGraph/Private/Expressions/TG_Expression_Variant.cpp
--- a/Engine/Plugins/Experimental/TextureGraph/Source/TextureGraph/Private/Expressions/TG_Expression_Variant.cpp
+++ b/Engine/Plugins/Experimental/TextureGraph/Source/TextureGraph/Private/Expressions/TG_Expression_Variant.cpp
@@ -110,7 +110,7 @@ void UTG_Expression_Variant::Evaluate(FTG_EvaluationContext* InContext)
 	}
 	case FTG_Variant::EType::Texture:
 	{
-		Output.EditTexture() = EvaluateTexture(InContext);
+		Output.EditTexture() = EvaluateTexture(InContext).RasterBlob;
 		break;
 	}
 	}
diff --git a/Engine/Plugins/Experimental/TextureGraph/Source/TextureGraph/Private/TG_Variant.cpp b/Engine/Plugins/Experimental/TextureGraph/Source/TextureGraph/Private/TG_Variant.cpp
--- a/Engine/Plugins/Experimental/TextureGraph/Source/TextureGraph/Private/TG_Variant.cpp
+++ b/Engine/Plugins/Experimental/TextureGraph/Source/TextureGraph/Private/TG_Variant.cpp
@@ -36,6 +36,14 @@ FTG_Variant::FTG_Variant(FTG_Texture RHS)
 	EditTexture() = RHS;
 }
 
+bool FTG_Variant::Serialize(FArchive& Ar)
+{
+	Ar << *this;
+	UE_LOG(LogTextureGraph, Log, TEXT("FTG_Variant::Serialize"));
+	return true;
+}
+
+
 bool FTG_Variant::ResetTypeAs(EType InType)
 {
 	if (InType != GetType())
@@ -245,6 +253,30 @@ FTG_Variant& FTG_Variant::operator = (const FTG_Texture RHS)
 	return *this;
 }
 
+bool FTG_Variant::operator==(const FTG_Variant& RHS) const
+{
+	if (GetType() == RHS.GetType())
+	{
+		switch(GetType())
+		{
+		default:
+		case EType::Scalar:
+			return Data.Get<float>() == RHS.Data.Get<float>();
+			break;
+		case EType::Color:
+			return Data.Get<FLinearColor>() == RHS.Data.Get<FLinearColor>();
+			break;
+		case EType::Vector:
+			return Data.Get<FVector4f>() == RHS.Data.Get<FVector4f>();
+			break;
+		case EType::Texture:
+			return Data.Get<FTG_Texture>() == RHS.Data.Get<FTG_Texture>();
+			break;
+		}
+	}
+	return false;
+}
+
 FTG_Variant::operator bool() const
 {
 	if (IsScalar())
diff --git a/Engine/Plugins/Experimental/TextureGraph/Source/TextureGraph/Public/Expressions/TG_Expression.h b/Engine/Plugins/Experimental/TextureGraph/Source/TextureGraph/Public/Expressions/TG_Expression.h
--- a/Engine/Plugins/Experimental/TextureGraph/Source/TextureGraph/Public/Expressions/TG_Expression.h
+++ b/Engine/Plugins/Experimental/TextureGraph/Source/TextureGraph/Public/Expressions/TG_Expression.h
@@ -111,10 +111,12 @@ protected:
 	virtual FTG_SignaturePtr BuildSignatureDynamically() const { return nullptr; }
 
 
-
+public:
 	// If some state has changed in the expression that affects its representation
 	// triggered when a property has changed and needs to be copied over to its corresponding Var
-	virtual void NotifyExpressionChanged(FPropertyChangedEvent& PropertyChangedEvent) const final;
+	// NB: required to be public for calling from TG_Variant customization
+	virtual void NotifyExpressionChanged(const FPropertyChangedEvent& PropertyChangedEvent) const final;
+protected:
 
 	// If the signature changes and the node need to regenerate its own signature.
 	// Only concrete implementation for Dynamic Expression
diff --git a/Engine/Plugins/Experimental/TextureGraph/Source/TextureGraph/Public/TG_Texture.h b/Engine/Plugins/Experimental/TextureGraph/Source/TextureGraph/Public/TG_Texture.h
--- a/Engine/Plugins/Experimental/TextureGraph/Source/TextureGraph/Public/TG_Texture.h
+++ b/Engine/Plugins/Experimental/TextureGraph/Source/TextureGraph/Public/TG_Texture.h
@@ -45,6 +45,14 @@ PRAGMA_ENABLE_DEPRECATION_WARNINGS
 	UPROPERTY(EditAnywhere, Category = "TextureDescriptor", DisplayName = "Texture Format")
 		ETG_TextureFormat TextureFormat = ETG_TextureFormat::Auto;
 
+	friend FArchive& operator<<(FArchive& Ar, FTG_TextureDescriptor& D)
+	{
+ 		Ar << D.Width;
+		Ar << D.Height;
+		Ar << D.TextureFormat;
+		return Ar;
+	}
+
 	FORCEINLINE BufferDescriptor ToBufferDescriptor() const
 	{
 		uint32 NumChannels = 0;
@@ -78,7 +86,11 @@ public:
 	operator TiledBlobPtr() { return RasterBlob; }
 	operator TiledBlobRef() { return RasterBlob; }
 	operator bool() const { return !!RasterBlob; }
-	
+	friend FArchive& operator<<(FArchive& Ar, FTG_Texture& T)
+	{
+		Ar << T.Descriptor;
+		return Ar;
+	}
 	FTG_Texture& operator = (TiledBlobRef RHS) { RasterBlob = RHS; return *this; }
 	FTG_Texture& operator = (TiledBlobPtr RHS) { RasterBlob = RHS; return *this; }
 	FORCEINLINE TiledBlob* operator -> () const { return RasterBlob.get(); }
@@ -95,6 +107,7 @@ public:
 	static FTG_Texture GetMagenta() { return { TextureHelper::GMagenta}; } 
 	static FTG_Texture GetWhiteMask() { return { TextureHelper::GWhiteMask}; } 
 	static FTG_Texture GetBlackMask() { return { TextureHelper::GBlackMask}; } 
+	
 
 	//////////////////////////////////////////////////////////////////////////
 	/// Inline functions
diff --git a/Engine/Plugins/Experimental/TextureGraph/Source/TextureGraph/Public/TG_Variant.h b/Engine/Plugins/Experimental/TextureGraph/Source/TextureGraph/Public/TG_Variant.h
--- a/Engine/Plugins/Experimental/TextureGraph/Source/TextureGraph/Public/TG_Variant.h
+++ b/Engine/Plugins/Experimental/TextureGraph/Source/TextureGraph/Public/TG_Variant.h
@@ -44,6 +44,7 @@ public:
 	FTG_Variant& operator = (const FVector4f RHS);
 	FTG_Variant& operator = (const FLinearColor RHS);
 	FTG_Variant& operator = (const FTG_Texture RHS);
+	bool operator == (const FTG_Variant& RHS) const;
 
 	// Retrieve the FName corresponding to a variant type
 	static FName GetNameFromType(ETG_VariantType InType)
@@ -107,7 +108,12 @@ public:
 		else
 			return T1;
 	}
-
+	friend FArchive& operator<<(FArchive& Ar, FTG_Variant& D)
+	{
+		return Ar << D.Data;
+	}
+	bool Serialize( FArchive& Ar );
+	
 	// FTG_Variant struct members and methods
 
 	// The concrete data
@@ -150,6 +156,17 @@ public:
 	FTG_Texture& EditTexture();
 };
 
+template<>
+struct TStructOpsTypeTraits<FTG_Variant>
+	: public TStructOpsTypeTraitsBase2<FTG_Variant>
+{
+	enum
+	{
+		WithSerializer = true,
+		WithCopy = true,
+		WithIdenticalViaEquality = true
+	};
+};
 //////////////////////////////////////////////////////////////////////////
 /// Base class for making working with Variants easier and less repetitive
 //////////////////////////////////////////////////////////////////////////
diff --git a/Engine/Plugins/Experimental/TextureGraph/Source/TextureGraphEditor/Private/Customizations/TG_VariantCustomization.h b/Engine/Plugins/Experimental/TextureGraph/Source/TextureGraphEditor/Private/Customizations/TG_VariantCustomization.h
new file mode 100644
--- /dev/null
+++ b/Engine/Plugins/Experimental/TextureGraph/Source/TextureGraphEditor/Private/Customizations/TG_VariantCustomization.h
@@ -0,0 +1,87 @@
+// Copyright Epic Games, Inc. All Rights Reserved.
+
+#pragma once
+
+#include "CoreMinimal.h"
+#include "IPropertyTypeCustomization.h"
+#include "DetailWidgetRow.h"
+#include "DetailLayoutBuilder.h"
+#include "IDetailChildrenBuilder.h"
+#include "IPropertyUtilities.h"
+#include "TG_Texture.h"
+#include "Expressions/TG_Expression.h"
+#include "STG_TextureHistogram.h"
+
+
+class FTG_VariantCustomization : public IPropertyTypeCustomization
+{
+public:
+	static TSharedRef<IPropertyTypeCustomization> Create()
+	{
+		return MakeShareable(new FTG_VariantCustomization);
+	}
+
+	virtual void CustomizeHeader(TSharedRef<IPropertyHandle> PropertyHandle, FDetailWidgetRow& HeaderRow, IPropertyTypeCustomizationUtils& CustomizationUtils) override
+	{
+	}
+
+	virtual void CustomizeChildren(TSharedRef<IPropertyHandle> PropertyHandle, IDetailChildrenBuilder& ChildBuilder, IPropertyTypeCustomizationUtils& CustomizationUtils) override
+	{
+		const FString TGType = PropertyHandle->GetMetaData(FName(TEXT("TGType"))); 
+		if (TGType == TEXT("TG_Output"))
+		{
+			
+			TArray<UObject*> OuterObjects;
+			PropertyHandle->GetOuterObjects(OuterObjects);
+			check(OuterObjects[0]);
+			UTG_Expression* ParentExpression = Cast<UTG_Expression>(OuterObjects[0]);
+			TArray<void*> VariantPtrs;
+			PropertyHandle->AccessRawData(VariantPtrs);
+			if (VariantPtrs.Num() == 0)
+			{
+				return;
+			}
+			FTG_Variant* VariantPtr = reinterpret_cast<FTG_Variant*>(VariantPtrs[0]);
+			// check if texture
+			if (VariantPtr->IsTexture())
+			{
+				TSharedRef<FStructOnScope> StructData = MakeShareable(new FStructOnScope(FTG_TextureDescriptor::StaticStruct(), (uint8*)&VariantPtr->GetTexture().Descriptor));
+				IDetailPropertyRow* DetailPropertyRow = ChildBuilder.AddExternalStructureProperty(StructData, NAME_None, FAddPropertyParams().CreateCategoryNodes(false).HideRootObjectNode(true));
+				DetailPropertyRow->ShouldAutoExpand(true);
+				
+				if (TSharedPtr<IPropertyHandle> Handle = DetailPropertyRow->GetPropertyHandle(); Handle.IsValid())
+				{
+					FString DisplayName = PropertyHandle->GetPropertyDisplayName().ToString() + TEXT(" Settings");
+					DetailPropertyRow->DisplayName(FText::FromString(DisplayName));
+					
+					// Handle->SetInstanceMetaData(TEXT("TGType"), TEXT("TG_Output"));
+					TSharedPtr<IPropertyUtilities> PropertyUtilities;
+					PropertyUtilities = CustomizationUtils.GetPropertyUtilities();
+					Handle->SetOnChildPropertyValuePreChange(FSimpleDelegate::CreateLambda([=, this]()
+					{
+						ParentExpression->Modify();
+					}));
+					Handle->SetOnChildPropertyValueChangedWithData(TDelegate<void(const FPropertyChangedEvent&)>::CreateLambda([=, this](const FPropertyChangedEvent& InEvent)
+					{
+						if (PropertyUtilities.IsValid())
+						{
+							void* StructData = nullptr;
+							const FPropertyAccess::Result Result = Handle->GetValueData(StructData);
+							if(Result == FPropertyAccess::Success)
+							{
+								check(StructData);
+								FTG_TextureDescriptor* Descriptor = static_cast<FTG_TextureDescriptor*>(StructData);
+
+								VariantPtr->EditTexture().Descriptor = *Descriptor;
+								FPropertyChangedEvent ChangeEvent(InEvent);
+								ChangeEvent.MemberProperty = PropertyHandle->GetProperty();
+								ParentExpression->NotifyExpressionChanged(ChangeEvent);
+								PropertyUtilities->RequestRefresh();
+							}
+						}
+					}));
+				}
+			}
+		}
+	}
+};
diff --git a/Engine/Plugins/Experimental/TextureGraph/Source/TextureGraphEditor/Private/TG_Editor.cpp b/Engine/Plugins/Experimental/TextureGraph/Source/TextureGraphEditor/Private/TG_Editor.cpp
--- a/Engine/Plugins/Experimental/TextureGraph/Source/TextureGraphEditor/Private/TG_Editor.cpp
+++ b/Engine/Plugins/Experimental/TextureGraph/Source/TextureGraphEditor/Private/TG_Editor.cpp
@@ -176,7 +176,11 @@ void FTG_Editor::InitEditor(const EToolkitMode::Type Mode, const TSharedPtr< cla
 	// All the modifications are applied to original TextureGraph when asset is saved.
 	// Propagate all object flags except for RF_Standalone, otherwise the preview material won't GC once
 	// the TS editor releases the reference.
-	EditedTextureGraph = Cast<UTextureGraph>(StaticDuplicateObject(OriginalTextureGraph, OriginalTextureGraph->GetOuter(), NAME_None, ~RF_Standalone, UTextureGraph::StaticClass()));
+	// overwrite the original TextureGraph in place by constructing a new one with the same name
+	FObjectDuplicationParameters Params = InitStaticDuplicateObjectParams(OriginalTextureGraph, OriginalTextureGraph->GetOuter(), NAME_None,
+	~RF_Standalone, UTextureGraph::StaticClass(), EDuplicateMode::Normal, EInternalObjectFlags::None);
+	
+	EditedTextureGraph = Cast<UTextureGraph>(StaticDuplicateObjectEx(Params));
 	FCoreUObjectDelegates::OnObjectPreSave.AddSP(this, &FTG_Editor::OnTextureGraphPreSave);
 
 	//Editor gets notified when rendering is done
@@ -1608,7 +1612,11 @@ bool FTG_Editor::UpdateOriginalTextureGraph()
 		const TMap<FName, FString>* MetaData = UMetaData::GetMapForObject(OriginalTextureGraph);
 
 		// overwrite the original TextureGraph in place by constructing a new one with the same name
-		OriginalTextureGraph = Cast<UTextureGraph>(StaticDuplicateObject(EditedTextureGraph, OriginalTextureGraph->GetOuter(), OriginalTextureGraph->GetFName(), RF_AllFlags, OriginalTextureGraph->GetClass()));
+		FObjectDuplicationParameters Params = InitStaticDuplicateObjectParams(EditedTextureGraph, OriginalTextureGraph->GetOuter(), OriginalTextureGraph->GetFName(),
+		RF_AllFlags, UTextureGraph::StaticClass(), EDuplicateMode::Normal, EInternalObjectFlags::None);
+
+		// UObject* NewAsset = StaticDuplicateObjectEx(Params);
+		OriginalTextureGraph = Cast<UTextureGraph>(StaticDuplicateObjectEx(Params));
 
 		// Restore the metadata
 		if (MetaData)
diff --git a/Engine/Plugins/Experimental/TextureGraph/Source/TextureGraphEditor/Private/TextureGraphEditorModule.cpp b/Engine/Plugins/Experimental/TextureGraph/Source/TextureGraphEditor/Private/TextureGraphEditorModule.cpp
--- a/Engine/Plugins/Experimental/TextureGraph/Source/TextureGraphEditor/Private/TextureGraphEditorModule.cpp
+++ b/Engine/Plugins/Experimental/TextureGraph/Source/TextureGraphEditor/Private/TextureGraphEditorModule.cpp
@@ -13,6 +13,7 @@
 #include "EdGraph/TG_EditorErrorReporter.h"
 #include "Customizations/TG_ParameterCustomization.h"
 #include "Customizations/TG_TextureCustomization.h"
+#include "Customizations/TG_VariantCustomization.h"
 #include "Customizations/TG_ScalarCustomization.h"
 #include "Customizations/TG_MaterialMappingInfoCustomization.h"
 #include "Customizations/TG_OutputSettingsCustomization.h"
@@ -32,6 +33,7 @@ void FTextureGraphEditorModule::StartupModule()
 	FPropertyEditorModule& PropertyEditorModule = FModuleManager::GetModuleChecked<FPropertyEditorModule>("PropertyEditor");
 	PropertyEditorModule.RegisterCustomPropertyTypeLayout("TG_ParameterInfo", FOnGetPropertyTypeCustomizationInstance::CreateStatic(&FTG_ParameterInfoCustomization::Create));
 	PropertyEditorModule.RegisterCustomPropertyTypeLayout("TG_Texture", FOnGetPropertyTypeCustomizationInstance::CreateStatic(&FTG_TextureCustomization::Create));
+	PropertyEditorModule.RegisterCustomPropertyTypeLayout("TG_Variant", FOnGetPropertyTypeCustomizationInstance::CreateStatic(&FTG_VariantCustomization::Create));
 	auto ScalarIdentifier = MakeShared<FTG_ScalarTypeIdentifier>();
 	PropertyEditorModule.RegisterCustomPropertyTypeLayout("FloatProperty", FOnGetPropertyTypeCustomizationInstance::CreateStatic(&FTG_ScalarCustomization::Create), ScalarIdentifier);
 	PropertyEditorModule.RegisterCustomPropertyTypeLayout("MaterialMappingInfo", FOnGetPropertyTypeCustomizationInstance::CreateStatic(&FTG_MaterialMappingInfoCustomization::Create));
