Split Channels mismatch in texture details and 2d preview. (Selected pin wasn't set correctly)

diff --git a/Engine/Plugins/Experimental/TextureGraph/Source/TextureGraphEditor/Private/EdGraph/TG_EdGraphNode.cpp b/Engine/Plugins/Experimental/TextureGraph/Source/TextureGraphEditor/Private/EdGraph/TG_EdGraphNode.cpp
--- a/Engine/Plugins/Experimental/TextureGraph/Source/TextureGraphEditor/Private/EdGraph/TG_EdGraphNode.cpp
+++ b/Engine/Plugins/Experimental/TextureGraph/Source/TextureGraphEditor/Private/EdGraph/TG_EdGraphNode.cpp
@@ -112,7 +112,7 @@ void UTG_EdGraphNode::GetNodeContextMenuActions(UToolMenu* Menu, class UGraphNod
 			FUIAction(
 				FExecuteAction::CreateLambda([Pin = Context->Pin, Node = Context->Node, this]
 				{
-					auto TSEdGraph = Cast<UTG_EdGraph>(Node->GetGraph());
+					UTG_EdGraph* TSEdGraph = Cast<UTG_EdGraph>(Node->GetGraph());
 					TSEdGraph->PinSelectionManager.UpdateSelection(const_cast<UEdGraphPin*>(Pin));
 				}),
 				FCanExecuteAction::CreateLambda([Pin = Context->Pin, this]
@@ -358,6 +358,16 @@ void UTG_EdGraphNode::ReconstructNode()
 
 	// Notify editor
 	OnNodeReconstructDelegate.Broadcast();
+
+	// if no selected pin, we set it to the first output pin
+	if (GetSelectedPin() == nullptr)
+	{
+		if (GetOutputPins().Num() > 0)
+		{
+			UTG_EdGraph* TSEdGraph = Cast<UTG_EdGraph>(GetGraph());
+			TSEdGraph->PinSelectionManager.UpdateSelection(GetOutputPins()[0]);
+		}
+	}
 }
 
 FString UTG_EdGraphNode::GetTitleDetail()
@@ -390,6 +400,7 @@ FString UTG_EdGraphNode::GetTitleDetail()
 				Args.Add(TEXT("Width"), FString::FromInt(Desc.Width));
 				Args.Add(TEXT("Height"), FString::FromInt(Desc.Height));
 				Details = FString::Format(TEXT("{Channels}_{Format}, {IsSRGB}\n{Width}x{Height}"), Args);
+				break;
 			}
 		}
 	}
