Fixed inconsistent alpha handling in various TextureGraph views

#virtualized

diff --git a/Engine/Plugins/Experimental/TextureGraph/Source/TextureGraphEditor/Private/TG_TexturePreviewViewportClient.cpp b/Engine/Plugins/Experimental/TextureGraph/Source/TextureGraphEditor/Private/TG_TexturePreviewViewportClient.cpp
--- a/Engine/Plugins/Experimental/TextureGraph/Source/TextureGraphEditor/Private/TG_TexturePreviewViewportClient.cpp
+++ b/Engine/Plugins/Experimental/TextureGraph/Source/TextureGraphEditor/Private/TG_TexturePreviewViewportClient.cpp
@@ -57,7 +57,20 @@ void FTG_TexturePreviewViewportClient::Draw(FViewport* Viewport, FCanvas* Canvas
 
 	if (!SelectionPreviewPtr.IsValid() || !Texture)
 	{
-		Canvas->Clear(ClearColorOverride);
+		// Draw the background checkerboard pattern in the same size/position as the render texture so it will show up anywhere
+		// the texture has transparency
+		if (Viewport && CheckerboardTexture)
+		{
+			const int32 CheckerboardSizeX = FMath::Max<int32>(1, CheckerboardTexture->GetSizeX());
+			const int32 CheckerboardSizeY = FMath::Max<int32>(1, CheckerboardTexture->GetSizeY());
+		
+			Canvas->DrawTile( 0.0f, 0.0f, Viewport->GetSizeXY().X, Viewport->GetSizeXY().Y, 0.0f, 0.0f, (float)Viewport->GetSizeXY().X / CheckerboardSizeX, (float)Viewport->GetSizeXY().Y / CheckerboardSizeY, FLinearColor::White, CheckerboardTexture->GetResource());
+		}
+		FCanvasTileItem TileItem( FVector2D::ZeroVector, Viewport->GetSizeXY(), ClearColorOverride);
+		TileItem.BlendMode = SelectionPreviewPtr.Pin()->GetColourChannelBlendMode();
+		
+		Canvas->DrawItem( TileItem );
+		
 		return;
 	}
 	
@@ -140,7 +153,7 @@ void FTG_TexturePreviewViewportClient::Draw(FViewport* Viewport, FCanvas* Canvas
 		const int32 CheckerboardSizeX = FMath::Max<int32>(1, CheckerboardTexture->GetSizeX());
 		const int32 CheckerboardSizeY = FMath::Max<int32>(1, CheckerboardTexture->GetSizeY());
 
-		Canvas->DrawTile(XPos, YPos, Width, Height, 0.0f, 0.0f, (float)Width / CheckerboardSizeX, (float)Height / CheckerboardSizeY, FLinearColor::White, CheckerboardTexture->GetResource());
+		Canvas->DrawTile( XPos, YPos, Width, Height, 0.0f, 0.0f, (float)Width / CheckerboardSizeX, (float)Height / CheckerboardSizeY, FLinearColor::White, CheckerboardTexture->GetResource());
 	}
 
 	FTexturePlatformData** RunningPlatformDataPtr = Texture->GetRunningPlatformData();
@@ -277,7 +290,7 @@ void FTG_TexturePreviewViewportClient::AddReferencedObjects(FReferenceCollector&
 void FTG_TexturePreviewViewportClient::ModifyCheckerboardTextureColors()
 {
 	DestroyCheckerboardTexture();
-	CheckerboardTexture = FImageUtils::CreateCheckerboardTexture(FColor(1,1,1,1), FColor(0, 0, 0, 1), 64);
+	CheckerboardTexture = FImageUtils::CreateCheckerboardTexture();
 }
 
 
