Reconciled offline work

diff --git a/Engine/Plugins/Experimental/TextureGraph/Source/TextureGraph/Public/Expressions/Filter/TG_Expression_Levels.h b/Engine/Plugins/Experimental/TextureGraph/Source/TextureGraph/Public/Expressions/Filter/TG_Expression_Levels.h
--- a/Engine/Plugins/Experimental/TextureGraph/Source/TextureGraph/Public/Expressions/Filter/TG_Expression_Levels.h
+++ b/Engine/Plugins/Experimental/TextureGraph/Source/TextureGraph/Public/Expressions/Filter/TG_Expression_Levels.h
@@ -77,8 +77,9 @@ public:
 	UPROPERTY(EditAnywhere, Category = NoCategory, meta = (TGType = "TG_Setting", PinDisplayName = "", MD_LevelsSettings, PinNotConnectable = true))
 	FTG_LevelsSettings					Levels;
 
-	UPROPERTY(EditAnywhere, Category = NoCategory, meta = (TGType = "TG_Setting", PinNotConnectable = true))
-	bool								AutoLevels = true;
+	// TODO: Show Property after 5.4 release
+	//	UPROPERTY(EditAnywhere, Category = NoCategory, meta = (TGType = "TG_Setting", PinNotConnectable = true))
+	bool								AutoLevels = false;
 
 	// The output image filtered result of the Levels operator 
 	UPROPERTY(EditAnywhere, Category = NoCategory, meta = (TGType = "TG_Output", PinDisplayName = ""))
