Fixed issue: Pattern showing transparency when selecting different pattern types

diff --git a/Engine/Plugins/Experimental/TextureGraph/Shaders/Mask/PatternMask.usf b/Engine/Plugins/Experimental/TextureGraph/Shaders/Mask/PatternMask.usf
--- a/Engine/Plugins/Experimental/TextureGraph/Shaders/Mask/PatternMask.usf
+++ b/Engine/Plugins/Experimental/TextureGraph/Shaders/Mask/PatternMask.usf
@@ -163,7 +163,8 @@ float4 FSH_PatternMask(float2 uv : TEXCOORD0) : SV_Target0
 	float checkerA = (blackWhite.x - blackWhite.y) % 2;
 	float checkerB = (blackWhite.y - blackWhite.x) % 2;
 	float checker = clamp(checkerA, 0, 1) + clamp(checkerB, 0, 1);
-	return (Invert  > 0.0 ? 1.0 - checker : checker);
+	float Final = (Invert  > 0.0 ? 1.0 - checker : checker);
+	return float4(Final, Final, Final, 1);
 	
 #endif
 #if (PATTERN_TYPE == PATTERN_GRADIENT)
@@ -177,7 +178,8 @@ float4 FSH_PatternMask(float2 uv : TEXCOORD0) : SV_Target0
 	
 	float2 rotatedOffsetUV = mul((layer_uv ) - center, float2x2(cosine, -sine, sine, cosine)) + center;
 	float gradValue = rotatedOffsetUV.x + tiltGrad;
-	return (Invert  > 0.0 ? 1.0 - gradValue : gradValue);
+	float Final = (Invert  > 0.0 ? 1.0 - gradValue : gradValue);
+	return float4(Final, Final, Final, 1);
 #endif
 
 	// BASE VALUES
@@ -213,7 +215,8 @@ float4 FSH_PatternMask(float2 uv : TEXCOORD0) : SV_Target0
 	float2 offsetUV =  repeatedUVCentralized * 0.5;
 	float circleFill = step(sqrt((pow(offsetUV.x, 2) / pow(spaces.x, 2)) + (pow(offsetUV.y, 2) / pow(spaces.y, 2))), lerp(0.5, fmod(sizeJitter, 0.5), any(JS_Threshold)));
 	
-	return (Invert  > 0.0 ? 1.0 - circleFill : circleFill);
+	float Final = (Invert  > 0.0 ? 1.0 - circleFill : circleFill);
+	return float4(Final, Final, Final, 1);
 #endif
 	
 	float2 stretchCompensation = StretchCompensation(tiles.x, tiles.y, Spacing.x, Spacing.y, spacingApplied);
