Fix histogram crash for transient textures.

diff --git a/Engine/Plugins/Experimental/TextureGraph/Source/TextureGraphEditor/Private/Customizations/STG_TextureHistogram.cpp b/Engine/Plugins/Experimental/TextureGraph/Source/TextureGraphEditor/Private/Customizations/STG_TextureHistogram.cpp
--- a/Engine/Plugins/Experimental/TextureGraph/Source/TextureGraphEditor/Private/Customizations/STG_TextureHistogram.cpp
+++ b/Engine/Plugins/Experimental/TextureGraph/Source/TextureGraphEditor/Private/Customizations/STG_TextureHistogram.cpp
@@ -29,32 +29,34 @@ void STG_TextureHistogram::Construct(const FArguments& InArgs)
 	];
 }
 
-
 void STG_TextureHistogram::SetTexture(FTG_Texture& Source, UTextureGraph* InTextureGraph)
 {
 	if (Source.RasterBlob && InTextureGraph)
 	{
+		if (Source.RasterBlob->IsTransient())
+			return;
+
 		Source.RasterBlob->OnFinalise()
-		.then([this, InTextureGraph, Source]() mutable
-		{
-			// OnFinalise can sometimes occur after the editor is closed and thus can potentially
-			// deallocate all corresponding slate objects
-			if (DoesSharedInstanceExist())
+			.then([this, InTextureGraph, Source]() mutable
+			{
+				// OnFinalise can sometimes occur after the editor is closed and thus can potentially
+				// deallocate all corresponding slate objects
+				if (DoesSharedInstanceExist() && !Source.RasterBlob->IsTransient())
+				{
+					T_TextureHistogram::CreateOnService(InTextureGraph, std::static_pointer_cast<TiledBlob>(Source.RasterBlob), 0);
+					return Source->GetHistogram()->OnFinalise();
+				}
+				return (AsyncBlobResultPtr)(cti::make_ready_continuable<const Blob*>(nullptr));
+			})
+			.then([this, Source]() mutable
 			{
-				T_TextureHistogram::CreateOnService(InTextureGraph, std::static_pointer_cast<TiledBlob>(Source.RasterBlob), 0);
-				return  Source->GetHistogram()->OnFinalise();
-			}
-			return (AsyncBlobResultPtr)(cti::make_ready_continuable<const Blob*>(nullptr));
-		})
-		.then([this, Source]() mutable
-		{
-			// OnFinalise can sometimes occur after the editor is closed and thus can potentially
-			// deallocate all corresponding slate objects
-			if (DoesSharedInstanceExist())
-            {
-				TiledBlobPtr FinalizedHistogram = std::static_pointer_cast<TiledBlob>(Source.RasterBlob->GetHistogram());
-				HistogramBars->Update(FinalizedHistogram);
-			}
-		});
+				// OnFinalise can sometimes occur after the editor is closed and thus can potentially
+				// deallocate all corresponding slate objects
+				if (DoesSharedInstanceExist() && !Source.RasterBlob->IsTransient())
+				{
+					TiledBlobPtr FinalizedHistogram = std::static_pointer_cast<TiledBlob>(Source.RasterBlob->GetHistogram());
+					HistogramBars->Update(FinalizedHistogram);
+				}
+			});
 	}
 }
\ No newline at end of file
diff --git a/Engine/Plugins/Experimental/TextureGraph/Source/TextureGraphEngine/Data/Blob.cpp b/Engine/Plugins/Experimental/TextureGraph/Source/TextureGraphEngine/Data/Blob.cpp
--- a/Engine/Plugins/Experimental/TextureGraph/Source/TextureGraphEngine/Data/Blob.cpp
+++ b/Engine/Plugins/Experimental/TextureGraph/Source/TextureGraphEngine/Data/Blob.cpp
@@ -494,6 +494,11 @@ void Blob::UpdateLinkedBlobs(bool bDoFinalise)
 	}
 }
 
+BlobPtr	Blob::GetHistogram() const
+{ 
+	return Histogram ? Histogram : std::static_pointer_cast<Blob>(TextureHelper::GetBlack()); 
+}
+
 void Blob::AddLinkedBlob(BlobPtr LinkedBlob)
 {
 	if (IsFinalised())
diff --git a/Engine/Plugins/Experimental/TextureGraph/Source/TextureGraphEngine/Data/Blob.h b/Engine/Plugins/Experimental/TextureGraph/Source/TextureGraphEngine/Data/Blob.h
--- a/Engine/Plugins/Experimental/TextureGraph/Source/TextureGraphEngine/Data/Blob.h
+++ b/Engine/Plugins/Experimental/TextureGraph/Source/TextureGraphEngine/Data/Blob.h
@@ -6,6 +6,7 @@
 #include "Device/DeviceBuffer.h"
 #include "BlobRef.h"
 #include "Helper/DataUtil.h"
+#include "2D/TextureHelper.h"
 
 class Device;
 class RawBuffer;
@@ -142,6 +143,7 @@ public:
 	virtual bool					HasLODLevel(int32 Index) const;
 	virtual int32					NumLODLevels() const;
 	virtual void					SetLODLevel(int32 Level, BlobPtr LODBlob, BlobPtrW LODParentBlob, BlobPtrW LODSourceBlob, bool bAddToBlobber); 
+	BlobPtr							GetHistogram() const; 
 
 	//////////////////////////////////////////////////////////////////////////
 	/// Ownership related
@@ -185,8 +187,6 @@ public:
 	//Check if blob already has a histogram
 	FORCEINLINE bool				HasHistogram() const { return Histogram != nullptr; }
 
-	FORCEINLINE BlobPtr				GetHistogram() { check(HasHistogram()); return Histogram; }
-
 	//////////////////////////////////////////////////////////////////////////
 	/// Static functions
 	//////////////////////////////////////////////////////////////////////////
