Window menu grouping

diff --git a/Engine/Plugins/Experimental/TextureGraph/Source/TextureGraphEditor/Private/TG_Editor.cpp b/Engine/Plugins/Experimental/TextureGraph/Source/TextureGraphEditor/Private/TG_Editor.cpp
--- a/Engine/Plugins/Experimental/TextureGraph/Source/TextureGraphEditor/Private/TG_Editor.cpp
+++ b/Engine/Plugins/Experimental/TextureGraph/Source/TextureGraphEditor/Private/TG_Editor.cpp
@@ -103,10 +103,10 @@ void FTG_Editor::RegisterTabSpawners(const TSharedRef<class FTabManager>& InTabM
 		.SetGroup(WorkspaceMenuCategoryRef)
 		.SetIcon(FSlateIcon(FAppStyle::GetAppStyleSetName(), "Kismet.Tabs.Palette"));
 
-	InTabManager->RegisterTabSpawner(FTG_EditorTabs::FindTabId, FOnSpawnTab::CreateSP(this, &FTG_Editor::SpawnTab_Find))
+	/*InTabManager->RegisterTabSpawner(FTG_EditorTabs::FindTabId, FOnSpawnTab::CreateSP(this, &FTG_Editor::SpawnTab_Find))
 		.SetDisplayName(LOCTEXT("FindTab", "Find Results"))
 		.SetGroup(WorkspaceMenuCategoryRef)
-		.SetIcon(FSlateIcon(FAppStyle::GetAppStyleSetName(), "Kismet.Tabs.FindResults"));
+		.SetIcon(FSlateIcon(FAppStyle::GetAppStyleSetName(), "Kismet.Tabs.FindResults"));*/
 
 	InTabManager->RegisterTabSpawner(FTG_EditorTabs::PreviewSceneSettingsTabId, FOnSpawnTab::CreateSP(this, &FTG_Editor::SpawnTab_PreviewSettings))
 		.SetDisplayName(LOCTEXT("PreviewSceneSettingsTab", "Preview Scene Settings"))
@@ -151,13 +151,14 @@ void FTG_Editor::UnregisterTabSpawners(const TSharedRef<class FTabManager>& InTa
 	InTabManager->UnregisterTabSpawner(FTG_EditorTabs::ViewportTabId);
 	InTabManager->UnregisterTabSpawner(FTG_EditorTabs::PropertiesTabId);
 	InTabManager->UnregisterTabSpawner(FTG_EditorTabs::PaletteTabId);
-	InTabManager->UnregisterTabSpawner(FTG_EditorTabs::FindTabId);
+	//InTabManager->UnregisterTabSpawner(FTG_EditorTabs::FindTabId);
 	InTabManager->UnregisterTabSpawner(FTG_EditorTabs::PreviewSceneSettingsTabId);
 	InTabManager->UnregisterTabSpawner(FTG_EditorTabs::ParameterDefaultsTabId);
 	InTabManager->UnregisterTabSpawner(FTG_EditorTabs::SelectionPreviewTabId);
 	InTabManager->UnregisterTabSpawner(FTG_EditorTabs::OutputTabId);
 	InTabManager->UnregisterTabSpawner(FTG_EditorTabs::PreviewSettingsTabId);
 	InTabManager->UnregisterTabSpawner(FTG_EditorTabs::ErrorsTabId);
+	InTabManager->UnregisterTabSpawner(FTG_EditorTabs::TextureDetailsTabId);
 
 	FAssetEditorToolkit::UnregisterTabSpawners(InTabManager);
 }
@@ -269,7 +270,7 @@ void FTG_Editor::InitEditor(const EToolkitMode::Type Mode, const TSharedPtr< cla
 	GEditor->RegisterForUndo(this);
 
 	// Setup layout 
-	const TSharedRef<FTabManager::FLayout> StandaloneDefaultLayout = FTabManager::NewLayout("Standalone_TG_Editor_Layout_v0.0.7")
+	const TSharedRef<FTabManager::FLayout> StandaloneDefaultLayout = FTabManager::NewLayout("Standalone_TG_Editor_Layout_v0.0.8")
 		->AddArea
 		(
 			FTabManager::NewPrimaryArea()->SetOrientation(Orient_Vertical)
@@ -348,9 +349,6 @@ void FTG_Editor::InitEditor(const EToolkitMode::Type Mode, const TSharedPtr< cla
 		// turn the viewporttab back off
 		ViewportTab->RequestCloseTab();
 	}
-	
-	bShowPaletteView = PaletteTab.IsValid();
-	bShowNodeHistogram = NodeHistogramTab.IsValid();
 
 	EditedTextureGraph->GetPackage()->ClearDirtyFlag();
 
@@ -480,7 +478,6 @@ void FTG_Editor::RegisterToolbar()
 			TAttribute<FText>(),
 			TAttribute<FText>(),
 			FSlateIcon(FAppStyle::GetAppStyleSetName(), "FontEditor.Update")));
-		// Section.AddEntry(FToolMenuEntry::InitWidget("Auto Update", AutoRun.ToSharedRef(), FText::FromString("Auto Update")));
 
 #if UE_BUILD_DEBUG
 		Section.AddEntry(FToolMenuEntry::InitToolBarButton(
@@ -491,22 +488,66 @@ void FTG_Editor::RegisterToolbar()
 #endif
 
 		{
-			FToolMenuSection& TabsSection = ToolBar->AddSection("TabWindows", TAttribute<FText>(), InsertAfterAssetSection);
-			TabsSection.AddEntry(FToolMenuEntry::InitToolBarButton(
-				FTG_EditorCommands::Get().TogglePaletteTab,
-				TAttribute<FText>(),
-				TAttribute<FText>(),
-				FSlateIcon(FAppStyle::GetAppStyleSetName(), "LevelEditor.Tabs.Details")));
-
-			TabsSection.AddEntry(FToolMenuEntry::InitToolBarButton(
-				FTG_EditorCommands::Get().ToggleNodeHistogramTab,
-				TAttribute<FText>(),
-				TAttribute<FText>(),
-				FSlateIcon(FAppStyle::GetAppStyleSetName(), "Kismet.Tabs.Palette")));
+			FToolMenuSection& TabsSection = ToolBar->AddSection("Texture Graph Windows", TAttribute<FText>(), InsertAfterAssetSection);
+
+			FToolMenuEntry TabsMenu = FToolMenuEntry::InitComboButton(
+				"Windows",
+				FUIAction(),
+				FNewToolMenuDelegate::CreateLambda([this](UToolMenu* InSubMenu)
+				{
+					FToolMenuSection& Section = InSubMenu->FindOrAddSection("TextureGraphWindowsOptions");
+					if (WorkspaceMenuCategory)
+					{
+						for (auto& item : WorkspaceMenuCategory->GetChildItems())
+						{
+							TSharedPtr<FTabSpawnerEntry> Spawner = item->AsSpawnerEntry();
+							Section.AddMenuEntry(Spawner->GetTabType(),
+								Spawner->GetDisplayName(),
+								Spawner->GetTooltipText(),
+								FSlateIcon(),
+								FUIAction(
+									FExecuteAction::CreateSP(this, &FTG_Editor::HandleTabWindowSelected, Spawner->GetTabType(), Spawner->GetDisplayName().ToString()),
+									FCanExecuteAction(),
+									FIsActionChecked::CreateLambda([this,&item]() {return GetTabSelected(item->AsSpawnerEntry()->GetTabType()); })
+								),
+								EUserInterfaceActionType::Check);
+						}
+					}
+				}),
+				LOCTEXT("TextureGraphWindows_Label", "Windows"),
+				LOCTEXT("TextureGraphWindows_Tooltip", "Show/Hide Texture Graph tab windows"),
+				FSlateIcon(FAppStyle::GetAppStyleSetName(), "Kismet.Tabs.Palette")
+			);
+			TabsMenu.StyleNameOverride = "CalloutToolbar";
+			TabsSection.AddEntry(TabsMenu);
 		}
 	}
 }
 
+void FTG_Editor::HandleTabWindowSelected(const FName TabID, const FString OuputName)
+{
+	TSharedPtr<SDockTab> Tab = GetTabManager()->FindExistingLiveTab(TabID);
+	if (Tab)
+	{
+		Tab->RequestCloseTab();
+	}
+	else
+	{
+		GetTabManager()->TryInvokeTab(TabID);
+	}
+}
+
+bool FTG_Editor::GetTabSelected(const FName TabID)
+{
+	TSharedPtr<SDockTab> Tab = GetTabManager()->FindExistingLiveTab(TabID);
+	if (Tab)
+	{
+		return true;
+	}
+
+	return false;
+}
+
 bool FTG_Editor::IsShowingAutoUpdate() const
 {
 	return bAutoRunGraph;
@@ -569,18 +610,6 @@ void FTG_Editor::BindCommands()
 		FGenericCommands::Get().Redo,
 		FExecuteAction::CreateSP(this, &FTG_Editor::RedoGraphAction));
 	
-	ToolkitCommands->MapAction(
-		TGEditorCommands.TogglePaletteTab,
-		FExecuteAction::CreateSP(this, &FTG_Editor::TogglePaletteView),
-		FCanExecuteAction(),
-		FIsActionChecked::CreateSP(this, &FTG_Editor::IsShowingPaletteView));
-
-	ToolkitCommands->MapAction(
-		TGEditorCommands.ToggleNodeHistogramTab,
-		FExecuteAction::CreateSP(this, &FTG_Editor::ToggleNodeHistogramView),
-		FCanExecuteAction(),
-		FIsActionChecked::CreateSP(this, &FTG_Editor::IsShowingNodeHistogramView));
-	
 	GraphEditorCommands->MapAction(
 		TGEditorCommands.ConvertInputParameterToConstant,
 		FExecuteAction::CreateSP(this, &FTG_Editor::OnConvertInputParameterToFromConstant));
@@ -596,58 +625,6 @@ void FTG_Editor::BindCommands()
 
 }
 
-void FTG_Editor::SetShowNodeHistogramView(const bool bValue)
-{
-	bShowNodeHistogram = bValue;
-
-	// open/close 3DPreview tab
-	DisplayNodeHistogramView(bShowNodeHistogram);
-}
-
-void FTG_Editor::SetShowPaletteView(const bool bValue)
-{
-	bShowPaletteView = bValue;
-
-	// open/close 3DPreview tab
-	DisplayPaletteView(bShowPaletteView);
-}
-
-void FTG_Editor::ToggleNodeHistogramView()
-{
-	// Toggle the showing of material Texture Details each time the user presses the show 3DPreview button
-	SetShowNodeHistogramView(!bShowNodeHistogram);
-}
-
-void FTG_Editor::TogglePaletteView()
-{
-	// Toggle the showing of material 3DPreview each time the user presses the show 3DPreview button
-	SetShowPaletteView(!bShowPaletteView);
-}
-
-void FTG_Editor::DisplayPaletteView(const bool bShow)
-{
-	if (bShow)
-	{
-		GetTabManager()->TryInvokeTab(FTG_EditorTabs::PaletteTabId);
-	}
-	else if (!bShowPaletteView && PaletteTab.IsValid())
-	{
-		PaletteTab.Pin()->RequestCloseTab();
-	}
-}
-
-void FTG_Editor::DisplayNodeHistogramView(const bool bShow)
-{
-	if (bShow)
-	{
-		GetTabManager()->TryInvokeTab(FTG_EditorTabs::TextureDetailsTabId);
-	}
-	else if (!bShowNodeHistogram && NodeHistogramTab.IsValid())
-	{
-		NodeHistogramTab.Pin()->RequestCloseTab();
-	}
-}
-
 void FTG_Editor::OnRunGraph_Clicked()
 {
 	// Trigger and Evaluation of the TextureGraph
diff --git a/Engine/Plugins/Experimental/TextureGraph/Source/TextureGraphEditor/Private/TG_Editor.h b/Engine/Plugins/Experimental/TextureGraph/Source/TextureGraphEditor/Private/TG_Editor.h
--- a/Engine/Plugins/Experimental/TextureGraph/Source/TextureGraphEditor/Private/TG_Editor.h
+++ b/Engine/Plugins/Experimental/TextureGraph/Source/TextureGraphEditor/Private/TG_Editor.h
@@ -165,6 +165,9 @@ private:
 	void											ToggleAutoUpdate();
 	bool											IsShowingAutoUpdate() const;
 
+	void											HandleTabWindowSelected(const FName TabID,const FString OuputName);
+	bool											GetTabSelected(const FName TabID);
+
 	void											OnRenderingDone(UMixInterface* TextureGraph, const FInvalidationDetails* Details);
 	void											OnViewportSettingsChanged();
 	void											OnMaterialMappingChanged();
@@ -255,11 +258,7 @@ private:
 	TSharedRef<SDockTab>							SpawnTab_Settings(const FSpawnTabArgs& Args);
 	TSharedRef<SDockTab>							SpawnTab_Errors(const FSpawnTabArgs& Args);
 	TSharedRef<SDockTab>							SpawnTab_TextureDetails(const FSpawnTabArgs& Args);
-	void 											DisplayNodeHistogramView(bool bShow3DPreview);
-	
-	void 											ToggleNodeHistogramView();
-	void 											TogglePaletteView();
-	void 											DisplayPaletteView(bool bShowSelectionPreview);
+
 	void											ExportAsUAsset();
 	FReply											OnExportClick();
 
@@ -274,9 +273,6 @@ private:
 	// FEditorUndoClient Interface
 	virtual void 									PostUndo(bool bSuccess) override;
 	virtual void 									PostRedo(bool bSuccess) override { PostUndo(bSuccess); }
-	
-	FORCEINLINE bool 								IsShowingNodeHistogramView() const { return bShowNodeHistogram; }
-	FORCEINLINE bool 								IsShowingPaletteView() const { return bShowPaletteView; };
 public:
 	/** Original Texture Graph */
 	TObjectPtr<UTextureGraph>						OriginalTextureGraph;
@@ -285,8 +281,6 @@ public:
 	TObjectPtr<UTextureGraph>						EditedTextureGraph;
 	TObjectPtr<UTG_EdGraph>							TG_EdGraph;
 	FDateTime										SessionStartTime;
-	bool 											bShowNodeHistogram = false;
-	bool 											bShowPaletteView = false;
 
 private:
 	TObjectPtr<UTG_Parameters>						TG_Parameters;
diff --git a/Engine/Plugins/Experimental/TextureGraph/Source/TextureGraphEditor/Private/TG_EditorCommands.cpp b/Engine/Plugins/Experimental/TextureGraph/Source/TextureGraphEditor/Private/TG_EditorCommands.cpp
--- a/Engine/Plugins/Experimental/TextureGraph/Source/TextureGraphEditor/Private/TG_EditorCommands.cpp
+++ b/Engine/Plugins/Experimental/TextureGraph/Source/TextureGraphEditor/Private/TG_EditorCommands.cpp
@@ -33,9 +33,6 @@ void FTG_EditorCommands::RegisterCommands()
 	UI_COMMAND(SetPreviewMeshFromSelection, "Mesh", "Sets the preview mesh based on the current content browser selection.", EUserInterfaceActionType::ToggleButton, FInputChord());
 	UI_COMMAND(TogglePreviewGrid, "Grid", "Toggles the preview pane's grid.", EUserInterfaceActionType::ToggleButton, FInputChord());
 	UI_COMMAND(TogglePreviewBackground, "Background", "Toggles the preview pane's background.", EUserInterfaceActionType::ToggleButton, FInputChord());
-	UI_COMMAND(TogglePaletteTab, "Palette", "Toggles the Palette window.", EUserInterfaceActionType::ToggleButton, FInputChord());
-	UI_COMMAND(ToggleNodeHistogramTab, "Node Histogram", "Toggles the Node Histogram window.", EUserInterfaceActionType::ToggleButton, FInputChord() );
-	
 }
 
 #undef LOCTEXT_NAMESPACE
diff --git a/Engine/Plugins/Experimental/TextureGraph/Source/TextureGraphEditor/Private/TG_EditorCommands.h b/Engine/Plugins/Experimental/TextureGraph/Source/TextureGraphEditor/Private/TG_EditorCommands.h
--- a/Engine/Plugins/Experimental/TextureGraph/Source/TextureGraphEditor/Private/TG_EditorCommands.h
+++ b/Engine/Plugins/Experimental/TextureGraph/Source/TextureGraphEditor/Private/TG_EditorCommands.h
@@ -14,28 +14,26 @@ public:
 	virtual void RegisterCommands() override;
 	// ~End TCommands<> interface
 
-	TSharedPtr<FUICommandInfo>									RunGraph;
-	TSharedPtr<FUICommandInfo>									AutoUpdateGraph;
-	TSharedPtr<FUICommandInfo>									LogGraph;
-	TSharedPtr<FUICommandInfo>									ExportAsUAsset;
-	TSharedPtr<FUICommandInfo>									TogglePaletteTab;
-	TSharedPtr<FUICommandInfo>									ToggleNodeHistogramTab;
+	TSharedPtr<FUICommandInfo> RunGraph;
+	TSharedPtr<FUICommandInfo> AutoUpdateGraph;
+	TSharedPtr<FUICommandInfo> LogGraph;
+	TSharedPtr<FUICommandInfo> ExportAsUAsset;
 
 	/**
 	 * Node Actions
 	 */
-	TSharedPtr< FUICommandInfo >								ConvertInputParameterToConstant;		/// Convert an input param to from constant
-	TSharedPtr< FUICommandInfo >								ConvertInputParameterFromConstant;		/// Convert an input param to from constant
+	TSharedPtr< FUICommandInfo > ConvertInputParameterToConstant;		/// Convert an input param to from constant
+	TSharedPtr< FUICommandInfo > ConvertInputParameterFromConstant;		/// Convert an input param to from constant
 
 	/**
 	 * Viewport Commands
 	 */
-	TSharedPtr< FUICommandInfo >								SetCylinderPreview;				/// Sets the preview mesh to a cylinder
-	TSharedPtr< FUICommandInfo >								SetSpherePreview;				/// Sets the preview mesh to a sphere
-	TSharedPtr< FUICommandInfo >								SetPlanePreview;				/// Sets the preview mesh to a plane
-	TSharedPtr< FUICommandInfo >								SetCubePreview;					/// Sets the preview mesh to a cube
-	TSharedPtr< FUICommandInfo >								SetPreviewMeshFromSelection;	// Sets the preview mesh to the current selection in the level editor
-	TSharedPtr< FUICommandInfo >								TogglePreviewGrid;				/// Toggles the preview pane's grid
-	TSharedPtr< FUICommandInfo >								TogglePreviewBackground;		/// Toggles the preview pane's background
+	TSharedPtr< FUICommandInfo > SetCylinderPreview;				/// Sets the preview mesh to a cylinder
+	TSharedPtr< FUICommandInfo > SetSpherePreview;				/// Sets the preview mesh to a sphere
+	TSharedPtr< FUICommandInfo > SetPlanePreview;				/// Sets the preview mesh to a plane
+	TSharedPtr< FUICommandInfo > SetCubePreview;					/// Sets the preview mesh to a cube
+	TSharedPtr< FUICommandInfo > SetPreviewMeshFromSelection;	// Sets the preview mesh to the current selection in the level editor
+	TSharedPtr< FUICommandInfo > TogglePreviewGrid;				/// Toggles the preview pane's grid
+	TSharedPtr< FUICommandInfo > TogglePreviewBackground;		/// Toggles the preview pane's background
 
 };
