Texture Graph's First Output Node Is unchecked by default when exporting

diff --git a/Engine/Plugins/Experimental/TextureGraph/Source/TextureGraphEditor/Private/STG_OutputSelectionDlg.cpp b/Engine/Plugins/Experimental/TextureGraph/Source/TextureGraphEditor/Private/STG_OutputSelectionDlg.cpp
--- a/Engine/Plugins/Experimental/TextureGraph/Source/TextureGraphEditor/Private/STG_OutputSelectionDlg.cpp
+++ b/Engine/Plugins/Experimental/TextureGraph/Source/TextureGraphEditor/Private/STG_OutputSelectionDlg.cpp
@@ -119,11 +119,21 @@ void STG_OutputSelectionDlg::AddExportItems()
 					ThumbnailWidget = SNew(SColorBlock)
 					.Color(ColorValue);
 				}
-				else
+				else if(Variant.IsVector())
 				{
-					TargetExpression->SetExport(false);
-					continue;
+					FVector4f Vector;
+					Pin->GetValue(Vector);
+					ThumbnailWidget = SNew(SColorBlock)
+					.Color(Vector);
 				}
+				else if (Variant.IsScalar())
+				{
+					float Scalar;
+					Pin->GetValue(Scalar);
+					ThumbnailWidget = SNew(SColorBlock)
+						.Color(FLinearColor(Scalar, Scalar, Scalar, 1));
+				}
+
 				ScrollBox->AddSlot()
 				.Padding(5)
 				[
