Fixed crash when switching the shape type in shape node.
Fixed changing texture output settings for shape node was reverting back

diff --git a/Engine/Plugins/Experimental/TextureGraph/Source/TextureGraph/Private/Expressions/Procedural/TG_Expression_Shape.cpp b/Engine/Plugins/Experimental/TextureGraph/Source/TextureGraph/Private/Expressions/Procedural/TG_Expression_Shape.cpp
--- a/Engine/Plugins/Experimental/TextureGraph/Source/TextureGraph/Private/Expressions/Procedural/TG_Expression_Shape.cpp
+++ b/Engine/Plugins/Experimental/TextureGraph/Source/TextureGraph/Private/Expressions/Procedural/TG_Expression_Shape.cpp
@@ -8,7 +8,7 @@ void UTG_Expression_Shape::Evaluate(FTG_EvaluationContext* InContext)
 {
 	Super::Evaluate(InContext);
 
-	auto DesiredDescriptor = Output.Descriptor;
+	FTG_TextureDescriptor DesiredDescriptor = Output.Descriptor;
 
 	/// If it's auto then make it grey scale as we don't need more than a single channel output
 	if (DesiredDescriptor.TextureFormat == ETG_TextureFormat::Auto)
diff --git a/Engine/Plugins/Experimental/TextureGraph/Source/TextureGraph/Private/TG_Node.cpp b/Engine/Plugins/Experimental/TextureGraph/Source/TextureGraph/Private/TG_Node.cpp
--- a/Engine/Plugins/Experimental/TextureGraph/Source/TextureGraph/Private/TG_Node.cpp
+++ b/Engine/Plugins/Experimental/TextureGraph/Source/TextureGraph/Private/TG_Node.cpp
@@ -137,16 +137,9 @@ void UTG_Node::OnExpressionChangedWithoutVar(const FPropertyChangedEvent& Proper
 
 	if (PropertyChangedEvent.Property)
 	{
-		auto PinId = GetPinId(PropertyChangedEvent.GetPropertyName());
-
-		// we check the property against the pin name. In most cases, the value modified is actually the property.
-		// In the case of FLinear color, we can modify R, G, B and A values independently too.
-		// The property will be R, G, B or A which will not match the pin name.
-		// To fix it, If we don't find the pin for the property changed, we will also try to find it in the Member property.
-		if (!PinId.IsValid())
-		{
-			PinId = GetPinId(PropertyChangedEvent.GetMemberPropertyName());
-		}
+		// MemberProperty name is what we're interested in, the expression's member property has triggered an event
+		// and we want to find the matching pin and take the value in from the expression's member
+		FTG_Id PinId = GetPinId(PropertyChangedEvent.GetMemberPropertyName());
 
 		if (PinId.IsValid())
 		{
diff --git a/Engine/Plugins/Experimental/TextureGraph/Source/TextureGraphEditor/Private/EdGraph/TG_EdGraphNode.cpp b/Engine/Plugins/Experimental/TextureGraph/Source/TextureGraphEditor/Private/EdGraph/TG_EdGraphNode.cpp
--- a/Engine/Plugins/Experimental/TextureGraph/Source/TextureGraphEditor/Private/EdGraph/TG_EdGraphNode.cpp
+++ b/Engine/Plugins/Experimental/TextureGraph/Source/TextureGraphEditor/Private/EdGraph/TG_EdGraphNode.cpp
@@ -326,6 +326,8 @@ FName UTG_EdGraphNode::GetPinCategory(UTG_Pin* Pin, TWeakObjectPtr<UObject>& Sub
 
 void UTG_EdGraphNode::ReconstructNode()
 {
+	const FName SelectedPinName = (SelectedPin) ? SelectedPin->PinName : FName();
+	UEdGraphPin* NewSelectedPin = nullptr;
 	// Store copy of old pins
 	TArray<UEdGraphPin*> OldPins = MoveTemp(Pins);
 	Pins.Reset();
@@ -347,6 +349,11 @@ void UTG_EdGraphNode::ReconstructNode()
 			{
 				(*NewPin)->MovePersistentDataFromOldPin(*OldPin);
 			}
+
+			if ((*NewPin)->PinName == SelectedPinName)
+			{
+				NewSelectedPin = (*NewPin);
+			}
 		}
 	}
 	
@@ -360,14 +367,13 @@ void UTG_EdGraphNode::ReconstructNode()
 	OnNodeReconstructDelegate.Broadcast();
 
 	// if no selected pin, we set it to the first output pin
-	if (GetSelectedPin() == nullptr)
+	if (!NewSelectedPin && GetOutputPins().Num() > 0)
 	{
-		if (GetOutputPins().Num() > 0)
-		{
-			UTG_EdGraph* TSEdGraph = Cast<UTG_EdGraph>(GetGraph());
-			TSEdGraph->PinSelectionManager.UpdateSelection(GetOutputPins()[0]);
-		}
+		NewSelectedPin = GetOutputPins()[0];
 	}
+	
+	UTG_EdGraph* TSEdGraph = Cast<UTG_EdGraph>(GetGraph());
+	TSEdGraph->PinSelectionManager.UpdateSelection(NewSelectedPin);
 }
 
 FString UTG_EdGraphNode::GetTitleDetail()
