Address crash when connecting a node  to output

diff --git a/Engine/Plugins/Experimental/TextureGraph/Source/TextureGraphEditor/Private/EdGraph/TG_PinSelectionManager.cpp b/Engine/Plugins/Experimental/TextureGraph/Source/TextureGraphEditor/Private/EdGraph/TG_PinSelectionManager.cpp
--- a/Engine/Plugins/Experimental/TextureGraph/Source/TextureGraphEditor/Private/EdGraph/TG_PinSelectionManager.cpp
+++ b/Engine/Plugins/Experimental/TextureGraph/Source/TextureGraphEditor/Private/EdGraph/TG_PinSelectionManager.cpp
@@ -12,26 +12,42 @@ void FTG_PinSelectionManager::UpdateSelection(UEdGraphPin* Pin)
 {
 	SelectedItems.Empty();
 
-	SelectedItems.Add(Pin);
-	SelectPin(Pin, true);
+	if (Pin)
+	{
+		SelectedItems.Add(Pin);
+		{
+			auto Node = Cast<UTG_EdGraphNode>(Pin->GetOwningNode());
+			check(Node);
+			Node->SelectPin(Pin, true);
+		}
+	}
 
 	OnPinSelectionUpdated.Broadcast(Pin);
 }
 
 void FTG_PinSelectionManager::ClearPinsForNonSelectedNodes(TArray<UObject*> SelectedNodes)
 {
-	for (int32 i = SelectedItems.Num() - 1; i >= 0; --i)
+	// No Selected nodes means clear all
+	if (SelectedNodes.IsEmpty())
 	{
-		if (!SelectedNodes.Contains(Cast<UTG_EdGraphNode>(SelectedItems[i]->GetOwningNode())->GetDetailsObject()))
+		SelectedItems.Empty();
+	}
+	else
+	{
+		for (int32 i = SelectedItems.Num() - 1; i >= 0; --i)
 		{
-			SelectedItems.RemoveAt(i);
+			UEdGraphPin* SelectedItem = SelectedItems[i];
+			if (SelectedItem)
+			{
+				UTG_EdGraphNode* SelectedPinNode = Cast<UTG_EdGraphNode>(SelectedItem->GetOwningNode());
+				if (SelectedPinNode)
+				{
+					if (!SelectedNodes.Contains(Cast<UTG_EdGraphNode>(SelectedItems[i]->GetOwningNode())->GetDetailsObject()))
+					{
+						SelectedItems.RemoveAt(i);
+					}
+				}
+			}
 		}
 	}
 }
-
-void FTG_PinSelectionManager::SelectPin(UEdGraphPin* Pin,bool IsSelected)
-{
-	auto Node = Cast<UTG_EdGraphNode>(Pin->GetOwningNode());
-	check(Node);
-	Node->SelectPin(Pin, IsSelected);
-}
\ No newline at end of file
diff --git a/Engine/Plugins/Experimental/TextureGraph/Source/TextureGraphEditor/Private/TG_Editor.cpp b/Engine/Plugins/Experimental/TextureGraph/Source/TextureGraphEditor/Private/TG_Editor.cpp
--- a/Engine/Plugins/Experimental/TextureGraph/Source/TextureGraphEditor/Private/TG_Editor.cpp
+++ b/Engine/Plugins/Experimental/TextureGraph/Source/TextureGraphEditor/Private/TG_Editor.cpp
@@ -1864,9 +1864,15 @@ bool FTG_Editor::IsOutputNode(UTG_Node* InNode)
 
 void FTG_Editor::OnPinSelectionUpdated(UEdGraphPin* Pin)
 {
-	check(Pin);
-	//When pin get updated we need to update its Node
-	GraphEditorWidget->GetGraphPanel()->SelectionManager.SelectSingleNode(Pin->GetOwningNode());
+	if (Pin)
+	{
+		//When pin get updated we need to update its Node
+		GraphEditorWidget->GetGraphPanel()->SelectionManager.SelectSingleNode(Pin->GetOwningNode());
+	}
+	else
+	{
+		GraphEditorWidget->GetGraphPanel()->SelectionManager.ClearSelectionSet();
+	}
 }
 
 void FTG_Editor::RefreshPreviewViewport()
diff --git a/Engine/Plugins/Experimental/TextureGraph/Source/TextureGraphEditor/Public/EdGraph/TG_PinSelectionManager.h b/Engine/Plugins/Experimental/TextureGraph/Source/TextureGraphEditor/Public/EdGraph/TG_PinSelectionManager.h
--- a/Engine/Plugins/Experimental/TextureGraph/Source/TextureGraphEditor/Public/EdGraph/TG_PinSelectionManager.h
+++ b/Engine/Plugins/Experimental/TextureGraph/Source/TextureGraphEditor/Public/EdGraph/TG_PinSelectionManager.h
@@ -33,7 +33,6 @@ public:
 	FOnPinSelectionUpdated OnPinSelectionUpdated;
 
 private:
-	void SelectPin(UEdGraphPin* Pin, bool IsSelected);
 
 	TArray<UEdGraphPin*> SelectedItems;
 };
