Fixing issues related to the SelectedPin pointer in the UTG_EdGRaphNode

When doing enough undo redo, the field SelectedPIn is invalid in the EdGRaphNode class
We are using the getter to do some basic checks and return nullptr in case the pointer is obviously wrong

removed uncessary state kept in the FTG_PinSelectionManager'

diff --git a/Engine/Plugins/Experimental/TextureGraph/Source/TextureGraphEditor/Private/EdGraph/TG_EdGraphNode.cpp b/Engine/Plugins/Experimental/TextureGraph/Source/TextureGraphEditor/Private/EdGraph/TG_EdGraphNode.cpp
--- a/Engine/Plugins/Experimental/TextureGraph/Source/TextureGraphEditor/Private/EdGraph/TG_EdGraphNode.cpp
+++ b/Engine/Plugins/Experimental/TextureGraph/Source/TextureGraphEditor/Private/EdGraph/TG_EdGraphNode.cpp
@@ -72,11 +72,12 @@ void UTG_EdGraphNode::SelectPin(UEdGraphPin* Pin, bool IsSelected)
 		OnPinSelectionChangeDelegate.Broadcast(SelectedPin);
 }
 
-const FGuid UTG_EdGraphNode::GetSelectedPinFGuid() const
+const UEdGraphPin* UTG_EdGraphNode::GetSelectedPin() const
 {
-	if (SelectedPin)
-		return SelectedPin->PinId;
-	else return FGuid();
+	if (SelectedPin && Pins.Contains(SelectedPin))
+		return SelectedPin;
+	else
+		return nullptr;
 }
 
 void UTG_EdGraphNode::GetNodeContextMenuActions(UToolMenu* Menu, class UGraphNodeContextMenuContext* Context) const
@@ -112,12 +113,17 @@ void UTG_EdGraphNode::GetNodeContextMenuActions(UToolMenu* Menu, class UGraphNod
 			FUIAction(
 				FExecuteAction::CreateLambda([Pin = Context->Pin, Node = Context->Node, this]
 				{
-					UTG_EdGraph* TSEdGraph = Cast<UTG_EdGraph>(Node->GetGraph());
-					TSEdGraph->PinSelectionManager.UpdateSelection(const_cast<UEdGraphPin*>(Pin));
+					const UTG_EdGraphNode* TGEdNode = Cast<UTG_EdGraphNode>(Node);
+					if (TGEdNode)
+					{
+						const_cast<UTG_EdGraphNode*>(TGEdNode)->SelectPin(const_cast<UEdGraphPin*>(Pin), true); // assign the current selected pin ion the node
+						UTG_EdGraph* TSEdGraph = Cast<UTG_EdGraph>(Node->GetGraph());
+						TSEdGraph->PinSelectionManager.UpdateSelection(const_cast<UEdGraphPin*>(Pin));
+					}
 				}),
 				FCanExecuteAction::CreateLambda([Pin = Context->Pin, this]
 				{
-					return Pin != SelectedPin;
+					return Pin != GetSelectedPin();
 				})));
 		}
 	}
@@ -326,7 +332,11 @@ FName UTG_EdGraphNode::GetPinCategory(UTG_Pin* Pin, TWeakObjectPtr<UObject>& Sub
 
 void UTG_EdGraphNode::ReconstructNode()
 {
-	const FName SelectedPinName = (SelectedPin) ? SelectedPin->PinName : FName();
+	const UEdGraphPin* SelectedPinPtr = GetSelectedPin();
+	const FName SelectedPinName = (SelectedPinPtr) ? SelectedPinPtr->PinName : FName();
+	// forget SelectedPin, eventually reassign below via the PinSelectionManager
+	SelectedPin = nullptr;
+
 	UEdGraphPin* NewSelectedPin = nullptr;
 	// Store copy of old pins
 	TArray<UEdGraphPin*> OldPins = MoveTemp(Pins);
@@ -372,6 +382,7 @@ void UTG_EdGraphNode::ReconstructNode()
 		NewSelectedPin = GetOutputPins()[0];
 	}
 	
+	SelectPin(NewSelectedPin, true); // assign the current selected pin ion the node
 	UTG_EdGraph* TSEdGraph = Cast<UTG_EdGraph>(GetGraph());
 	TSEdGraph->PinSelectionManager.UpdateSelection(NewSelectedPin);
 }
@@ -383,13 +394,15 @@ FString UTG_EdGraphNode::GetTitleDetail()
 	TArray<const UTG_Pin*> OutPins;
 	GetNode()->GetOutputPins(OutPins);
 	
+	const UEdGraphPin* SelectedPinPtr = GetSelectedPin();
+
 	for (auto Pin : OutPins)
 	{
 		const UTG_Pin* TGPin = Pin;
 		
-		if (SelectedPin)
+		if (SelectedPinPtr)
 		{
-			TGPin = Schema->GetTGPinFromEdPin(SelectedPin);
+			TGPin = Schema->GetTGPinFromEdPin(SelectedPinPtr);
 		}
 
 		// Only interested by the texture type pin
diff --git a/Engine/Plugins/Experimental/TextureGraph/Source/TextureGraphEditor/Private/EdGraph/TG_PinSelectionManager.cpp b/Engine/Plugins/Experimental/TextureGraph/Source/TextureGraphEditor/Private/EdGraph/TG_PinSelectionManager.cpp
--- a/Engine/Plugins/Experimental/TextureGraph/Source/TextureGraphEditor/Private/EdGraph/TG_PinSelectionManager.cpp
+++ b/Engine/Plugins/Experimental/TextureGraph/Source/TextureGraphEditor/Private/EdGraph/TG_PinSelectionManager.cpp
@@ -10,44 +10,5 @@ void FTG_PinSelectionManager::AddReferencedObjects(FReferenceCollector& Collecto
 
 void FTG_PinSelectionManager::UpdateSelection(UEdGraphPin* Pin)
 {
-	SelectedItems.Empty();
-
-	if (Pin)
-	{
-		SelectedItems.Add(Pin);
-		{
-			auto Node = Cast<UTG_EdGraphNode>(Pin->GetOwningNode());
-			check(Node);
-			Node->SelectPin(Pin, true);
-		}
-	}
-
 	OnPinSelectionUpdated.Broadcast(Pin);
 }
-
-void FTG_PinSelectionManager::ClearPinsForNonSelectedNodes(TArray<UObject*> SelectedNodes)
-{
-	// No Selected nodes means clear all
-	if (SelectedNodes.IsEmpty())
-	{
-		SelectedItems.Empty();
-	}
-	else
-	{
-		for (int32 i = SelectedItems.Num() - 1; i >= 0; --i)
-		{
-			UEdGraphPin* SelectedItem = SelectedItems[i];
-			if (SelectedItem)
-			{
-				UTG_EdGraphNode* SelectedPinNode = Cast<UTG_EdGraphNode>(SelectedItem->GetOwningNode());
-				if (SelectedPinNode)
-				{
-					if (!SelectedNodes.Contains(Cast<UTG_EdGraphNode>(SelectedItems[i]->GetOwningNode())->GetDetailsObject()))
-					{
-						SelectedItems.RemoveAt(i);
-					}
-				}
-			}
-		}
-	}
-}
diff --git a/Engine/Plugins/Experimental/TextureGraph/Source/TextureGraphEditor/Private/TG_Editor.cpp b/Engine/Plugins/Experimental/TextureGraph/Source/TextureGraphEditor/Private/TG_Editor.cpp
--- a/Engine/Plugins/Experimental/TextureGraph/Source/TextureGraphEditor/Private/TG_Editor.cpp
+++ b/Engine/Plugins/Experimental/TextureGraph/Source/TextureGraphEditor/Private/TG_Editor.cpp
@@ -1197,8 +1197,6 @@ void FTG_Editor::OnSelectedNodesChanged(const TSet<class UObject*>& NewSelection
 		}
 	}
 
-	TG_EdGraph->PinSelectionManager.ClearPinsForNonSelectedNodes(SelectedObjects);
-
 	GetDetailView()->SetObjects(SelectedObjects, true);
 	FocusDetailsPanel();
 
diff --git a/Engine/Plugins/Experimental/TextureGraph/Source/TextureGraphEditor/Public/EdGraph/TG_EdGraphNode.h b/Engine/Plugins/Experimental/TextureGraph/Source/TextureGraphEditor/Public/EdGraph/TG_EdGraphNode.h
--- a/Engine/Plugins/Experimental/TextureGraph/Source/TextureGraphEditor/Public/EdGraph/TG_EdGraphNode.h
+++ b/Engine/Plugins/Experimental/TextureGraph/Source/TextureGraphEditor/Public/EdGraph/TG_EdGraphNode.h
@@ -89,9 +89,8 @@ public:
 	*/
 	void SelectPin(UEdGraphPin* Pin, bool IsSelected);
 
-	const UEdGraphPin* GetSelectedPin() const { return SelectedPin; }
+	const UEdGraphPin* GetSelectedPin() const;
 
-	const FGuid GetSelectedPinFGuid() const;
 
 	TArray<UEdGraphPin*> GetOutputPins() const;
 	TArray<UEdGraphPin*> GetTextureOutputPins() const;
@@ -115,6 +114,6 @@ protected:
 		TObjectPtr<UTG_Node> Node;
 
 private:
-	UEdGraphPin* SelectedPin;
+	UEdGraphPin* SelectedPin = nullptr;
 };
 
diff --git a/Engine/Plugins/Experimental/TextureGraph/Source/TextureGraphEditor/Public/EdGraph/TG_PinSelectionManager.h b/Engine/Plugins/Experimental/TextureGraph/Source/TextureGraphEditor/Public/EdGraph/TG_PinSelectionManager.h
--- a/Engine/Plugins/Experimental/TextureGraph/Source/TextureGraphEditor/Public/EdGraph/TG_PinSelectionManager.h
+++ b/Engine/Plugins/Experimental/TextureGraph/Source/TextureGraphEditor/Public/EdGraph/TG_PinSelectionManager.h
@@ -27,12 +27,9 @@ public:
 	}
 
 	void UpdateSelection(UEdGraphPin* Pin);
-	void ClearPinsForNonSelectedNodes(TArray<UObject*> SelectedNodes);
 
 	DECLARE_MULTICAST_DELEGATE_OneParam(FOnPinSelectionUpdated, UEdGraphPin* /*Pin*/)
 	FOnPinSelectionUpdated OnPinSelectionUpdated;
 
 private:
-
-	TArray<UEdGraphPin*> SelectedItems;
 };
