Minor fixes:
- Fixed height/width dependency on each other in some places
- Typo fix

diff --git a/Engine/Plugins/Experimental/TextureGraph/Source/TextureGraph/Private/Expressions/Procedural/TG_Expression_Noise.cpp b/Engine/Plugins/Experimental/TextureGraph/Source/TextureGraph/Private/Expressions/Procedural/TG_Expression_Noise.cpp
--- a/Engine/Plugins/Experimental/TextureGraph/Source/TextureGraph/Private/Expressions/Procedural/TG_Expression_Noise.cpp
+++ b/Engine/Plugins/Experimental/TextureGraph/Source/TextureGraph/Private/Expressions/Procedural/TG_Expression_Noise.cpp
@@ -43,9 +43,14 @@ void UTG_Expression_Noise::Evaluate(FTG_EvaluationContext* InContext)
 
 	Desc.DefaultValue = FLinearColor(1.0f, 1.0f, 1.0f, 1.0f);
 
-	if (Desc.Width <= 0 || Desc.Height <= 0)
+	if (Desc.Width <= 0)
 	{
-		Desc.Width = Desc.Height = DefaultSize;
+		Desc.Width = DefaultSize;
+	}
+	
+	if(Desc.Height <= 0)
+	{
+		Desc.Height = DefaultSize;
 	}
 
 	Output = RenderJob->InitResult(Name, &Desc);
diff --git a/Engine/Plugins/Experimental/TextureGraph/Source/TextureGraph/Private/Expressions/Procedural/TG_Expression_Pattern.cpp b/Engine/Plugins/Experimental/TextureGraph/Source/TextureGraph/Private/Expressions/Procedural/TG_Expression_Pattern.cpp
--- a/Engine/Plugins/Experimental/TextureGraph/Source/TextureGraph/Private/Expressions/Procedural/TG_Expression_Pattern.cpp
+++ b/Engine/Plugins/Experimental/TextureGraph/Source/TextureGraph/Private/Expressions/Procedural/TG_Expression_Pattern.cpp
@@ -65,9 +65,14 @@ void UTG_Expression_Pattern::Evaluate(FTG_EvaluationContext* InContext)
 
 	Desc.DefaultValue = FLinearColor(1.0f, 1.0f, 1.0f, 1.0f);
 
-	if (Desc.Width <= 0 || Desc.Height <= 0)
+	if (Desc.Width <= 0)
 	{
-		Desc.Width = Desc.Height = DefaultSize;
+		Desc.Width = DefaultSize;
+	}
+	
+	if(Desc.Height <= 0)
+	{
+		Desc.Height = DefaultSize;
 	}
 
 	Output = RenderJob->InitResult(Name, &Desc);
diff --git a/Engine/Plugins/Experimental/TextureGraph/Source/TextureGraph/Public/Expressions/Procedural/TG_Expression_Shape.h b/Engine/Plugins/Experimental/TextureGraph/Source/TextureGraph/Public/Expressions/Procedural/TG_Expression_Shape.h
--- a/Engine/Plugins/Experimental/TextureGraph/Source/TextureGraph/Public/Expressions/Procedural/TG_Expression_Shape.h
+++ b/Engine/Plugins/Experimental/TextureGraph/Source/TextureGraph/Public/Expressions/Procedural/TG_Expression_Shape.h
@@ -59,7 +59,7 @@ public:
 	UPROPERTY(EditAnywhere, BlueprintReadWrite, Category = NoCategory, meta = (TGType = "TG_Setting", UIMin = "-180", ClampMin = "-180", UIMax = "180", ClampMax = "180", Delta = "1", Units = "Degrees", EditConditionHides))
 	float								Orientation = 0.0f;
 
-	// Rounding ts about the 2D footprint of the shape generated, 0 means the perfect shape, 1 means the most rounded version of the shape
+	// Rounding is about the 2D footprint of the shape generated, 0 means the perfect shape, 1 means the most rounded version of the shape
 	UPROPERTY(EditAnywhere, BlueprintReadWrite, Category = NoCategory, meta = (TGType = "TG_Setting", ClampMin = 0, ClampMax = 1, EditConditionHides))
 	float								Rounding = 0.0f;
 
diff --git a/Engine/Plugins/Experimental/TextureGraph/Source/TextureGraphEngine/Transform/Expressions/T_MaterialIDMask.cpp b/Engine/Plugins/Experimental/TextureGraph/Source/TextureGraphEngine/Transform/Expressions/T_MaterialIDMask.cpp
--- a/Engine/Plugins/Experimental/TextureGraph/Source/TextureGraphEngine/Transform/Expressions/T_MaterialIDMask.cpp
+++ b/Engine/Plugins/Experimental/TextureGraph/Source/TextureGraphEngine/Transform/Expressions/T_MaterialIDMask.cpp
@@ -61,12 +61,16 @@ TiledBlobPtr T_MaterialIDMask::Create(MixUpdateCyclePtr InCycle, TiledBlobPtr In
 
 	Desc.DefaultValue = FLinearColor(1.0f, 1.0f, 1.0f, 1.0f);
 
-	if (Desc.Width <= 0 || Desc.Height <= 0)
+	if (Desc.Width <= 0) 
 	{
 		Desc.Width = InMaterialIDTexture->GetWidth();
-		Desc.Height = InMaterialIDTexture->GetHeight();
 	}
 
+	if(Desc.Height <= 0)
+	{
+		Desc.Height = InMaterialIDTexture->GetHeight();
+	}
+	
 	auto Result = JobPtr->InitResult(MaterialIDMaskMat->GetName(), &Desc);
 
 	InCycle->AddJob(InTargetId, std::move(JobPtr));
