Fixed issues related to material ID:

Deleting MaterialID or any input connected with it and undo, does not render the selected color in MaterialID node on Node preview and 3D preview

Undoing after enabling a color in materialID does not update the rendering

diff --git a/Engine/Plugins/Experimental/TextureGraph/Source/TextureGraph/Private/Expressions/Utilities/TG_Expression_MaterialID.cpp b/Engine/Plugins/Experimental/TextureGraph/Source/TextureGraph/Private/Expressions/Utilities/TG_Expression_MaterialID.cpp
--- a/Engine/Plugins/Experimental/TextureGraph/Source/TextureGraph/Private/Expressions/Utilities/TG_Expression_MaterialID.cpp
+++ b/Engine/Plugins/Experimental/TextureGraph/Source/TextureGraph/Private/Expressions/Utilities/TG_Expression_MaterialID.cpp
@@ -40,6 +40,8 @@ void UTG_Expression_MaterialID::Initialize()
 
 void UTG_Expression_MaterialID::UpdateActiveColors()
 {
+	Modify();
+	
 	ActiveColors.Empty();
 
 	ActiveColors.SetNum(128);
diff --git a/Engine/Plugins/Experimental/TextureGraph/Source/TextureGraph/Private/TG_Node.cpp b/Engine/Plugins/Experimental/TextureGraph/Source/TextureGraph/Private/TG_Node.cpp
--- a/Engine/Plugins/Experimental/TextureGraph/Source/TextureGraph/Private/TG_Node.cpp
+++ b/Engine/Plugins/Experimental/TextureGraph/Source/TextureGraph/Private/TG_Node.cpp
@@ -52,9 +52,12 @@ void UTG_Node::OnPostUndo()
 		// For now, just copying all the values from the expression.
 		for(const auto& Pin : Pins)
 		{
-			auto Arg = Pin->GetArgument();
-			Pin->EditSelfVar()->CopyFrom(GetExpression(), Arg);	
-
+			// exclude the case of input textures 
+			if(!(Pin->IsInput() && Pin->IsArgTexture()))
+			{
+				auto Arg = Pin->GetArgument();
+				Pin->EditSelfVar()->CopyFrom(GetExpression(), Arg);	
+			}
 		}
 
 		Signature = GetExpression()->GetSignature();
diff --git a/Engine/Plugins/Experimental/TextureGraph/Source/TextureGraph/Public/Expressions/Utilities/TG_Expression_MaterialID.h b/Engine/Plugins/Experimental/TextureGraph/Source/TextureGraph/Public/Expressions/Utilities/TG_Expression_MaterialID.h
--- a/Engine/Plugins/Experimental/TextureGraph/Source/TextureGraph/Public/Expressions/Utilities/TG_Expression_MaterialID.h
+++ b/Engine/Plugins/Experimental/TextureGraph/Source/TextureGraph/Public/Expressions/Utilities/TG_Expression_MaterialID.h
@@ -28,8 +28,10 @@ public:
 	UPROPERTY()
 	FMaterialIDCollection MaterialIDInfoCollection;
 
+	UPROPERTY(Transient)
 	TArray<FLinearColor> ActiveColors;
 
+	UPROPERTY(Transient)
 	int32 ActiveColorsCount;
 	
 	// The output of the node, which is the loaded texture asset
diff --git a/Engine/Plugins/Experimental/TextureGraph/Source/TextureGraphEditor/Private/Pins/STG_GraphMaterialIdMaskInfos.h b/Engine/Plugins/Experimental/TextureGraph/Source/TextureGraphEditor/Private/Pins/STG_GraphMaterialIdMaskInfos.h
--- a/Engine/Plugins/Experimental/TextureGraph/Source/TextureGraphEditor/Private/Pins/STG_GraphMaterialIdMaskInfos.h
+++ b/Engine/Plugins/Experimental/TextureGraph/Source/TextureGraphEditor/Private/Pins/STG_GraphMaterialIdMaskInfos.h
@@ -152,6 +152,8 @@ protected:
 	
 	void HandleOnMaterialIdClicked(const int32 Id)
 	{
+		const FScopedTransaction Transaction(NSLOCTEXT("STG_MaterialIdMaskInfo", "Enabled", "Edit Is Enabled"));
+		
 		UTG_Expression_MaterialID* MaterialIdExpression = GetMaterialIDExpression();
 
 		const bool NewValue = !MaterialIdExpression->MaterialIDMaskInfos[Id].bIsEnabled;
