Fixed RGBA_half texture types

#virtualized

diff --git a/Engine/Plugins/Experimental/TextureGraph/Source/TextureGraphEngine/Data/RawBuffer.cpp b/Engine/Plugins/Experimental/TextureGraph/Source/TextureGraphEngine/Data/RawBuffer.cpp
--- a/Engine/Plugins/Experimental/TextureGraph/Source/TextureGraphEngine/Data/RawBuffer.cpp
+++ b/Engine/Plugins/Experimental/TextureGraph/Source/TextureGraphEngine/Data/RawBuffer.cpp
@@ -223,7 +223,7 @@ EPixelFormat BufferDescriptor::BufferPixelFormat(BufferFormat InFormat, uint32 I
 		case BufferFormat::Float:
 			return EPixelFormat::PF_A32B32G32R32F;
 		case BufferFormat::Half:
-			return EPixelFormat::PF_A16B16G16R16;
+			return EPixelFormat::PF_FloatRGBA;
 		case BufferFormat::Int:
 			return EPixelFormat::PF_R32G32B32A32_UINT;
 		case BufferFormat::Short:
