Fix Thumbnail for Exported Texture doesn't show

diff --git a/Engine/Plugins/Experimental/TextureGraph/Source/TextureGraphEngine/Export/TextureExporter.cpp b/Engine/Plugins/Experimental/TextureGraph/Source/TextureGraphEngine/Export/TextureExporter.cpp
--- a/Engine/Plugins/Experimental/TextureGraph/Source/TextureGraphEngine/Export/TextureExporter.cpp
+++ b/Engine/Plugins/Experimental/TextureGraph/Source/TextureGraphEngine/Export/TextureExporter.cpp
@@ -20,7 +20,8 @@
 #include "FileHelpers.h"
 #include "AssetToolsModule.h"
 #endif
-#include <Misc/PackageName.h>
+#include <Misc/PackageName.h> 
+#include "TextureCompiler.h"
 DEFINE_LOG_CATEGORY(ExportLogs);
 bool TextureExporter::IsPackageNameValid(FString Path, FString AssetName)
 {
@@ -184,6 +185,8 @@ AsyncInt TextureExporter::ExportRawAsUAsset(RawBufferPtr RawObj,const FExportMap
 		NewTexture->Source.Init(WidthToUse, HeightToUse, 1,1, SourceFormat, RawBufferData);
 		NewTexture->UpdateResource();
 		FAssetRegistryModule::AssetCreated(NewTexture);
+		Package->FullyLoad();
+		FTextureCompilingManager::Get().FinishCompilation({NewTexture});
 		bool success = UEditorLoadingAndSavingUtils::SavePackages({ Package }, false);
 		Promise.set_value(success);
 #else
