Fixed Undo redo crash for nodes that have dynamic signatures and cause regenerate upon type change, Like Material, blur, and shape nodes.

diff --git a/Engine/Plugins/Experimental/TextureGraph/Source/TextureGraphEditor/Private/EdGraph/TG_EdGraphNode.cpp b/Engine/Plugins/Experimental/TextureGraph/Source/TextureGraphEditor/Private/EdGraph/TG_EdGraphNode.cpp
--- a/Engine/Plugins/Experimental/TextureGraph/Source/TextureGraphEditor/Private/EdGraph/TG_EdGraphNode.cpp
+++ b/Engine/Plugins/Experimental/TextureGraph/Source/TextureGraphEditor/Private/EdGraph/TG_EdGraphNode.cpp
@@ -332,6 +332,7 @@ FName UTG_EdGraphNode::GetPinCategory(UTG_Pin* Pin, TWeakObjectPtr<UObject>& Sub
 
 void UTG_EdGraphNode::ReconstructNode()
 {
+	Modify();
 	const UEdGraphPin* SelectedPinPtr = GetSelectedPin();
 	const FName SelectedPinName = (SelectedPinPtr) ? SelectedPinPtr->PinName : FName();
 	// forget SelectedPin, eventually reassign below via the PinSelectionManager
@@ -484,7 +485,9 @@ void UTG_EdGraphNode::UpdatePinVisibility(UEdGraphPin* Pin, UTG_Pin* TGPin) cons
 	Pin->bHidden = !bCanEditChange && bEditConditionHides;
 	Pin->bDefaultValueIsReadOnly = bCanEditChange;
 	
-	if(Pin->bHidden) 
+	//Adding GIsTransacting check here as functions that are creating the Transaction
+	//should not be called from here 
+	if(Pin->bHidden && !GIsTransacting)
 	{
 		Pin->GetSchema()->BreakPinLinks(*Pin, false);
 	}
