Fixed issue: Histogram doesn't update after tweaking

Fix/How: When tweaking, blobs are still being cached for a while. They are differentiated by the IsTransient flag. If you stop tweaking, these cached blobs can be used by non-transient blobs. But we should still set the blobs to non-transient correctly. The current issue was that there were multiple places where these blobs were being set to transient.

diff --git a/Engine/Plugins/Experimental/TextureGraph/Source/TextureGraphEngine/Data/Blob.cpp b/Engine/Plugins/Experimental/TextureGraph/Source/TextureGraphEngine/Data/Blob.cpp
--- a/Engine/Plugins/Experimental/TextureGraph/Source/TextureGraphEngine/Data/Blob.cpp
+++ b/Engine/Plugins/Experimental/TextureGraph/Source/TextureGraphEngine/Data/Blob.cpp
@@ -532,7 +532,20 @@ void Blob::SyncWith(Blob* RHS)
 
 void Blob::FinaliseFrom(Blob* RHS)
 {
+	// We are only saving the buffer transient state when buffer is valid.
+	// This buffer will be empty in case of non single blob.
+	const bool HadValidBuffer = Buffer && Buffer->IsValid();
+	const bool bIsTransient = Buffer && Buffer->IsValid() && Buffer->IsTransient();
+
+	// We need to save the transient state before copying the RHS buffer.
 	Buffer = RHS->Buffer;
+
+	// Only try to update the buffer transient state when we found a valid buffer from RHS.
+	if(HadValidBuffer && Buffer && Buffer->IsValid())
+	{
+		Buffer->Desc.bIsTransient &= bIsTransient;		
+	}
+	
 	LODLevels = RHS->LODLevels;
 	MinMax = RHS->MinMax;
 	MinValue = RHS->MinValue;
diff --git a/Engine/Plugins/Experimental/TextureGraph/Source/TextureGraphEngine/Data/TiledBlob.cpp b/Engine/Plugins/Experimental/TextureGraph/Source/TextureGraphEngine/Data/TiledBlob.cpp
--- a/Engine/Plugins/Experimental/TextureGraph/Source/TextureGraphEngine/Data/TiledBlob.cpp
+++ b/Engine/Plugins/Experimental/TextureGraph/Source/TextureGraphEngine/Data/TiledBlob.cpp
@@ -1070,7 +1070,7 @@ void TiledBlob_Promise::AddLinkedBlob(BlobPtr LinkedBlob)
 	if (IsFinalised())
 	{
 		/// Copy the contents over
-		*TiledPromiseLinkedBlob = *this;
+		TiledPromiseLinkedBlob->FinaliseFrom(this);
 	}
 
 	TiledPromiseLinkedBlob->CachedBlob = std::static_pointer_cast<TiledBlob_Promise>(shared_from_this());
diff --git a/Engine/Plugins/Experimental/TextureGraph/Source/TextureGraphEngine/Device/Device.cpp b/Engine/Plugins/Experimental/TextureGraph/Source/TextureGraphEngine/Device/Device.cpp
--- a/Engine/Plugins/Experimental/TextureGraph/Source/TextureGraphEngine/Device/Device.cpp
+++ b/Engine/Plugins/Experimental/TextureGraph/Source/TextureGraphEngine/Device/Device.cpp
@@ -710,7 +710,7 @@ AsyncDeviceBufferRef Device::CombineFromTiles(const CombineSplitArgs& CombineArg
 	CHashPtr Hash = CHash::ConstructFromSources(Hashes);
 	
 	/// We've already calculated the Hash from all the child blobs
-	RawBufferPtr Raw = TextureHelper::CombineRaw_Tiles(rawTiles, Hash, true);
+	RawBufferPtr Raw = TextureHelper::CombineRaw_Tiles(rawTiles, Hash, Buffer->IsTransient());
 
 	Buffer->UpdateRaw(Raw);
 
diff --git a/Engine/Plugins/Experimental/TextureGraph/Source/TextureGraphEngine/Job/Job.cpp b/Engine/Plugins/Experimental/TextureGraph/Source/TextureGraphEngine/Job/Job.cpp
--- a/Engine/Plugins/Experimental/TextureGraph/Source/TextureGraphEngine/Job/Job.cpp
+++ b/Engine/Plugins/Experimental/TextureGraph/Source/TextureGraphEngine/Job/Job.cpp
@@ -396,7 +396,10 @@ bool Job::CheckCached()
 			ResultOrg->FinaliseFrom(ExistingResult.get());
 		else
 		{
-			ResultOrg = std::make_shared<TiledBlob_Promise>(ExistingResult->GetDescriptor(), ExistingResult->Rows(), ExistingResult->Cols(), ExistingResult->Hash());
+			BufferDescriptor Desc = ExistingResult->GetDescriptor();
+			Desc.bIsTransient = IsDiscard();
+			
+			ResultOrg = std::make_shared<TiledBlob_Promise>(Desc, ExistingResult->Rows(), ExistingResult->Cols(), ExistingResult->Hash());
 			ResultOrg->FinaliseFrom(ExistingResult.get());
 		}
 
