This CL is  gathering the list of material attributes assigned (meaning connected) in a Material expression
this is done by reflecting on the attributes connected on the material
from this we create 2 similar arrays
one of the enum used in the TextureGraphMaterialDraw shader to extract the corresponding attribute
the second is the display name of the attribute as shown in the Material editor

we changed the RenderedAttributeId field of the Expression_Material from an enum to a FName aka the name of the Mateiral Attribute
this is more flexible.

to be able to display and control that Fname field from the Node ui, we implemented the STG_GraphPinString  class directly copied from the EdGraphPinString but supportin gthe indirection to the true owner of the Property.

diff --git a/Engine/Plugins/Experimental/TextureGraph/Source/TextureGraph/Private/Expressions/Input/TG_Expression_Material.cpp b/Engine/Plugins/Experimental/TextureGraph/Source/TextureGraph/Private/Expressions/Input/TG_Expression_Material.cpp
--- a/Engine/Plugins/Experimental/TextureGraph/Source/TextureGraph/Private/Expressions/Input/TG_Expression_Material.cpp
+++ b/Engine/Plugins/Experimental/TextureGraph/Source/TextureGraph/Private/Expressions/Input/TG_Expression_Material.cpp
@@ -6,11 +6,17 @@
 void UTG_Expression_Material::PostEditChangeProperty(FPropertyChangedEvent& PropertyChangedEvent)
 {
 	// First catch if Material changes
-	if (PropertyChangedEvent.GetPropertyName() == FName(TEXT("Material")))
+	if (PropertyChangedEvent.GetPropertyName() == GET_MEMBER_NAME_CHECKED(UTG_Expression_Material, Material))
 	{
 		UE_LOG(LogTextureGraph, Log, TEXT("Material Expression PostEditChangeProperty."));
 		SetMaterialInternal(Material);
 	}
+	// Second catch if AttributeName changes
+	if (PropertyChangedEvent.GetPropertyName() == GET_MEMBER_NAME_CHECKED(UTG_Expression_Material, RenderedAttribute))
+	{
+		UE_LOG(LogTextureGraph, Log, TEXT("Material Expression PostEditChangeProperty."));
+		SetRenderedAttribute(RenderedAttribute);
+	}
 
 	Super::PostEditChangeProperty(PropertyChangedEvent);
 }
@@ -27,6 +33,12 @@ void UTG_Expression_Material::PostEditUndo()
 
 #endif
 
+void UTG_Expression_Material::Initialize()
+{
+	Super::Initialize();
+	SetRenderedAttribute(RenderedAttribute); // reassign the RenderedAttribute to make sure it is the correct one
+}
+
 void UTG_Expression_Material::SetMaterialInternal(UMaterialInterface* InMaterial)
 {
 	if (!InMaterial)
@@ -44,6 +56,8 @@ void UTG_Expression_Material::SetMaterialInternal(UMaterialInterface* InMaterial
 	}
 
 	Super::SetMaterialInternal(InMaterial);
+
+	SetRenderedAttribute(RenderedAttribute);
 }
 
 
@@ -64,6 +78,27 @@ void UTG_Expression_Material::SetMaterial(UMaterialInterface* InMaterial)
 	SetMaterialInternal(InMaterial);
 }
 
+void UTG_Expression_Material::SetRenderedAttribute(FName InRenderedAttribute)
+{
+	if (GetAvailableMaterialAttributeNames().Num())
+	{
+		int32 RenderAttributeIndex = GetAvailableMaterialAttributeNames().Find(InRenderedAttribute);
+		if (RenderAttributeIndex == INDEX_NONE)
+		{
+			RenderedAttribute = GetAvailableMaterialAttributeNames()[0];
+		}
+		else
+		{
+			RenderedAttribute = InRenderedAttribute;
+		}
+	}
+	else
+	{
+		RenderedAttribute = TEXT("None");
+	}
+	
+}
+
 bool UTG_Expression_Material::CanHandleAsset(UObject* Asset)
 {
 	const UMaterialInterface* Mat = Cast<UMaterialInterface>(Asset);
@@ -101,5 +136,28 @@ FName UTG_Expression_Material::GetTitleName() const
 	return TitleName;
 }
 
+TArray<FName> UTG_Expression_Material::GetRenderAttributeOptions() const
+{
+	return GetAvailableMaterialAttributeNames();
+}
 
+EDrawMaterialAttributeTarget UTG_Expression_Material::GetRenderedAttributeId()
+{
+	if (GetAvailableMaterialAttributeNames().Num())
+	{
+		int32 RenderAttributeIndex = GetAvailableMaterialAttributeNames().Find(RenderedAttribute);
+		if (RenderAttributeIndex == INDEX_NONE)
+		{
+			return GetAvailableMaterialAttributeIds()[0];
+		}
+		else
+		{
+			return GetAvailableMaterialAttributeIds()[RenderAttributeIndex];
+		}
+	}
+	else
+	{
+		return EDrawMaterialAttributeTarget::Emissive;
+	}
+}
 
diff --git a/Engine/Plugins/Experimental/TextureGraph/Source/TextureGraph/Private/Expressions/TG_Expression.cpp b/Engine/Plugins/Experimental/TextureGraph/Source/TextureGraph/Private/Expressions/TG_Expression.cpp
--- a/Engine/Plugins/Experimental/TextureGraph/Source/TextureGraph/Private/Expressions/TG_Expression.cpp
+++ b/Engine/Plugins/Experimental/TextureGraph/Source/TextureGraph/Private/Expressions/TG_Expression.cpp
@@ -77,6 +77,8 @@ FTG_Signature::FInit UTG_Expression::GetSignatureInitArgsFromClass() const
 			bIsNotConnectable |= (ByteProperty != nullptr);
 			FEnumProperty* EnumProperty = CastField<FEnumProperty>(*Prop);
 			bIsNotConnectable |= (EnumProperty != nullptr);
+			FNameProperty* NameProperty = CastField<FNameProperty>(*Prop);
+			bIsNotConnectable |= (NameProperty != nullptr);
 		}
 		uint8 MaskNotConnectable = bIsNotConnectable ? static_cast<uint8>(ETG_Access::NotConnectableFlag) : 0;
 		
diff --git a/Engine/Plugins/Experimental/TextureGraph/Source/TextureGraph/Private/Expressions/TG_Expression_MaterialBase.cpp b/Engine/Plugins/Experimental/TextureGraph/Source/TextureGraph/Private/Expressions/TG_Expression_MaterialBase.cpp
--- a/Engine/Plugins/Experimental/TextureGraph/Source/TextureGraph/Private/Expressions/TG_Expression_MaterialBase.cpp
+++ b/Engine/Plugins/Experimental/TextureGraph/Source/TextureGraph/Private/Expressions/TG_Expression_MaterialBase.cpp
@@ -10,6 +10,64 @@
 #include "TG_Texture.h"
 #include "TG_Var.h"
 #include "FxMat/MaterialManager.h"
+#include "Materials/MaterialAttributeDefinitionMap.h"
+
+
+EDrawMaterialAttributeTarget UTG_Expression_MaterialBase::ConvertEMaterialPropertyToEDrawMaterialAttributeTarget(EMaterialProperty InMaterialProperty)
+{
+	switch (InMaterialProperty)
+	{
+		case MP_EmissiveColor:
+			return EDrawMaterialAttributeTarget::Emissive;
+		case MP_Opacity:
+			return EDrawMaterialAttributeTarget::Opacity;
+		case MP_OpacityMask:
+			return EDrawMaterialAttributeTarget::OpacityMask;
+		case MP_BaseColor:
+			return EDrawMaterialAttributeTarget::BaseColor;
+		case MP_Metallic:
+			return EDrawMaterialAttributeTarget::Metallic;
+		case MP_Specular:
+			return EDrawMaterialAttributeTarget::Specular;
+		case MP_Roughness:
+			return EDrawMaterialAttributeTarget::Roughness;
+		case MP_Anisotropy:
+			return EDrawMaterialAttributeTarget::Anisotropy;
+		case MP_Normal:
+			return EDrawMaterialAttributeTarget::Normal;
+		case MP_Tangent:
+			return EDrawMaterialAttributeTarget::Tangent;
+
+		case MP_DiffuseColor:
+		case MP_SpecularColor:
+		case MP_WorldPositionOffset:
+		case MP_WorldDisplacement_DEPRECATED:
+		case MP_TessellationMultiplier_DEPRECATED:
+		case MP_SubsurfaceColor:
+		case MP_AmbientOcclusion:
+		case MP_Refraction:
+		case MP_CustomData0:
+		case MP_CustomData1:
+		case MP_CustomizedUVs0:
+		case MP_CustomizedUVs1:
+		case MP_CustomizedUVs2:
+		case MP_CustomizedUVs3:
+		case MP_CustomizedUVs4:
+		case MP_CustomizedUVs5:
+		case MP_CustomizedUVs6:
+		case MP_CustomizedUVs7:
+		case MP_PixelDepthOffset:
+		case MP_ShadingModel:
+		case MP_FrontMaterial:
+		case MP_SurfaceThickness:
+		case MP_Displacement:
+		case MP_MaterialAttributes:
+		case MP_CustomOutput:
+		default:
+			return EDrawMaterialAttributeTarget::Count; // Not supported yet in the shader
+	 }
+}
+
 
 bool UTG_Expression_MaterialBase::Validate(MixUpdateCyclePtr	Cycle)
 {
@@ -36,11 +94,42 @@ void UTG_Expression_MaterialBase::SetMaterialInternal(UMaterialInterface* InMate
 		}
 	}
 
+	// Detect the set of available material properties for rendering
+	GenerateMaterialAttributeOptions();
+	
 	// Signature is reset, notify the owning node / graph to update itself
 	NotifySignatureChanged();
 }
 
 
+void UTG_Expression_MaterialBase::GenerateMaterialAttributeOptions()
+{
+	// Detect the set of available material properties for rendering
+	AvailableMaterialAttributeIds.Empty();
+	AvailableMaterialAttributeNames.Empty();
+	if (MaterialInstance)
+	{
+		UMaterial* RefMaterial = MaterialInstance->GetMaterial();
+		for (int i = MP_EmissiveColor; i < MP_MAX; ++i)
+		{
+			if (RefMaterial->IsPropertyConnected(EMaterialProperty(i)))
+			{
+				EDrawMaterialAttributeTarget Attribute = ConvertEMaterialPropertyToEDrawMaterialAttributeTarget(EMaterialProperty(i));
+				if (Attribute != EDrawMaterialAttributeTarget::Count)
+				{
+					AvailableMaterialAttributeIds.Add(Attribute);
+
+
+					FName AttributeName = FName(UEnum::GetDisplayValueAsText(Attribute).ToString());
+					AttributeName = FName(FMaterialAttributeDefinitionMap::GetDisplayNameForMaterial(EMaterialProperty(i), RefMaterial).ToString());
+
+					AvailableMaterialAttributeNames.Add(AttributeName);
+				}
+			}
+		}
+	}
+}
+
 void UTG_Expression_MaterialBase::Evaluate(FTG_EvaluationContext* InContext)
 {
 	Super::Evaluate(InContext);
@@ -67,6 +156,7 @@ void UTG_Expression_MaterialBase::Initialize()
 	{
 		MaterialInstance = UMaterialInstanceDynamic::Create(GetMaterial(), this);
 	}
+	GenerateMaterialAttributeOptions(); // also populate the attributes availables
 }
 
 TiledBlobPtr UTG_Expression_MaterialBase::CreateRenderMaterialJob(FTG_EvaluationContext* InContext, const FString& InName, const FString& InMaterialPath, const BufferDescriptor& InDescriptor, EDrawMaterialAttributeTarget InDrawMaterialAttributeTarget)
diff --git a/Engine/Plugins/Experimental/TextureGraph/Source/TextureGraph/Private/TG_Var.cpp b/Engine/Plugins/Experimental/TextureGraph/Source/TextureGraph/Private/TG_Var.cpp
--- a/Engine/Plugins/Experimental/TextureGraph/Source/TextureGraph/Private/TG_Var.cpp
+++ b/Engine/Plugins/Experimental/TextureGraph/Source/TextureGraph/Private/TG_Var.cpp
@@ -36,6 +36,12 @@ template <> FString TG_Var_LogValue(float& Value)
 	return LogMessage;
 }
 
+template <> FString TG_Var_LogValue(FName& Value)
+{
+	FString LogMessage = Value.ToString();
+	return LogMessage;
+}
+
 template <> FString TG_Var_LogValue(FLinearColor& Value)
 {
 	FString LogMessage = Value.ToString();
@@ -86,6 +92,11 @@ template <> void TG_Var_SetValueFromString(uint8& Value, const FString& StrVal)
 {
 	Value = static_cast<uint8>(FCString::Atoi(*StrVal));
 }
+template <> void TG_Var_SetValueFromString(FName& Value, const FString& StrVal)
+{
+	Value = FName(StrVal);
+}
+
 template <> void TG_Var_SetValueFromString(FLinearColor& Value, const FString& StrVal)
 {
 	Value.InitFromString(StrVal);
@@ -209,6 +220,11 @@ void VarPropertySerializer_bool(FTG_Var::VarPropertySerialInfo& Info)
 	Generic_Simple_Serializer<FBoolProperty, bool>(Info);
 }
 
+void VarPropertySerializer_FName(FTG_Var::VarPropertySerialInfo& Info)
+{
+	Generic_Simple_Serializer<FNameProperty, FName>(Info);
+}
+
 void VarPropertySerializer_UObjectPtr(FTG_Var::VarPropertySerialInfo& Info)
 {
 	Generic_Struct_Serializer<FTG_Texture>(Info);
@@ -344,6 +360,7 @@ FTG_Var::VarPropertySerializerMap FTG_Var::DefaultPropertySerializers
 		VAR_PROPERTY_SERIALIZER_DEF(int32),
 		VAR_PROPERTY_SERIALIZER_DEF(uint32),
 		VAR_PROPERTY_SERIALIZER_DEF(float),
+		VAR_PROPERTY_SERIALIZER_DEF(FName),
 		VAR_PROPERTY_SERIALIZER_DEF(FTG_Texture),
 		VAR_PROPERTY_SERIALIZER_DEF(FVector4f),
 		VAR_PROPERTY_SERIALIZER_DEF(FVector2f),
diff --git a/Engine/Plugins/Experimental/TextureGraph/Source/TextureGraph/Public/Expressions/Input/TG_Expression_Material.h b/Engine/Plugins/Experimental/TextureGraph/Source/TextureGraph/Public/Expressions/Input/TG_Expression_Material.h
--- a/Engine/Plugins/Experimental/TextureGraph/Source/TextureGraph/Public/Expressions/Input/TG_Expression_Material.h
+++ b/Engine/Plugins/Experimental/TextureGraph/Source/TextureGraph/Public/Expressions/Input/TG_Expression_Material.h
@@ -28,21 +28,28 @@ public:
 	
 	void SetMaterial(UMaterialInterface* InMaterial);
 
-	// The attribute identifier among all the attributes of the material that is rendered in the output
-	UPROPERTY(EditAnywhere, BlueprintReadWrite, Category = NoCategory, meta = (TGType = "TG_Setting"))
-	EDrawMaterialAttributeTarget RenderedAttributeId = EDrawMaterialAttributeTarget::BaseColor;	
-	
+	// The Material attribute identifier among all the attributes of the material that is rendered in the output
+	UPROPERTY(EditAnywhere, BlueprintReadWrite, Setter, Category = NoCategory, meta = (TGType = "TG_Setting", GetOptions = "GetRenderAttributeOptions"))
+	FName RenderedAttribute;
+	void SetRenderedAttribute(FName InRenderedAttribute);
+
+	// THe list of Rendered attribute options available 
+	UFUNCTION(CallInEditor)
+	TArray<FName> GetRenderAttributeOptions() const;
+
+
 	virtual bool CanHandleAsset(UObject* Asset) override;
 	virtual void SetAsset(UObject* Asset) override;
 	virtual FText GetTooltipText() const override { return FText::FromString(TEXT("Renders a material into a quad and makes it available. It is automatically exposed as a graph input parameter.")); } 
 
 protected:
-	
 	virtual void SetMaterialInternal(UMaterialInterface* InMaterial) override;
 
 	UPROPERTY(EditAnywhere, BlueprintReadWrite, Category = "")
 	FName TitleName = TEXT("Material");
 
+	virtual void Initialize() override;
+
 public:
 	
 	virtual void SetTitleName(FName NewName) override;
@@ -52,6 +59,6 @@ public:
 	
 protected:
 	virtual TObjectPtr<UMaterialInterface> GetMaterial() const override { return Material;};
-	virtual EDrawMaterialAttributeTarget GetRenderedAttributeId()  override { return RenderedAttributeId;}
+	virtual EDrawMaterialAttributeTarget GetRenderedAttributeId()  override;
 };
 
diff --git a/Engine/Plugins/Experimental/TextureGraph/Source/TextureGraph/Public/Expressions/TG_Expression_MaterialBase.h b/Engine/Plugins/Experimental/TextureGraph/Source/TextureGraph/Public/Expressions/TG_Expression_MaterialBase.h
--- a/Engine/Plugins/Experimental/TextureGraph/Source/TextureGraph/Public/Expressions/TG_Expression_MaterialBase.h
+++ b/Engine/Plugins/Experimental/TextureGraph/Source/TextureGraph/Public/Expressions/TG_Expression_MaterialBase.h
@@ -25,6 +25,7 @@ enum class EDrawMaterialAttributeTarget : uint8
 	Opacity,
 	OpacityMask,
 	Normal,
+	Tangent,
 
 	/// Always has to be the last
 	Count UMETA(Hidden),
@@ -56,13 +57,21 @@ public:
 	// Validate internal checks, warnings and errors
 	virtual bool Validate(MixUpdateCyclePtr	Cycle) override;
 
+	// Access the list of attributes available for rendering from the current Material
+	const TArray<EDrawMaterialAttributeTarget>& GetAvailableMaterialAttributeIds() const { return AvailableMaterialAttributeIds; }
+	const TArray<FName>& GetAvailableMaterialAttributeNames() const { return AvailableMaterialAttributeNames; }
 
 protected:
 	UPROPERTY(Transient)
 	TObjectPtr<UMaterialInstanceDynamic> MaterialInstance; // A local Instance Material is recreated for the Material member
 
+	TArray<EDrawMaterialAttributeTarget>	AvailableMaterialAttributeIds; // The set of material properties available for rendering
+	TArray<FName>							AvailableMaterialAttributeNames; // same with the attribute names
+
 	virtual void Initialize() override;
-	
+
+	void GenerateMaterialAttributeOptions(); // Based on the current Material, list of material attributes available
+
 	TiledBlobPtr CreateRenderMaterialJob(FTG_EvaluationContext* InContext, const FString& InName, const FString& InMaterialPath, const BufferDescriptor& InDescriptor, EDrawMaterialAttributeTarget InDrawMaterialAttributeTarget);
 	TiledBlobPtr CreateRenderMaterialJob(FTG_EvaluationContext* InContext, const RenderMaterial_BPPtr& InRenderMaterial, const BufferDescriptor& InDescriptor, EDrawMaterialAttributeTarget InDrawMaterialAttributeTarget);
 	void LinkMaterialParameters(FTG_EvaluationContext* InContext, JobUPtr& InMaterialJob, const UMaterial* InMaterial, BufferDescriptor InDescriptor);
@@ -76,6 +85,7 @@ protected:
 private:
 	void AddSignatureParam(TArray<FMaterialParameterInfo> OutParameterInfoScalar, FName CPPTypeName, FTG_Signature::FInit& SignatureInit, bool IsScalar = false) const;
 
+	static EDrawMaterialAttributeTarget ConvertEMaterialPropertyToEDrawMaterialAttributeTarget(EMaterialProperty InMaterialProperty);
 };
 
  
\ No newline at end of file
diff --git a/Engine/Plugins/Experimental/TextureGraph/Source/TextureGraph/Public/TG_Var.h b/Engine/Plugins/Experimental/TextureGraph/Source/TextureGraph/Public/TG_Var.h
--- a/Engine/Plugins/Experimental/TextureGraph/Source/TextureGraph/Public/TG_Var.h
+++ b/Engine/Plugins/Experimental/TextureGraph/Source/TextureGraph/Public/TG_Var.h
@@ -201,12 +201,14 @@ template <> FString TG_Var_LogValue(float& Value);
 template <> FString TG_Var_LogValue(FLinearColor& Value);
 template <> FString TG_Var_LogValue(FVector4f& Value);
 template <> FString TG_Var_LogValue(FVector2f& Value);
+template <> FString TG_Var_LogValue(FName& Value);
 template <> FString TG_Var_LogValue(TObjectPtr<UObject>& Value);
 template <> FString TG_Var_LogValue(struct FTG_Texture& Value);
 template <> FString TG_Var_LogValue(struct FTG_Scalar& Value);
 template <> FString TG_Var_LogValue(struct FTG_OutputSettings& Value);
 template <> FString TG_Var_LogValue(struct FTG_Variant& Value);
 
+
 template <> void TG_Var_SetValueFromString(int& Value, const FString& StrVal);
 template <> void TG_Var_SetValueFromString(bool& Value, const FString& StrVal);
 template <> void TG_Var_SetValueFromString(float& Value, const FString& StrVal);
@@ -216,9 +218,11 @@ template <> void TG_Var_SetValueFromString(TObjectPtr<UObject>& Value, const FSt
 template <> void TG_Var_SetValueFromString(FLinearColor& Value, const FString& StrVal);
 template <> void TG_Var_SetValueFromString(FVector4f& Value, const FString& StrVal);
 template <> void TG_Var_SetValueFromString(FVector2f& Value, const FString& StrVal);
+template <> void TG_Var_SetValueFromString(FName& Value, const FString& StrVal);
 
 //template <> void TG_Var_SetValueFromString(struct FTG_Texture& Value, const FString& StrVal);
 template <> void TG_Var_SetValueFromString(struct FTG_Scalar& Value, const FString& StrVal);
 template <> void TG_Var_SetValueFromString(struct FTG_OutputSettings& Value, const FString& StrVal);
 template <> void TG_Var_SetValueFromString(struct FTG_Variant& Value, const FString& StrVal);
 
+
diff --git a/Engine/Plugins/Experimental/TextureGraph/Source/TextureGraphEditor/Private/EdGraph/TG_EditorGraphPanelPinFactory.cpp b/Engine/Plugins/Experimental/TextureGraph/Source/TextureGraphEditor/Private/EdGraph/TG_EditorGraphPanelPinFactory.cpp
--- a/Engine/Plugins/Experimental/TextureGraph/Source/TextureGraphEditor/Private/EdGraph/TG_EditorGraphPanelPinFactory.cpp
+++ b/Engine/Plugins/Experimental/TextureGraph/Source/TextureGraphEditor/Private/EdGraph/TG_EditorGraphPanelPinFactory.cpp
@@ -35,6 +35,8 @@
 #include "Pins/STG_GraphMaterialIdMaskInfos.h"
 #include "Transform/Expressions/T_ExtractMaterialIds.h"
 #include "Pins/STG_GraphPinOutputSettings.h"
+#include "Pins/STG_GraphPinString.h"
+
 #include "NumericPropertyParams.h"
 
 TSharedPtr<class SGraphPin> FTG_EditorGraphPanelPinFactory::CreatePin(class UEdGraphPin* InPin) const
@@ -142,6 +144,13 @@ TSharedPtr<class SGraphPin> FTG_EditorGraphPanelPinFactory::CreatePin(class UEdG
 			//const UClass* ObjectMetaClass = Cast<UClass>(InPin->PinType.PinSubCategoryObject.Get());
 			return SNew(SGraphPinObject, InPin);
 		}
+		else if (InPin->PinType.PinCategory == "PinCategory")
+		{
+			if (InPin->PinType.PinSubCategory == "FName")
+			{
+				return SNew(STG_GraphPinString, InPin);
+			}
+		}
 		return SNew(SGraphPin, InPin);
 	}
 	return nullptr;
diff --git a/Engine/Plugins/Experimental/TextureGraph/Source/TextureGraphEditor/Private/Pins/STG_GraphPinString.cpp b/Engine/Plugins/Experimental/TextureGraph/Source/TextureGraphEditor/Private/Pins/STG_GraphPinString.cpp
new file mode 100644
--- /dev/null
+++ b/Engine/Plugins/Experimental/TextureGraph/Source/TextureGraphEditor/Private/Pins/STG_GraphPinString.cpp
@@ -0,0 +1,158 @@
+// Copyright Epic Games, Inc. All Rights Reserved.
+
+
+#include "STG_GraphPinString.h"
+
+#include "Containers/UnrealString.h"
+#include "Delegates/Delegate.h"
+#include "EdGraph/EdGraphPin.h"
+#include "EdGraph/EdGraphSchema.h"
+#include "GenericPlatform/GenericApplication.h"
+#include "Internationalization/Internationalization.h"
+#include "K2Node_CallFunction.h"
+#include "PropertyEditorUtils.h"
+#include "Misc/Attribute.h"
+#include "ScopedTransaction.h"
+#include "SSearchableComboBox.h"
+#include "Styling/AppStyle.h"
+#include "Styling/SlateColor.h"
+#include "Types/SlateStructs.h"
+#include "UObject/NameTypes.h"
+#include "Widgets/Input/SEditableTextBox.h"
+#include "Widgets/Input/SMultiLineEditableTextBox.h"
+#include "Widgets/Layout/SBox.h"
+#include "EdGraph/TG_EdGraphSchema.h"
+#include "Expressions/TG_Expression.h"
+class SWidget;
+
+void STG_GraphPinString::Construct(const FArguments& InArgs, UEdGraphPin* InGraphPinObj)
+{
+	SGraphPin::Construct(SGraphPin::FArguments(), InGraphPinObj);
+}
+
+TSharedRef<SWidget>	STG_GraphPinString::GetDefaultValueWidget()
+{
+	const TSharedPtr<SWidget> DropdownWidget = TryBuildComboBoxWidget();
+	if (DropdownWidget.Get())
+	{
+		return DropdownWidget.ToSharedRef();
+	}
+
+	return SNew(SBox)
+		.MinDesiredWidth(18)
+		.MaxDesiredWidth(400)
+		[
+			SNew(SEditableTextBox)
+				.Style(FAppStyle::Get(), "Graph.EditableTextBox")
+				.Text(this, &STG_GraphPinString::GetTypeInValue)
+				.SelectAllTextWhenFocused(true)
+				.Visibility(this, &SGraphPin::GetDefaultValueVisibility)
+				.IsReadOnly(this, &STG_GraphPinString::GetDefaultValueIsReadOnly)
+				.OnTextCommitted(this, &STG_GraphPinString::SetTypeInValue)
+				.ForegroundColor(FSlateColor::UseForeground())
+		];
+}
+
+FText STG_GraphPinString::GetTypeInValue() const
+{
+	return FText::FromString( GraphPinObj->GetDefaultAsString() );
+}
+
+void STG_GraphPinString::SetTypeInValue(const FText& NewTypeInValue, ETextCommit::Type /*CommitInfo*/)
+{
+	if(GraphPinObj->IsPendingKill())
+	{
+		return;
+	}
+
+	if(!GraphPinObj->GetDefaultAsString().Equals(NewTypeInValue.ToString()))
+	{
+		const FScopedTransaction Transaction( NSLOCTEXT("GraphEditor", "ChangeStringPinValue", "Change String Pin Value" ) );
+		GraphPinObj->Modify();
+
+		GraphPinObj->GetSchema()->TrySetDefaultValue(*GraphPinObj, NewTypeInValue.ToString());
+	}
+}
+
+bool STG_GraphPinString::GetDefaultValueIsReadOnly() const
+{
+	return GraphPinObj->bDefaultValueIsReadOnly;
+}
+
+TSharedRef<SWidget> STG_GraphPinString::GenerateComboBoxEntry(const TSharedPtr<FString> Value)
+{
+	return SNew(STextBlock)
+		.Text(FText::FromString(*Value))
+		.Font(FAppStyle::GetFontStyle("PropertyWindow.NormalFont"));
+}
+
+void STG_GraphPinString::HandleComboBoxSelectionChanged(const TSharedPtr<FString> Value,
+	const ESelectInfo::Type InSelectInfo)
+{
+	if (GraphPinObj->IsPendingKill())
+	{
+		return;
+	}
+
+	if (!GraphPinObj->GetDefaultAsString().Equals(*Value))
+	{
+		const FScopedTransaction Transaction(
+			NSLOCTEXT("GraphEditor", "ChangeStringPinValue", "Change String Pin Value")
+		);
+
+		GraphPinObj->Modify();
+		GraphPinObj->GetSchema()->TrySetDefaultValue(*GraphPinObj, *Value);
+	}
+}
+
+TSharedPtr<SWidget> STG_GraphPinString::TryBuildComboBoxWidget()
+{
+	const UTG_EdGraphSchema* Schema = Cast<const UTG_EdGraphSchema>(GraphPinObj->GetOwningNode()->GetSchema());
+	UTG_Pin* TSPin = Schema->GetTGPinFromEdPin(GraphPinObj);
+	UTG_Expression* PropObject = TSPin->GetNodePtr()->GetExpression();
+	FProperty* PinProperty = TSPin->GetExpressionProperty();
+
+	FString GetOptionsFunctionName = PinProperty->GetMetaData("GetOptions");
+
+	if (GetOptionsFunctionName.IsEmpty())
+	{
+		return nullptr;
+	}
+
+
+	UFunction* GetOpetionsFunction = PropObject->GetClass()->FindFunctionByName(FName(GetOptionsFunctionName));
+	if (!GetOpetionsFunction)
+	{
+		return nullptr;
+	}
+
+	ComboBoxOptions.Empty();
+
+	TArray<UObject*> PropertyContainers = { PropObject };
+	PropertyEditorUtils::GetPropertyOptions(PropertyContainers, GetOptionsFunctionName, ComboBoxOptions);
+
+	const FString DefaultString = GraphPinObj->GetDefaultAsString();
+	const TSharedPtr<FString>* DefaultValue = ComboBoxOptions.FindByPredicate([&DefaultString](
+		const TSharedPtr<FString>& Value)
+	{
+		return Value.IsValid() ? *Value == DefaultString : false;
+	});
+
+	return SNew(SBox)
+		.MinDesiredWidth(18)
+		.MaxDesiredWidth(400)
+		[
+			SNew(SSearchableComboBox)
+				.OptionsSource(&ComboBoxOptions)
+				.OnGenerateWidget(this, &STG_GraphPinString::GenerateComboBoxEntry)
+				.OnSelectionChanged(this, &STG_GraphPinString::HandleComboBoxSelectionChanged)
+				.ContentPadding(3.0f)
+				.InitiallySelectedItem(DefaultValue ? *DefaultValue : nullptr)
+				.Content()
+				[
+					SNew(STextBlock)
+						.Text(this, &STG_GraphPinString::GetTypeInValue)
+						.Font(FAppStyle::GetFontStyle("PropertyWindow.NormalFont"))
+				]
+		];
+}
diff --git a/Engine/Plugins/Experimental/TextureGraph/Source/TextureGraphEditor/Private/Pins/STG_GraphPinString.h b/Engine/Plugins/Experimental/TextureGraph/Source/TextureGraphEditor/Private/Pins/STG_GraphPinString.h
new file mode 100644
--- /dev/null
+++ b/Engine/Plugins/Experimental/TextureGraph/Source/TextureGraphEditor/Private/Pins/STG_GraphPinString.h
@@ -0,0 +1,41 @@
+// Copyright Epic Games, Inc. All Rights Reserved.
+
+#pragma once
+
+#include "CoreMinimal.h"
+#include "Internationalization/Text.h"
+#include "SGraphPin.h"
+#include "Templates/SharedPointer.h"
+#include "Types/SlateEnums.h"
+#include "Widgets/DeclarativeSyntaxSupport.h"
+#include "Widgets/SWidget.h"
+
+class SWidget;
+class UEdGraphPin;
+
+class STG_GraphPinString : public SGraphPin
+{
+public:
+	SLATE_BEGIN_ARGS(STG_GraphPinString) {}
+	SLATE_END_ARGS()
+
+	void Construct(const FArguments& InArgs, UEdGraphPin* InGraphPinObj);
+
+protected:
+	//~ Begin SGraphPin Interface
+	virtual TSharedRef<SWidget>	GetDefaultValueWidget() override;
+	//~ End SGraphPin Interface
+
+	FText GetTypeInValue() const;
+	virtual void SetTypeInValue(const FText& NewTypeInValue, ETextCommit::Type CommitInfo);
+
+	/** @return True if the pin default value field is read-only */
+	bool GetDefaultValueIsReadOnly() const;
+
+private:
+	TSharedRef<SWidget> GenerateComboBoxEntry(TSharedPtr<FString> Value);
+	void HandleComboBoxSelectionChanged(TSharedPtr<FString> Value, ESelectInfo::Type InSelectInfo);
+	TSharedPtr<SWidget> TryBuildComboBoxWidget();
+
+	TArray<TSharedPtr<FString>> ComboBoxOptions;
+};
