Deprecate:
 - FRHITexture2D
 - FRHITexture2DArray
 - FRHITexture3D
 - FRHITextureCube
 - FTexture2DRHIRef
 - FTexture2DArrayRHIRef
 - FTexture3DRHIRef
 - FTextureCubeRHIRef

Replaced with FRHITexture and FTextureRHIRef

These types were unified in UE 5.1 and have been defined via "using" statements to the same underlying texture type for several engine releases.

diff --git a/Engine/Plugins/Compositing/LensDistortion/Source/LensDistortion/Private/LensDistortionRendering.cpp b/Engine/Plugins/Compositing/LensDistortion/Source/LensDistortion/Private/LensDistortionRendering.cpp
--- a/Engine/Plugins/Compositing/LensDistortion/Source/LensDistortion/Private/LensDistortionRendering.cpp
+++ b/Engine/Plugins/Compositing/LensDistortion/Source/LensDistortion/Private/LensDistortionRendering.cpp
@@ -165,7 +165,7 @@ static void DrawUVDisplacementToRenderTarget_RenderThread(
 
 	SCOPED_DRAW_EVENTF(RHICmdList, SceneCapture, TEXT("LensDistortionDisplacementGeneration %s"), TextureRenderTargetName);
 
-	FRHITexture2D* RenderTargetTexture = OutTextureRenderTargetResource->GetRenderTargetTexture();
+	FRHITexture* RenderTargetTexture = OutTextureRenderTargetResource->GetRenderTargetTexture();
 
 	RHICmdList.Transition(FRHITransitionInfo(RenderTargetTexture, ERHIAccess::SRVMask, ERHIAccess::RTV));
 
diff --git a/Engine/Plugins/Experimental/Avalanche/Source/AvalancheMedia/Private/Broadcast/Channel/AvaBroadcastOutputChannel.cpp b/Engine/Plugins/Experimental/Avalanche/Source/AvalancheMedia/Private/Broadcast/Channel/AvaBroadcastOutputChannel.cpp
--- a/Engine/Plugins/Experimental/Avalanche/Source/AvalancheMedia/Private/Broadcast/Channel/AvaBroadcastOutputChannel.cpp
+++ b/Engine/Plugins/Experimental/Avalanche/Source/AvalancheMedia/Private/Broadcast/Channel/AvaBroadcastOutputChannel.cpp
@@ -121,10 +121,10 @@ namespace UE::AvaMedia::Private
 	{
 		FRDGBuilder GraphBuilder(RHICmdList);
 
-		const FTexture2DRHIRef SourceTexture = SourceRT->GetRenderTargetResource()->GetTexture2DRHI();
+		const FTextureRHIRef SourceTexture = SourceRT->GetRenderTargetResource()->GetTexture2DRHI();
 		const FRDGTextureRef SourceRGBTexture = GraphBuilder.RegisterExternalTexture(CreateRenderTarget(SourceTexture, TEXT("SourceTexture")));
 
-		const FTexture2DRHIRef DestinationTexture = DestinationRT->GetRenderTargetResource()->GetTexture2DRHI();
+		const FTextureRHIRef DestinationTexture = DestinationRT->GetRenderTargetResource()->GetTexture2DRHI();
 		const FRDGTextureRef OutputResource = GraphBuilder.RegisterExternalTexture(CreateRenderTarget(DestinationTexture, TEXT("DestTexture")));
 
 		const FVector2D SizeU = { 0.0f, 1.0f };
diff --git a/Engine/Plugins/Experimental/Avalanche/Source/AvalancheMedia/Private/Broadcast/OutputDevices/AvaBroadcastDisplayMediaCapture.cpp b/Engine/Plugins/Experimental/Avalanche/Source/AvalancheMedia/Private/Broadcast/OutputDevices/AvaBroadcastDisplayMediaCapture.cpp
--- a/Engine/Plugins/Experimental/Avalanche/Source/AvalancheMedia/Private/Broadcast/OutputDevices/AvaBroadcastDisplayMediaCapture.cpp
+++ b/Engine/Plugins/Experimental/Avalanche/Source/AvalancheMedia/Private/Broadcast/OutputDevices/AvaBroadcastDisplayMediaCapture.cpp
@@ -151,7 +151,7 @@ bool UAvaBroadcastDisplayMediaCapture::InitializeCapture()
 bool UAvaBroadcastDisplayMediaCapture::PostInitializeCaptureViewport(TSharedPtr<FSceneViewport>& InSceneViewport)
 {
 	bool bSuccess = false;
-	const FTexture2DRHIRef& BackBuffer = InSceneViewport->GetRenderTargetTexture();
+	const FTextureRHIRef& BackBuffer = InSceneViewport->GetRenderTargetTexture();
 	if (BackBuffer.IsValid())
 	{
 		const FRHITextureDesc& Desc = BackBuffer->GetDesc();
diff --git a/Engine/Plugins/Experimental/Avalanche/Source/AvalancheMedia/Private/Broadcast/OutputDevices/AvaBroadcastRenderTargetMediaCapture.cpp b/Engine/Plugins/Experimental/Avalanche/Source/AvalancheMedia/Private/Broadcast/OutputDevices/AvaBroadcastRenderTargetMediaCapture.cpp
--- a/Engine/Plugins/Experimental/Avalanche/Source/AvalancheMedia/Private/Broadcast/OutputDevices/AvaBroadcastRenderTargetMediaCapture.cpp
+++ b/Engine/Plugins/Experimental/Avalanche/Source/AvalancheMedia/Private/Broadcast/OutputDevices/AvaBroadcastRenderTargetMediaCapture.cpp
@@ -66,7 +66,7 @@ bool UAvaBroadcastRenderTargetMediaCapture::InitializeCapture()
 bool UAvaBroadcastRenderTargetMediaCapture::PostInitializeCaptureViewport(TSharedPtr<FSceneViewport>& InSceneViewport)
 {
 	bool bSuccess = false;
-	const FTexture2DRHIRef& BackBuffer = InSceneViewport->GetRenderTargetTexture();
+	const FTextureRHIRef& BackBuffer = InSceneViewport->GetRenderTargetTexture();
 	if (BackBuffer.IsValid())
 	{
 		const FRHITextureDesc& Desc = BackBuffer->GetDesc();
diff --git a/Engine/Plugins/Experimental/ImagePlate/Source/ImagePlate/Private/ImagePlateFileSequence.cpp b/Engine/Plugins/Experimental/ImagePlate/Source/ImagePlate/Private/ImagePlateFileSequence.cpp
--- a/Engine/Plugins/Experimental/ImagePlate/Source/ImagePlate/Private/ImagePlateFileSequence.cpp
+++ b/Engine/Plugins/Experimental/ImagePlate/Source/ImagePlate/Private/ImagePlateFileSequence.cpp
@@ -142,7 +142,7 @@ bool FImagePlateSourceFrame::EnsureTextureMetrics(UTexture* DestinationTexture)
 			FlushRenderingCommands();
 		}
 
-		FTexture2DRHIRef Texture2DRHI = DestinationTexture->GetResource()->TextureRHI ? DestinationTexture->GetResource()->TextureRHI->GetTexture2D() : nullptr;
+		FTextureRHIRef Texture2DRHI = DestinationTexture->GetResource()->TextureRHI ? DestinationTexture->GetResource()->TextureRHI->GetTexture2D() : nullptr;
 		if (!Texture2DRHI)
 		{
 			UE_LOG(LogImagePlateFileSequence, Warning, TEXT("Unsupported texture type encountered: Unable to update texture to fit source frame size or bitdepth."));
@@ -197,7 +197,7 @@ TFuture<void> FImagePlateSourceFrame::CopyTo(UTexture* DestinationTexture)
 				return;
 			}
 
-			FTexture2DRHIRef Texture2DRHI = CommandData->DestinationTexture->GetResource()->TextureRHI->GetTexture2D();
+			FTextureRHIRef Texture2DRHI = CommandData->DestinationTexture->GetResource()->TextureRHI->GetTexture2D();
 			if (!Texture2DRHI)
 			{
 				return;
diff --git a/Engine/Plugins/Experimental/Mutable/Source/CustomizableObjectEditor/Private/MuCOE/SCustomizableObjectLayoutGrid.cpp b/Engine/Plugins/Experimental/Mutable/Source/CustomizableObjectEditor/Private/MuCOE/SCustomizableObjectLayoutGrid.cpp
--- a/Engine/Plugins/Experimental/Mutable/Source/CustomizableObjectEditor/Private/MuCOE/SCustomizableObjectLayoutGrid.cpp
+++ b/Engine/Plugins/Experimental/Mutable/Source/CustomizableObjectEditor/Private/MuCOE/SCustomizableObjectLayoutGrid.cpp
@@ -34,7 +34,7 @@ public:
 	}
 
 	/** Sets the texture that this target renders to */
-	void SetRenderTargetTexture(FTexture2DRHIRef& InRHIRef)
+	void SetRenderTargetTexture(FTextureRHIRef& InRHIRef)
 	{
 		RenderTargetTextureRHI = InRHIRef;
 	}
@@ -1223,7 +1223,7 @@ void FUVCanvasDrawer::Draw_RenderThread(class FRHICommandListImmediate& RHICmdLi
 {
 	if (Initialized)
 	{
-		RenderTarget->SetRenderTargetTexture(*(FTexture2DRHIRef*)InWindowBackBuffer);
+		RenderTarget->SetRenderTargetTexture(*(FTextureRHIRef*)InWindowBackBuffer);
 
 		FCanvas Canvas(RenderTarget, nullptr, FGameTime(), GMaxRHIFeatureLevel);
 		
diff --git a/Engine/Plugins/Experimental/PixelStreamingPlayer/Source/PixelStreamingPlayer/Private/PixelStreamingMediaTexture.cpp b/Engine/Plugins/Experimental/PixelStreamingPlayer/Source/PixelStreamingPlayer/Private/PixelStreamingMediaTexture.cpp
--- a/Engine/Plugins/Experimental/PixelStreamingPlayer/Source/PixelStreamingPlayer/Private/PixelStreamingMediaTexture.cpp
+++ b/Engine/Plugins/Experimental/PixelStreamingPlayer/Source/PixelStreamingPlayer/Private/PixelStreamingMediaTexture.cpp
@@ -81,7 +81,7 @@ void UPixelStreamingMediaTexture::InitializeResources()
 	});
 }
 
-void UPixelStreamingMediaTexture::UpdateTextureReference(FRHICommandList& RHICmdList, FTexture2DRHIRef Reference)
+void UPixelStreamingMediaTexture::UpdateTextureReference(FRHICommandList& RHICmdList, FTextureRHIRef Reference)
 {
 	if (CurrentResource != nullptr)
 	{
diff --git a/Engine/Plugins/Experimental/PixelStreamingPlayer/Source/PixelStreamingPlayer/Public/PixelStreamingMediaTexture.h b/Engine/Plugins/Experimental/PixelStreamingPlayer/Source/PixelStreamingPlayer/Public/PixelStreamingMediaTexture.h
--- a/Engine/Plugins/Experimental/PixelStreamingPlayer/Source/PixelStreamingPlayer/Public/PixelStreamingMediaTexture.h
+++ b/Engine/Plugins/Experimental/PixelStreamingPlayer/Source/PixelStreamingPlayer/Public/PixelStreamingMediaTexture.h
@@ -32,7 +32,7 @@ private:
 	void InitializeResources();
 
 	// updates the internal texture resource after each frame.
-	void UpdateTextureReference(FRHICommandList& RHICmdList, FTexture2DRHIRef Reference);
+	void UpdateTextureReference(FRHICommandList& RHICmdList, FTextureRHIRef Reference);
 
 	TArray<uint8_t> Buffer;
 
diff --git a/Engine/Plugins/Experimental/TextureGraph/Source/TextureGraphEngine/2D/Tex.cpp b/Engine/Plugins/Experimental/TextureGraph/Source/TextureGraphEngine/2D/Tex.cpp
--- a/Engine/Plugins/Experimental/TextureGraph/Source/TextureGraphEngine/2D/Tex.cpp
+++ b/Engine/Plugins/Experimental/TextureGraph/Source/TextureGraphEngine/2D/Tex.cpp
@@ -1320,7 +1320,7 @@ UTexture* Tex::GetTexture() const
 //	return nullptr;
 //}
 
-FRHITexture2D* Tex::GetRHITexture() const
+FRHITexture* Tex::GetRHITexture() const
 {
 	if (RT)
 		return RT->GetResource()->GetTextureRHI();
diff --git a/Engine/Plugins/Experimental/TextureGraph/Source/TextureGraphEngine/2D/Tex.h b/Engine/Plugins/Experimental/TextureGraph/Source/TextureGraphEngine/2D/Tex.h
--- a/Engine/Plugins/Experimental/TextureGraph/Source/TextureGraphEngine/2D/Tex.h
+++ b/Engine/Plugins/Experimental/TextureGraph/Source/TextureGraphEngine/2D/Tex.h
@@ -161,7 +161,7 @@ public:
 	void							GenerateMips();
 
 	virtual bool					IsArray()  { return false; }
-	virtual FRHITexture2D*			GetRHITexture() const;
+	virtual FRHITexture*			GetRHITexture() const;
 	virtual UTexture*				GetTexture() const;
 
 	virtual bool					IsNull() const;
@@ -188,7 +188,7 @@ public:
 	FORCEINLINE FLinearColor		ClearColor() const { return Desc.ClearColor; }
 	FORCEINLINE FLinearColor&		ClearColor() { return Desc.ClearColor; }
 	FORCEINLINE FIntPoint			GetSize() const { return FIntPoint(Desc.Width, Desc.Height); }
-	FORCEINLINE operator			FRHITexture2D*() const { return GetRHITexture(); }
+	FORCEINLINE operator			FRHITexture*() const { return GetRHITexture(); }
 	//FORCEINLINE operator			FTextureRHIRef() const { return RHITextureRef(); }
 	FORCEINLINE operator			UTexture*() const { return GetTexture(); }
 	FORCEINLINE FString				GetName() const { return GetTexture()->GetName(); }
diff --git a/Engine/Plugins/Experimental/TextureGraph/Source/TextureGraphEngine/2D/TextureHelper.cpp b/Engine/Plugins/Experimental/TextureGraph/Source/TextureGraphEngine/2D/TextureHelper.cpp
--- a/Engine/Plugins/Experimental/TextureGraph/Source/TextureGraphEngine/2D/TextureHelper.cpp
+++ b/Engine/Plugins/Experimental/TextureGraph/Source/TextureGraphEngine/2D/TextureHelper.cpp
@@ -764,7 +764,7 @@ TextureType TextureHelper::TextureStringToType(const FString& TypeString)
 RawBufferPtr TextureHelper::RawFromRT(UTextureRenderTarget2D* RenderTarget, const BufferDescriptor& Desc)
 {
 	//check(IsInRenderingThread()); //Can be from any thread
-	return RawFromResource(FTexture2DRHIRef(((FTexture2DResource*)RenderTarget->GetResource())->GetTexture2DRHI()), Desc);
+	return RawFromResource(FTextureRHIRef(((FTexture2DResource*)RenderTarget->GetResource())->GetTexture2DRHI()), Desc);
 }
 
 RawBufferPtr TextureHelper::RawFromTexture(UTexture2D* Texture, const BufferDescriptor& Desc)
@@ -819,7 +819,7 @@ size_t TextureHelper::RoundUpTo(size_t Size, size_t DesiredRounding)
 	return RoundedSize;
 }
 
-RawBufferPtr TextureHelper::RawFromResource(const FTexture2DRHIRef& ResourceRHI, const BufferDescriptor& Desc)
+RawBufferPtr TextureHelper::RawFromResource(const FTextureRHIRef& ResourceRHI, const BufferDescriptor& Desc)
 {
 	FRHICommandListImmediate& RHICmdList = FRHICommandListImmediate::Get();
 
@@ -898,18 +898,18 @@ RawBufferPtr TextureHelper::RawFromResource(const FTexture2DRHIRef& ResourceRHI,
 void TextureHelper::RawFromRT_Tiled(UTextureRenderTarget2D* RenderTarget, const BufferDescriptor& Desc, size_t TileSizeX, size_t TileSizeY, RawBufferPtrTiles& Tiles)
 {
 	check(IsInRenderingThread());
-	FTexture2DRHIRef ResourceRHI = ((FTextureRenderTarget2DResource*)RenderTarget->GetResource())->GetTextureRHI();
+	FTextureRHIRef ResourceRHI = ((FTextureRenderTarget2DResource*)RenderTarget->GetResource())->GetTextureRHI();
 	return RawFromResource_Tiled(ResourceRHI, Desc, TileSizeX, TileSizeY, Tiles);
 }
 
 void TextureHelper::RawFromTexture_Tiled(UTexture2D* Texture, const BufferDescriptor& Desc, size_t TileSizeX, size_t TileSizeY, RawBufferPtrTiles& Tiles)
 {
 	check(IsInRenderingThread());
-	FTexture2DRHIRef ResourceRHI = ((FTexture2DResource*)Texture->GetResource())->GetTexture2DRHI();
+	FTextureRHIRef ResourceRHI = ((FTexture2DResource*)Texture->GetResource())->GetTexture2DRHI();
 	return RawFromResource_Tiled(ResourceRHI, Desc, TileSizeX, TileSizeY, Tiles);
 }
 
-void TextureHelper::RawFromResource_Tiled(FTexture2DRHIRef ResourceRHI, const BufferDescriptor& Desc, size_t TileSizeX, size_t TileSizeY, RawBufferPtrTiles& Tiles)
+void TextureHelper::RawFromResource_Tiled(FTextureRHIRef ResourceRHI, const BufferDescriptor& Desc, size_t TileSizeX, size_t TileSizeY, RawBufferPtrTiles& Tiles)
 {
 	/// If any of the dimensions are less than the tile size requested then we don't tile at all
 	if (Desc.Width < TileSizeX && Desc.Height < TileSizeY)
diff --git a/Engine/Plugins/Experimental/TextureGraph/Source/TextureGraphEngine/2D/TextureHelper.h b/Engine/Plugins/Experimental/TextureGraph/Source/TextureGraphEngine/2D/TextureHelper.h
--- a/Engine/Plugins/Experimental/TextureGraph/Source/TextureGraphEngine/2D/TextureHelper.h
+++ b/Engine/Plugins/Experimental/TextureGraph/Source/TextureGraphEngine/2D/TextureHelper.h
@@ -82,12 +82,12 @@ struct TEXTUREGRAPHENGINE_API TextureHelper
 
 	static RawBufferPtr					RawFromRT(UTextureRenderTarget2D* RenderTarget, const BufferDescriptor& Desc);
 	static RawBufferPtr					RawFromTexture(UTexture2D* Texture, const BufferDescriptor& Desc);
-	static RawBufferPtr					RawFromResource(const FTexture2DRHIRef& ResourceRHI, const BufferDescriptor& Desc);
+	static RawBufferPtr					RawFromResource(const FTextureRHIRef& ResourceRHI, const BufferDescriptor& Desc);
 	static BufferFormat					FindOptimalSupportedFormat(BufferFormat SrcFormat);
 
 	static void							RawFromRT_Tiled(UTextureRenderTarget2D* RenderTarget, const BufferDescriptor& Desc, size_t TileSizeX, size_t TileSizeY, RawBufferPtrTiles& Tiles);
 	static void							RawFromTexture_Tiled(UTexture2D* Texture, const BufferDescriptor& Desc, size_t TileSizeX, size_t TileSizeY, RawBufferPtrTiles& Tiles);
-	static void							RawFromResource_Tiled(FTexture2DRHIRef ResourceRHI, const BufferDescriptor& Desc, size_t TileSizeX, size_t TileSizeY, RawBufferPtrTiles& Tiles);
+	static void							RawFromResource_Tiled(FTextureRHIRef ResourceRHI, const BufferDescriptor& Desc, size_t TileSizeX, size_t TileSizeY, RawBufferPtrTiles& Tiles);
 	static void							RawFromMem_Tiled(const uint8* SrcData, size_t SrcDataLength, const BufferDescriptor& SrcDesc, size_t TileSizeX, size_t TileSizeY, RawBufferPtrTiles& Tiles);
 	static RawBufferPtr					CombineRaw_Tiles(const RawBufferPtrTiles& Tiles, CHashPtr HashValue = nullptr, bool bIsTransient = false);
 
diff --git a/Engine/Plugins/Experimental/TextureGraph/Source/TextureGraphEngine/Device/FX/DeviceBuffer_FX.cpp b/Engine/Plugins/Experimental/TextureGraph/Source/TextureGraphEngine/Device/FX/DeviceBuffer_FX.cpp
--- a/Engine/Plugins/Experimental/TextureGraph/Source/TextureGraphEngine/Device/FX/DeviceBuffer_FX.cpp
+++ b/Engine/Plugins/Experimental/TextureGraph/Source/TextureGraphEngine/Device/FX/DeviceBuffer_FX.cpp
@@ -215,7 +215,7 @@ void DeviceBuffer_FX::UpdateRenderTarget(TexPtr TextureObj)
 		FTextureRenderTarget2DResource* rtRes = (FTextureRenderTarget2DResource*)rt->GetRenderTargetResource();
 		check(rtRes);
 
-		FTexture2DRHIRef rhiTexture = rtRes->GetTextureRHI();
+		FTextureRHIRef rhiTexture = rtRes->GetTextureRHI();
 		check(rhiTexture);
 		rt->UpdateResourceImmediate();
 
diff --git a/Engine/Plugins/Experimental/TextureGraph/Source/TextureGraphEngine/Device/FX/Device_FX.cpp b/Engine/Plugins/Experimental/TextureGraph/Source/TextureGraphEngine/Device/FX/Device_FX.cpp
--- a/Engine/Plugins/Experimental/TextureGraph/Source/TextureGraphEngine/Device/FX/Device_FX.cpp
+++ b/Engine/Plugins/Experimental/TextureGraph/Source/TextureGraphEngine/Device/FX/Device_FX.cpp
@@ -270,7 +270,7 @@ AsyncDeviceBufferRef Device_FX::FillTextureArray_Deferred(DeviceBufferRef Buffer
 			check(DstTex);
 
 			UTextureRenderTarget2D* RTDest = (UTextureRenderTarget2D*)DstTex->GetTexture();
-			FTexture2DRHIRef RTResDest = ((FTextureRenderTarget2DResource*)RTDest->GetResource())->GetTextureRHI();
+			FTextureRHIRef RTResDest = ((FTextureRenderTarget2DResource*)RTDest->GetResource())->GetTextureRHI();
 
 			FRHICommandListImmediate& DevRHI = RHI();
 
@@ -288,7 +288,7 @@ AsyncDeviceBufferRef Device_FX::FillTextureArray_Deferred(DeviceBufferRef Buffer
 					int32 TileHeight = TileTex->GetHeight();
 
 					UTextureRenderTarget2D* TileRT = TileTex->GetRenderTarget();
-					FRHITexture2D* TileResource = TileTex->GetRHITexture();
+					FRHITexture* TileResource = TileTex->GetRHITexture();
 
 					FRHICopyTextureInfo CopyInfo;
 					CopyInfo.Size = TileResource->GetSizeXYZ();
@@ -318,7 +318,7 @@ AsyncDeviceBufferRef Device_FX::DrawTilesToBuffer_Deferred(DeviceBufferRef Buffe
 			check(DstTex);
 
 			UTextureRenderTarget2D* RTDest = DstTex->GetRenderTarget();
-			FTexture2DRHIRef RTResDest = ((FTextureRenderTarget2DResource*)RTDest->GetResource())->GetTextureRHI();
+			FTextureRHIRef RTResDest = ((FTextureRenderTarget2DResource*)RTDest->GetResource())->GetTextureRHI();
 
 			FRHICommandListImmediate& DevRHI = RHI();
 
@@ -357,7 +357,7 @@ AsyncDeviceBufferRef Device_FX::DrawTilesToBuffer_Deferred(DeviceBufferRef Buffe
 						check(TileTex);
 
 						UTextureRenderTarget2D* TileRT = TileTex->GetRenderTarget();
-						FRHITexture2D* TileResource = TileTex->GetRHITexture();
+						FRHITexture* TileResource = TileTex->GetRHITexture();
 
 						const float SrcTextureWidth = TileTex->GetWidth();
 						const float SrcTextureHeight = TileTex->GetHeight();
@@ -405,7 +405,7 @@ AsyncDeviceBufferRef Device_FX::SplitToTiles_Internal(const CombineSplitArgs& Sp
 
 			/// This texture must have been created beforehand AND it must be a render target
 			TexPtr SrcTex = FXBuffer->GetTexture();
-			FTexture2DRHIRef sourceResource = ((UTextureRenderTarget2D*)SrcTex->GetTexture())->GetRenderTargetResource()->GetRenderTargetTexture();
+			FTextureRHIRef sourceResource = ((UTextureRenderTarget2D*)SrcTex->GetTexture())->GetRenderTargetResource()->GetRenderTargetTexture();
 
 			for (int32 TileX = 0; TileX < Tiles.Rows(); TileX++)
 			{
@@ -417,7 +417,7 @@ AsyncDeviceBufferRef Device_FX::SplitToTiles_Internal(const CombineSplitArgs& Sp
 					check(TileTex);
 					check(TileTex->GetRenderTarget());
 
-					FTexture2DRHIRef texture2Dres = ((FTextureRenderTarget2DResource*)TileTex->GetRenderTarget()->GetResource())->GetTextureRHI();
+					FTextureRHIRef texture2Dres = ((FTextureRenderTarget2DResource*)TileTex->GetRenderTarget()->GetResource())->GetTextureRHI();
 
 					FRHICopyTextureInfo CopyInfo;
 					CopyInfo.SourcePosition = FIntVector(TileX * TileTex->GetWidth(), TileY * TileTex->GetHeight(), 0);
@@ -679,7 +679,7 @@ int32 Device_FX::InitRTResource(TexPtr TextureObj, UTextureRenderTarget* RT)
 	UE_LOG(LogDevice, VeryVerbose, TEXT("New RT Array allocation: %s %llu [Ptr: 0x%x, Size: %dx%d]"), *TextureObj->GetDescriptor().Name,
 		TextureObj->GetDescriptor().Format_HashValue(), RT, RTRes->GetSizeX(), RTRes->GetSizeY());
 
-	FTexture2DRHIRef rhiTexture = RTRes->GetTextureRHI();
+	FTextureRHIRef rhiTexture = RTRes->GetTextureRHI();
 	check(rhiTexture);
 
 	RHIBindDebugLabelName(rhiTexture, *TextureObj->GetDescriptor().Name);
diff --git a/Engine/Plugins/Experimental/TextureGraph/Source/TextureGraphEngine/FxMat/FxMaterial.cpp b/Engine/Plugins/Experimental/TextureGraph/Source/TextureGraphEngine/FxMat/FxMaterial.cpp
--- a/Engine/Plugins/Experimental/TextureGraph/Source/TextureGraphEngine/FxMat/FxMaterial.cpp
+++ b/Engine/Plugins/Experimental/TextureGraph/Source/TextureGraphEngine/FxMat/FxMaterial.cpp
@@ -124,7 +124,7 @@ void FxMaterial::BindTexturesForBlitting()
 			FRHITexture* tex = BoundTex.Texture->GetResource()->TextureRHI;
 			check(tex);
 
-			FRHITexture2DArray* texture2DArray = tex->GetTexture2DArray();
+			FRHITexture* texture2DArray = tex->GetTexture2DArray();
 
 			if (!texture2DArray)
 			{
@@ -132,7 +132,7 @@ void FxMaterial::BindTexturesForBlitting()
 			}
 			else
 			{
-				memcpy(BoundTex.Arg, (const char*)&texture2DArray, sizeof(FRHITexture2DArray**));
+				memcpy(BoundTex.Arg, (const char*)&texture2DArray, sizeof(FRHITexture**));
 			}
 		}
 		else if (BoundTex.tiles.size())
diff --git a/Engine/Plugins/Experimental/TextureGraph/Source/TextureGraphEngine/FxMat/FxMaterial.h b/Engine/Plugins/Experimental/TextureGraph/Source/TextureGraphEngine/FxMat/FxMaterial.h
--- a/Engine/Plugins/Experimental/TextureGraph/Source/TextureGraphEngine/FxMat/FxMaterial.h
+++ b/Engine/Plugins/Experimental/TextureGraph/Source/TextureGraphEngine/FxMat/FxMaterial.h
@@ -238,7 +238,7 @@ public:
 	virtual std::shared_ptr<FxMaterial>	
 									Clone() = 0;
 	virtual FxMetadataSet			GetMetadata() const = 0;
-	virtual void					Blit(FRHICommandListImmediate& RHI, FRHITexture2D* Target, const RenderMesh* MeshObj, int32 TargetId, FGraphicsPipelineStateInitializer* PSO = nullptr) = 0;
+	virtual void					Blit(FRHICommandListImmediate& RHI, FRHITexture* Target, const RenderMesh* MeshObj, int32 TargetId, FGraphicsPipelineStateInitializer* PSO = nullptr) = 0;
 
 	//////////////////////////////////////////////////////////////////////////
 	/// Mimicing UMaterialInstanceDynamic
@@ -327,7 +327,7 @@ public:
 		Params.Result = UAV;
 	}
 
-	virtual void Blit(FRHICommandListImmediate& RHI, FRHITexture2D* Target, const RenderMesh* MeshObj, int32 TargetId, FGraphicsPipelineStateInitializer* PSO = nullptr) override
+	virtual void Blit(FRHICommandListImmediate& RHI, FRHITexture* Target, const RenderMesh* MeshObj, int32 TargetId, FGraphicsPipelineStateInitializer* PSO = nullptr) override
 	{
 		BindTexturesForBlitting();
 
@@ -423,7 +423,7 @@ public:
 		};
 	}
 
-	virtual void Blit(FRHICommandListImmediate& RHI, FRHITexture2D* Target, const RenderMesh* MeshObj, int32 InTargetId, FGraphicsPipelineStateInitializer* InPSO = nullptr) override
+	virtual void Blit(FRHICommandListImmediate& RHI, FRHITexture* Target, const RenderMesh* MeshObj, int32 InTargetId, FGraphicsPipelineStateInitializer* InPSO = nullptr) override
 	{
 		BindTexturesForBlitting();
 
diff --git a/Engine/Plugins/Experimental/TextureGraph/Source/TextureGraphEngine/FxMat/FxMaterial_DrawMaterial.h b/Engine/Plugins/Experimental/TextureGraph/Source/TextureGraphEngine/FxMat/FxMaterial_DrawMaterial.h
--- a/Engine/Plugins/Experimental/TextureGraph/Source/TextureGraphEngine/FxMat/FxMaterial_DrawMaterial.h
+++ b/Engine/Plugins/Experimental/TextureGraph/Source/TextureGraphEngine/FxMat/FxMaterial_DrawMaterial.h
@@ -153,12 +153,12 @@ public:
 		};
 	}
 
-	virtual void Blit(FRHICommandListImmediate& RHI, FRHITexture2D* Target, const RenderMesh* MeshObj, int32 TargetId, FGraphicsPipelineStateInitializer* InPSO = nullptr) override
+	virtual void Blit(FRHICommandListImmediate& RHI, FRHITexture* Target, const RenderMesh* MeshObj, int32 TargetId, FGraphicsPipelineStateInitializer* InPSO = nullptr) override
 	{
 		
 	}
 
-	void MyBlit(FRHICommandListImmediate& RHI, UTextureRenderTarget2D* RenderTarget, FRHITexture2D* Target, const RenderMesh* MeshObj, int32 TargetId, FGraphicsPipelineStateInitializer* InPSO = nullptr)
+	void MyBlit(FRHICommandListImmediate& RHI, UTextureRenderTarget2D* RenderTarget, FRHITexture* Target, const RenderMesh* MeshObj, int32 TargetId, FGraphicsPipelineStateInitializer* InPSO = nullptr)
 	{
 		BindTexturesForBlitting();
 
diff --git a/Engine/Plugins/Experimental/TextureGraph/Source/TextureGraphEngine/FxMat/RenderMaterial_BP.cpp b/Engine/Plugins/Experimental/TextureGraph/Source/TextureGraphEngine/FxMat/RenderMaterial_BP.cpp
--- a/Engine/Plugins/Experimental/TextureGraph/Source/TextureGraphEngine/FxMat/RenderMaterial_BP.cpp
+++ b/Engine/Plugins/Experimental/TextureGraph/Source/TextureGraphEngine/FxMat/RenderMaterial_BP.cpp
@@ -176,7 +176,7 @@ void RenderMaterial_BP::BlitTo(FRHICommandListImmediate& RHI, UTextureRenderTarg
 		FTextureRenderTarget2DResource* RTRes = (FTextureRenderTarget2DResource*)RT->GetRenderTargetResource();
 		check(RTRes);
 
-		FTexture2DRHIRef TextureRHI = RTRes->GetTextureRHI();
+		FTextureRHIRef TextureRHI = RTRes->GetTextureRHI();
 		check(TextureRHI);
 
 		TextureRHI->SetName(FName(*RT->GetName()));
diff --git a/Engine/Plugins/Experimental/TextureGraph/Source/TextureGraphEngine/FxMat/RenderMaterial_FX.cpp b/Engine/Plugins/Experimental/TextureGraph/Source/TextureGraphEngine/FxMat/RenderMaterial_FX.cpp
--- a/Engine/Plugins/Experimental/TextureGraph/Source/TextureGraphEngine/FxMat/RenderMaterial_FX.cpp
+++ b/Engine/Plugins/Experimental/TextureGraph/Source/TextureGraphEngine/FxMat/RenderMaterial_FX.cpp
@@ -83,7 +83,7 @@ void RenderMaterial_FX::BlitTo(FRHICommandListImmediate& RHI, UTextureRenderTarg
 	FTextureRenderTarget2DResource* RTRes = (FTextureRenderTarget2DResource*)RenderTarget->GetRenderTargetResource();
 	check(RTRes);
 
-	FTexture2DRHIRef TextureRHI = RTRes->GetTextureRHI();
+	FTextureRHIRef TextureRHI = RTRes->GetTextureRHI();
 	check(TextureRHI);
 
 	TextureRHI->SetName(FName(*RenderTarget->GetName()));
diff --git a/Engine/Plugins/Experimental/TextureGraph/Source/TextureGraphEngine/Transform/Utility/T_TextureHistogram.h b/Engine/Plugins/Experimental/TextureGraph/Source/TextureGraphEngine/Transform/Utility/T_TextureHistogram.h
--- a/Engine/Plugins/Experimental/TextureGraph/Source/TextureGraphEngine/Transform/Utility/T_TextureHistogram.h
+++ b/Engine/Plugins/Experimental/TextureGraph/Source/TextureGraphEngine/Transform/Utility/T_TextureHistogram.h
@@ -43,7 +43,7 @@ public:
 
 	}
 
-	virtual void Blit(FRHICommandListImmediate& RHI, FRHITexture2D* Target, const RenderMesh* MeshObj, int32 TargetId, FGraphicsPipelineStateInitializer* PSO = nullptr) override
+	virtual void Blit(FRHICommandListImmediate& RHI, FRHITexture* Target, const RenderMesh* MeshObj, int32 TargetId, FGraphicsPipelineStateInitializer* PSO = nullptr) override
 	{
 		SetBuffer(RHI);
 
diff --git a/Engine/Plugins/Experimental/TextureMediaPlayer/Source/TextureMediaPlayer/Private/TextureMediaPlayerVideoDecoderOutput.cpp b/Engine/Plugins/Experimental/TextureMediaPlayer/Source/TextureMediaPlayer/Private/TextureMediaPlayerVideoDecoderOutput.cpp
--- a/Engine/Plugins/Experimental/TextureMediaPlayer/Source/TextureMediaPlayer/Private/TextureMediaPlayerVideoDecoderOutput.cpp
+++ b/Engine/Plugins/Experimental/TextureMediaPlayer/Source/TextureMediaPlayer/Private/TextureMediaPlayerVideoDecoderOutput.cpp
@@ -10,9 +10,9 @@ void TextureMediaPlayerVideoDecoderOutput::Initialize(TSharedPtr<Electra::FParam
 }
 
 #if PLATFORM_WINDOWS
-void TextureMediaPlayerVideoDecoderOutput::Initialize(TSharedPtr<Electra::FParamDict, ESPMode::ThreadSafe> InParamDict, FTexture2DRHIRef InTexture, const FIntPoint& InSampleDim, TRefCountPtr<ID3D12Fence> InFence, uint64 InFenceValue)
+void TextureMediaPlayerVideoDecoderOutput::Initialize(TSharedPtr<Electra::FParamDict, ESPMode::ThreadSafe> InParamDict, FTextureRHIRef InTexture, const FIntPoint& InSampleDim, TRefCountPtr<ID3D12Fence> InFence, uint64 InFenceValue)
 #else
-void TextureMediaPlayerVideoDecoderOutput::Initialize(TSharedPtr<Electra::FParamDict, ESPMode::ThreadSafe> InParamDict, FTexture2DRHIRef InTexture, const FIntPoint& InSampleDim, FGPUFenceRHIRef InFence, uint64 InFenceValue)
+void TextureMediaPlayerVideoDecoderOutput::Initialize(TSharedPtr<Electra::FParamDict, ESPMode::ThreadSafe> InParamDict, FTextureRHIRef InTexture, const FIntPoint& InSampleDim, FGPUFenceRHIRef InFence, uint64 InFenceValue)
 #endif
 {
 	FVideoDecoderOutputPC::Initialize(MoveTemp(InParamDict));
diff --git a/Engine/Plugins/Experimental/TextureMediaPlayer/Source/TextureMediaPlayer/Private/TextureMediaPlayerVideoDecoderOutput.h b/Engine/Plugins/Experimental/TextureMediaPlayer/Source/TextureMediaPlayer/Private/TextureMediaPlayerVideoDecoderOutput.h
--- a/Engine/Plugins/Experimental/TextureMediaPlayer/Source/TextureMediaPlayer/Private/TextureMediaPlayerVideoDecoderOutput.h
+++ b/Engine/Plugins/Experimental/TextureMediaPlayer/Source/TextureMediaPlayer/Private/TextureMediaPlayerVideoDecoderOutput.h
@@ -19,9 +19,9 @@ public:
 	void Initialize(TSharedPtr<Electra::FParamDict, ESPMode::ThreadSafe> InParamDict, const TArray<uint8>& InBuffer, const FIntPoint& InSampleDim);
 
 #if PLATFORM_WINDOWS
-	void Initialize(TSharedPtr<Electra::FParamDict, ESPMode::ThreadSafe> InParamDict, FTexture2DRHIRef InTexture, const FIntPoint& InSampleDim, TRefCountPtr<ID3D12Fence> InFence, uint64 InFenceValue);
+	void Initialize(TSharedPtr<Electra::FParamDict, ESPMode::ThreadSafe> InParamDict, FTextureRHIRef InTexture, const FIntPoint& InSampleDim, TRefCountPtr<ID3D12Fence> InFence, uint64 InFenceValue);
 #else
-	void Initialize(TSharedPtr<Electra::FParamDict, ESPMode::ThreadSafe> InParamDict, FTexture2DRHIRef InTexture, const FIntPoint& InSampleDim, FGPUFenceRHIRef InFence, uint64 InFenceValue);
+	void Initialize(TSharedPtr<Electra::FParamDict, ESPMode::ThreadSafe> InParamDict, FTextureRHIRef InTexture, const FIntPoint& InSampleDim, FGPUFenceRHIRef InFence, uint64 InFenceValue);
 #endif // PLATFORM_WINDOWS
 
 	// FVideoDecoderOutputPC interface.
@@ -41,7 +41,7 @@ public:
 
 private:
 	/** Texture from WebRTC with accompaning sync data. */
-	FTexture2DRHIRef Texture;
+	FTextureRHIRef Texture;
 	
 #if PLATFORM_WINDOWS
 	TRefCountPtr<ID3D12Fence> Fence;
diff --git a/Engine/Plugins/Experimental/VirtualHeightfieldMesh/Source/VirtualHeightfieldMesh/Private/VirtualHeightfieldMeshSceneProxy.cpp b/Engine/Plugins/Experimental/VirtualHeightfieldMesh/Source/VirtualHeightfieldMesh/Private/VirtualHeightfieldMeshSceneProxy.cpp
--- a/Engine/Plugins/Experimental/VirtualHeightfieldMesh/Source/VirtualHeightfieldMesh/Private/VirtualHeightfieldMeshSceneProxy.cpp
+++ b/Engine/Plugins/Experimental/VirtualHeightfieldMesh/Source/VirtualHeightfieldMesh/Private/VirtualHeightfieldMeshSceneProxy.cpp
@@ -115,7 +115,7 @@ namespace VirtualHeightfieldMesh
 
 struct FOcclusionResults
 {
-	FTexture2DRHIRef OcclusionTexture;
+	FTextureRHIRef OcclusionTexture;
 	FIntPoint TextureSize;
 	int32 NumTextureMips;
 	TArray<bool> UploadData;
diff --git a/Engine/Plugins/Experimental/VirtualProductionUtilities/Source/VPUtilities/Private/Widgets/Fullscreen/PostProcessSceneViewExtension.cpp b/Engine/Plugins/Experimental/VirtualProductionUtilities/Source/VPUtilities/Private/Widgets/Fullscreen/PostProcessSceneViewExtension.cpp
--- a/Engine/Plugins/Experimental/VirtualProductionUtilities/Source/VPUtilities/Private/Widgets/Fullscreen/PostProcessSceneViewExtension.cpp
+++ b/Engine/Plugins/Experimental/VirtualProductionUtilities/Source/VPUtilities/Private/Widgets/Fullscreen/PostProcessSceneViewExtension.cpp
@@ -89,7 +89,7 @@ namespace UE::VirtualProductionUtilities::Private
 			return;
 		}
 		
-		const FTexture2DRHIRef WidgetRenderTarget_RHI = WidgetRenderTarget->GetRenderTargetResource()->GetTexture2DRHI();
+		const FTextureRHIRef WidgetRenderTarget_RHI = WidgetRenderTarget->GetRenderTargetResource()->GetTexture2DRHI();
 		const FRDGTextureRef WidgetRenderTarget_RDG = GraphBuilder.RegisterExternalTexture(CreateRenderTarget(WidgetRenderTarget_RHI, TEXT("WidgetRenderTarget")));
 
 		const FGlobalShaderMap* GlobalShaderMap = GetGlobalShaderMap(GMaxRHIFeatureLevel);
diff --git a/Engine/Plugins/FX/Niagara/Source/Niagara/Classes/NiagaraDataInterfaceRenderTarget2D.h b/Engine/Plugins/FX/Niagara/Source/Niagara/Classes/NiagaraDataInterfaceRenderTarget2D.h
--- a/Engine/Plugins/FX/Niagara/Source/Niagara/Classes/NiagaraDataInterfaceRenderTarget2D.h
+++ b/Engine/Plugins/FX/Niagara/Source/Niagara/Classes/NiagaraDataInterfaceRenderTarget2D.h
@@ -42,7 +42,7 @@ struct FRenderTarget2DRWInstanceData_RenderThread
 	bool bWroteThisFrame = false;
 
 	FSamplerStateRHIRef	SamplerStateRHI;
-	FTexture2DRHIRef	TextureRHI;
+	FTextureRHIRef	    TextureRHI;
 
 	FRDGTextureRef		TransientRDGTexture = nullptr;
 	FRDGTextureSRVRef	TransientRDGSRV = nullptr;
diff --git a/Engine/Plugins/FX/Niagara/Source/Niagara/Private/NiagaraRenderer.cpp b/Engine/Plugins/FX/Niagara/Source/Niagara/Private/NiagaraRenderer.cpp
--- a/Engine/Plugins/FX/Niagara/Source/Niagara/Private/NiagaraRenderer.cpp
+++ b/Engine/Plugins/FX/Niagara/Source/Niagara/Private/NiagaraRenderer.cpp
@@ -132,7 +132,7 @@ public:
 
 			case ETextureDimension::Texture3D:
 			{
-				FTexture3DRHIRef Texture3D = RHICreateTexture(Desc);
+				FTextureRHIRef Texture3D = RHICreateTexture(Desc);
 				Texture = Texture3D;
 
 				const FPixelFormatInfo& Info = GPixelFormats[PixelFormat];
diff --git a/Engine/Plugins/Media/AjaMedia/Source/AjaMediaOutput/Private/AjaMediaCapture.cpp b/Engine/Plugins/Media/AjaMedia/Source/AjaMediaOutput/Private/AjaMediaCapture.cpp
--- a/Engine/Plugins/Media/AjaMedia/Source/AjaMediaOutput/Private/AjaMediaCapture.cpp
+++ b/Engine/Plugins/Media/AjaMedia/Source/AjaMediaOutput/Private/AjaMediaCapture.cpp
@@ -611,7 +611,7 @@ void UAjaMediaCapture::LockDMATexture_RenderThread(FTextureRHIRef InTexture)
 		{
 			TexturesToRelease.Add(InTexture);
 
-			FRHITexture2D* Texture = InTexture->GetTexture2D();
+			FRHITexture* Texture = InTexture->GetTexture2D();
 			UE::GPUTextureTransfer::FRegisterDMATextureArgs Args;
 			Args.RHITexture = Texture;
 			
diff --git a/Engine/Plugins/Media/AndroidCamera/Source/AndroidCamera/Private/Player/AndroidCameraPlayer.cpp b/Engine/Plugins/Media/AndroidCamera/Source/AndroidCamera/Private/Player/AndroidCameraPlayer.cpp
--- a/Engine/Plugins/Media/AndroidCamera/Source/AndroidCamera/Private/Player/AndroidCameraPlayer.cpp
+++ b/Engine/Plugins/Media/AndroidCamera/Source/AndroidCamera/Private/Player/AndroidCameraPlayer.cpp
@@ -406,7 +406,7 @@ static void DoUpdateTextureCameraSampleExecute(TWeakPtr<FJavaAndroidCameraPlayer
 	}
 
 	// write frame into texture
-	FRHITexture2D* Texture = VideoSample->InitializeTexture(SampleTime);
+	FRHITexture* Texture = VideoSample->InitializeTexture(SampleTime);
 
 	if (Texture != nullptr)
 	{
diff --git a/Engine/Plugins/Media/AndroidCamera/Source/AndroidCamera/Private/Player/AndroidCameraTextureSample.h b/Engine/Plugins/Media/AndroidCamera/Source/AndroidCamera/Private/Player/AndroidCameraTextureSample.h
--- a/Engine/Plugins/Media/AndroidCamera/Source/AndroidCamera/Private/Player/AndroidCameraTextureSample.h
+++ b/Engine/Plugins/Media/AndroidCamera/Source/AndroidCamera/Private/Player/AndroidCameraTextureSample.h
@@ -131,7 +131,7 @@ public:
 	 * @note This method must be called on the render thread.
 	 * @see InitializeBuffer
 	 */
-	FRHITexture2D* InitializeTexture(FTimespan InTime)
+	FRHITexture* InitializeTexture(FTimespan InTime)
 	{
 		check(IsInRenderingThread() || IsInRHIThread());
 
@@ -264,7 +264,7 @@ private:
 #if WITH_ENGINE
 
 	/** Texture resource. */
-	TRefCountPtr<FRHITexture2D> Texture;
+	TRefCountPtr<FRHITexture> Texture;
 
 #endif //WITH_ENGINE
 };
diff --git a/Engine/Plugins/Media/AndroidMedia/Source/AndroidMedia/Private/Player/AndroidMediaPlayer.cpp b/Engine/Plugins/Media/AndroidMedia/Source/AndroidMedia/Private/Player/AndroidMediaPlayer.cpp
--- a/Engine/Plugins/Media/AndroidMedia/Source/AndroidMedia/Private/Player/AndroidMediaPlayer.cpp
+++ b/Engine/Plugins/Media/AndroidMedia/Source/AndroidMedia/Private/Player/AndroidMediaPlayer.cpp
@@ -575,7 +575,7 @@ static void DoUpdateTextureMediaSampleExecute(TWeakPtr<FJavaAndroidMediaPlayer,
 	const FTimespan Time = FTimespan::FromMilliseconds(CurrentFramePosition);
 
 	// write frame into texture
-	FRHITexture2D* Texture = VideoSample->InitializeTexture(Time);
+	FRHITexture* Texture = VideoSample->InitializeTexture(Time);
 
 	if (Texture != nullptr)
 	{
diff --git a/Engine/Plugins/Media/AndroidMedia/Source/AndroidMedia/Private/Player/AndroidMediaTextureSample.h b/Engine/Plugins/Media/AndroidMedia/Source/AndroidMedia/Private/Player/AndroidMediaTextureSample.h
--- a/Engine/Plugins/Media/AndroidMedia/Source/AndroidMedia/Private/Player/AndroidMediaTextureSample.h
+++ b/Engine/Plugins/Media/AndroidMedia/Source/AndroidMedia/Private/Player/AndroidMediaTextureSample.h
@@ -132,7 +132,7 @@ public:
 	 * @note This method must be called on the render thread.
 	 * @see InitializeBuffer
 	 */
-	FRHITexture2D* InitializeTexture(FTimespan InTime)
+	FRHITexture* InitializeTexture(FTimespan InTime)
 	{
 		check(IsInRenderingThread() || IsInRHIThread());
 
@@ -271,7 +271,7 @@ private:
 #if WITH_ENGINE
 
 	/** Texture resource. */
-	TRefCountPtr<FRHITexture2D> Texture;
+	TRefCountPtr<FRHITexture> Texture;
 
 #endif //WITH_ENGINE
 };
diff --git a/Engine/Plugins/Media/AvfMedia/Source/AvfMedia/Private/Player/AvfMediaVideoSampler.cpp b/Engine/Plugins/Media/AvfMedia/Source/AvfMedia/Private/Player/AvfMediaVideoSampler.cpp
--- a/Engine/Plugins/Media/AvfMedia/Source/AvfMedia/Private/Player/AvfMediaVideoSampler.cpp
+++ b/Engine/Plugins/Media/AvfMedia/Source/AvfMedia/Private/Player/AvfMediaVideoSampler.cpp
@@ -185,7 +185,7 @@ void FAvfMediaVideoSampler::ProcessFrame(CVPixelBufferRef Frame, FTimespan Sampl
 	auto VideoSample = VideoSamplePool->AcquireShared();
 
 #if WITH_ENGINE
-	TRefCountPtr<FRHITexture2D> ShaderResource;
+	TRefCountPtr<FRHITexture> ShaderResource;
 
 	// We have to support Metal for this object now
 	check(COREVIDEO_SUPPORTS_METAL);
diff --git a/Engine/Plugins/Media/AvfMedia/Source/AvfMedia/Public/AvfMediaTextureSample.h b/Engine/Plugins/Media/AvfMedia/Source/AvfMedia/Public/AvfMediaTextureSample.h
--- a/Engine/Plugins/Media/AvfMedia/Source/AvfMedia/Public/AvfMediaTextureSample.h
+++ b/Engine/Plugins/Media/AvfMedia/Source/AvfMedia/Public/AvfMediaTextureSample.h
@@ -44,7 +44,7 @@ public:
 	 * @return true on success, false otherwise.
 	 */
 	bool Initialize(
-		TRefCountPtr<FRHITexture2D>& InTexture,
+		TRefCountPtr<FRHITexture>& InTexture,
 		const FIntPoint& InDim,
 		const FIntPoint& InOutputDim,
 		FTimespan InTime,
@@ -185,7 +185,7 @@ private:
 
 #if WITH_ENGINE
 	/** The sample's texture resource. */
-	TRefCountPtr<FRHITexture2D> Texture;
+	TRefCountPtr<FRHITexture> Texture;
 
 #else
 	/** The sample's data buffer. */
diff --git a/Engine/Plugins/Media/AvfMedia/Source/AvfMediaCapture/Private/Player/AvfMediaCapturePlayer.cpp b/Engine/Plugins/Media/AvfMedia/Source/AvfMediaCapture/Private/Player/AvfMediaCapturePlayer.cpp
--- a/Engine/Plugins/Media/AvfMedia/Source/AvfMediaCapture/Private/Player/AvfMediaCapturePlayer.cpp
+++ b/Engine/Plugins/Media/AvfMedia/Source/AvfMediaCapture/Private/Player/AvfMediaCapturePlayer.cpp
@@ -723,7 +723,7 @@ void FAvfMediaCapturePlayer::ProcessSampleBufferVideo(CMSampleBufferRef SampleBu
 			.SetFlags(ETextureCreateFlags::SRGB | ETextureCreateFlags::Dynamic | ETextureCreateFlags::NoTiling | ETextureCreateFlags::ShaderResource)
 			.SetBulkData(new FAvfTexture2DResourceWrapper(TextureRef));
 
-		TRefCountPtr<FRHITexture2D> ShaderResource = RHICreateTexture(Desc);
+		TRefCountPtr<FRHITexture> ShaderResource = RHICreateTexture(Desc);
 		CFRelease(TextureRef);
 		
 		FIntPoint Dim(Width, Height);
diff --git a/Engine/Plugins/Media/BinkMedia/Source/BinkMediaPlayer/Private/Assets/BinkMediaPlayer.cpp b/Engine/Plugins/Media/BinkMedia/Source/BinkMediaPlayer/Private/Assets/BinkMediaPlayer.cpp
--- a/Engine/Plugins/Media/BinkMedia/Source/BinkMediaPlayer/Private/Assets/BinkMediaPlayer.cpp
+++ b/Engine/Plugins/Media/BinkMedia/Source/BinkMediaPlayer/Private/Assets/BinkMediaPlayer.cpp
@@ -577,7 +577,7 @@ void UBinkMediaPlayer::Tick(float DeltaTime)
 	}
 }
 
-void UBinkMediaPlayer::UpdateTexture(FRHICommandListImmediate &RHICmdList, FTexture2DRHIRef ref, void *nativePtr, int width, int height, bool isEditor, bool tonemap, int output_nits, float alpha, bool srgb_decode, bool is_hdr) 
+void UBinkMediaPlayer::UpdateTexture(FRHICommandListImmediate &RHICmdList, FTextureRHIRef ref, void *nativePtr, int width, int height, bool isEditor, bool tonemap, int output_nits, float alpha, bool srgb_decode, bool is_hdr) 
 {
 	check(IsInRenderingThread());
 
@@ -598,7 +598,7 @@ void UBinkMediaPlayer::Draw(UTexture *texture, bool tonemap, int out_nits, float
 	{
 		return;
 	}
-	FTexture2DRHIRef ref = texture->GetResource()->TextureRHI->GetTexture2D();
+	FTextureRHIRef ref = texture->GetResource()->TextureRHI->GetTexture2D();
 	if ((!IsPlaying() && !IsPaused()) || !ref) 
 	{
 		return;
@@ -615,7 +615,7 @@ void UBinkMediaPlayer::Draw(UTexture *texture, bool tonemap, int out_nits, float
 	struct parms_t 
 	{
 		UBinkMediaPlayer *player;
-		FTexture2DRHIRef ref;
+		FTextureRHIRef ref;
 		void *native;
 		int width, height;
 		bool tonemap;
diff --git a/Engine/Plugins/Media/BinkMedia/Source/BinkMediaPlayer/Private/Assets/BinkMediaTextureResource.cpp b/Engine/Plugins/Media/BinkMedia/Source/BinkMediaPlayer/Private/Assets/BinkMediaTextureResource.cpp
--- a/Engine/Plugins/Media/BinkMedia/Source/BinkMediaPlayer/Private/Assets/BinkMediaTextureResource.cpp
+++ b/Engine/Plugins/Media/BinkMedia/Source/BinkMediaPlayer/Private/Assets/BinkMediaTextureResource.cpp
@@ -92,7 +92,7 @@ void FBinkMediaTextureResource::UpdateDeferredResource(FRHICommandListImmediate&
 	{
 		return;
 	}
-	FTexture2DRHIRef tex = TextureRHI->GetTexture2D();
+	FTextureRHIRef tex = TextureRHI->GetTexture2D();
 	if (!tex.GetReference()) 
 	{
 		return;
@@ -112,7 +112,7 @@ void FBinkMediaTextureResource::Clear()
 	w = (w + 7) & -8;
 	h = (h + 7) & -8;
 
-	FTexture2DRHIRef ref = RenderTargetTextureRHI;
+	FTextureRHIRef ref = RenderTargetTextureRHI;
 	FTextureRHIRef ref2 = TextureRHI;
 	ENQUEUE_RENDER_COMMAND(BinkMediaPlayer_Draw)([ref,ref2,w,h](FRHICommandListImmediate& RHICmdList) 
 	{ 
diff --git a/Engine/Plugins/Media/BinkMedia/Source/BinkMediaPlayer/Private/BinkFunctionLibrary.cpp b/Engine/Plugins/Media/BinkMedia/Source/BinkMediaPlayer/Private/BinkFunctionLibrary.cpp
--- a/Engine/Plugins/Media/BinkMedia/Source/BinkMediaPlayer/Private/BinkFunctionLibrary.cpp
+++ b/Engine/Plugins/Media/BinkMedia/Source/BinkMediaPlayer/Private/BinkFunctionLibrary.cpp
@@ -20,7 +20,7 @@ static void Bink_DrawOverlays_Internal(FRHICommandListImmediate &RHICmdList, con
 
 	FVector2D screenSize;
 	gameViewport->GetViewportSize(screenSize);
-	const FTexture2DRHIRef &backbuffer = gameViewport->Viewport->GetRenderTargetTexture();
+	const FTextureRHIRef &backbuffer = gameViewport->Viewport->GetRenderTargetTexture();
 	if(!backbuffer.GetReference()) 
 	{
 		return;
diff --git a/Engine/Plugins/Media/BinkMedia/Source/BinkMediaPlayer/Private/BinkMediaPlayerModule.cpp b/Engine/Plugins/Media/BinkMedia/Source/BinkMediaPlayer/Private/BinkMediaPlayerModule.cpp
--- a/Engine/Plugins/Media/BinkMedia/Source/BinkMediaPlayer/Private/BinkMediaPlayerModule.cpp
+++ b/Engine/Plugins/Media/BinkMedia/Source/BinkMediaPlayer/Private/BinkMediaPlayerModule.cpp
@@ -17,7 +17,7 @@ DEFINE_LOG_CATEGORY(LogBink);
 
 TSharedPtr<FBinkMovieStreamer, ESPMode::ThreadSafe> MovieStreamer;
 
-TArray< FTexture2DRHIRef > BinkActiveTextureRefs;
+TArray< FTextureRHIRef > BinkActiveTextureRefs;
 
 #if BINKPLUGIN_UE4_EDITOR
 class UFactory;
diff --git a/Engine/Plugins/Media/BinkMedia/Source/BinkMediaPlayer/Private/BinkMediaPlayerPrivate.h b/Engine/Plugins/Media/BinkMedia/Source/BinkMediaPlayer/Private/BinkMediaPlayerPrivate.h
--- a/Engine/Plugins/Media/BinkMedia/Source/BinkMediaPlayer/Private/BinkMediaPlayerPrivate.h
+++ b/Engine/Plugins/Media/BinkMedia/Source/BinkMediaPlayer/Private/BinkMediaPlayerPrivate.h
@@ -23,7 +23,7 @@ extern BINKMEDIAPLAYER_API unsigned bink_gpu_api;
 extern BINKMEDIAPLAYER_API unsigned bink_gpu_api_hdr;
 extern BINKMEDIAPLAYER_API EPixelFormat bink_force_pixel_format;
 extern BINKMEDIAPLAYER_API FString BinkUE4CookOnTheFlyPath(FString path, const TCHAR *filename);
-extern BINKMEDIAPLAYER_API TArray< FTexture2DRHIRef > BinkActiveTextureRefs;
+extern BINKMEDIAPLAYER_API TArray< FTextureRHIRef > BinkActiveTextureRefs;
 extern BINKMEDIAPLAYER_API bool BinkInitialize();
 
 static int GetNumSpeakers() 
diff --git a/Engine/Plugins/Media/BinkMedia/Source/BinkMediaPlayer/Private/BinkMovieStreamer.cpp b/Engine/Plugins/Media/BinkMedia/Source/BinkMediaPlayer/Private/BinkMovieStreamer.cpp
--- a/Engine/Plugins/Media/BinkMedia/Source/BinkMediaPlayer/Private/BinkMovieStreamer.cpp
+++ b/Engine/Plugins/Media/BinkMedia/Source/BinkMediaPlayer/Private/BinkMovieStreamer.cpp
@@ -118,7 +118,7 @@ bool FBinkMovieStreamer::Tick(float DeltaTime)
 			CurrentTexture->InitResource(RHICmdList);
 		}
 
-		FTexture2DRHIRef tex = CurrentTexture->GetTypedResource();
+		FTextureRHIRef tex = CurrentTexture->GetTypedResource();
 		uint32 binkw = tex.GetReference()->GetSizeX();
 		uint32 binkh = tex.GetReference()->GetSizeY();
 		bool is_hdr = tex.GetReference()->GetFormat() != PF_B8G8R8A8;
diff --git a/Engine/Plugins/Media/BinkMedia/Source/BinkMediaPlayer/Private/binkpluginRHI.cpp b/Engine/Plugins/Media/BinkMedia/Source/BinkMediaPlayer/Private/binkpluginRHI.cpp
--- a/Engine/Plugins/Media/BinkMedia/Source/BinkMediaPlayer/Private/binkpluginRHI.cpp
+++ b/Engine/Plugins/Media/BinkMedia/Source/BinkMediaPlayer/Private/binkpluginRHI.cpp
@@ -15,7 +15,7 @@ static BINKSHADERS * shaders;
 static int attached = 0;
 static unsigned rfb;
 
-FRHITexture2D *BinkRHIRenderTarget;
+FRHITexture *BinkRHIRenderTarget;
 ERenderTargetLoadAction BinkRenderTargetLoadAction;
 
 RADDEFFUNC int setup_rhi( void * device, BINKPLUGININITINFO * info, S32 gpu_assisted, void ** context )
@@ -54,7 +54,7 @@ RADDEFFUNC void selectrendertarget_rhi(void* texture_target, S32 width, S32 heig
 	if (texture_target)
 	{
 		attached = 1;
-		BinkRHIRenderTarget = (FRHITexture2D*)texture_target;
+		BinkRHIRenderTarget = (FRHITexture*)texture_target;
 		BinkRenderTargetLoadAction = do_clear ? ERenderTargetLoadAction::EClear : ERenderTargetLoadAction::ELoad;
 	}
 }
@@ -64,7 +64,7 @@ RADDEFFUNC void selectscreenrendertarget_rhi( void * texture_target, S32 width,
 	if (texture_target)
 	{
 		attached = 2;
-		BinkRHIRenderTarget = (FRHITexture2D*)texture_target;
+		BinkRHIRenderTarget = (FRHITexture*)texture_target;
 		BinkRenderTargetLoadAction = ERenderTargetLoadAction::ELoad;
 	}
 }
diff --git a/Engine/Plugins/Media/BinkMedia/Source/BinkMediaPlayer/Private/binktexturesRHI.cpp b/Engine/Plugins/Media/BinkMedia/Source/BinkMediaPlayer/Private/binktexturesRHI.cpp
--- a/Engine/Plugins/Media/BinkMedia/Source/BinkMediaPlayer/Private/binktexturesRHI.cpp
+++ b/Engine/Plugins/Media/BinkMedia/Source/BinkMediaPlayer/Private/binktexturesRHI.cpp
@@ -13,7 +13,7 @@
 #include "RenderingThread.h"
 
 
-extern FRHITexture2D *BinkRHIRenderTarget;
+extern FRHITexture *BinkRHIRenderTarget;
 extern ERenderTargetLoadAction BinkRenderTargetLoadAction;
 
 FRDGTextureRef BinkRegisterExternalTexture(FRDGBuilder& GraphBuilder, FRHITexture* Texture, const TCHAR* NameIfUnregistered)
@@ -291,7 +291,7 @@ static void Finish_texture_update( BINKTEXTURES * ptextures )
 
 //-----------------------------------------------------------------------------
 
-static void update_plane_texture_rect(FRHICommandListImmediate& RHI, FRHITexture2D* RHITexture, BINKPLANE const* plane, unsigned w, unsigned h)
+static void update_plane_texture_rect(FRHICommandListImmediate& RHI, FRHITexture* RHITexture, BINKPLANE const* plane, unsigned w, unsigned h)
 {
 	uint32 Stride = 0;
 	unsigned char* TextureMemory = (unsigned char*)GDynamicRHI->LockTexture2D_RenderThread(RHI, RHITexture, 0, RLM_WriteOnly, Stride, false);
diff --git a/Engine/Plugins/Media/BinkMedia/Source/BinkMediaPlayer/Public/BinkMediaPlayer.h b/Engine/Plugins/Media/BinkMedia/Source/BinkMediaPlayer/Public/BinkMediaPlayer.h
--- a/Engine/Plugins/Media/BinkMedia/Source/BinkMediaPlayer/Public/BinkMediaPlayer.h
+++ b/Engine/Plugins/Media/BinkMedia/Source/BinkMediaPlayer/Public/BinkMediaPlayer.h
@@ -462,7 +462,7 @@ public:
 	bool Open(const FString& Url);
 	void Close();
 	bool IsReady() const { return bnk != NULL; }
-	void UpdateTexture(FRHICommandListImmediate &RHICmdList, FTexture2DRHIRef ref, void *nativePtr, int width, int height, bool isEditor, bool tonemap, int output_nits, float alpha, bool srgb_decode, bool is_hdr);
+	void UpdateTexture(FRHICommandListImmediate &RHICmdList, FTextureRHIRef ref, void *nativePtr, int width, int height, bool isEditor, bool tonemap, int output_nits, float alpha, bool srgb_decode, bool is_hdr);
 
 	FIntPoint GetDimensions() const;
 	float GetFrameRate() const;
diff --git a/Engine/Plugins/Media/BinkMedia/Source/BinkMediaPlayerEditor/Private/Models/BinkMediaPlayerEditorTexture.h b/Engine/Plugins/Media/BinkMedia/Source/BinkMediaPlayerEditor/Private/Models/BinkMediaPlayerEditorTexture.h
--- a/Engine/Plugins/Media/BinkMedia/Source/BinkMediaPlayerEditor/Private/Models/BinkMediaPlayerEditorTexture.h
+++ b/Engine/Plugins/Media/BinkMedia/Source/BinkMediaPlayerEditor/Private/Models/BinkMediaPlayerEditorTexture.h
@@ -28,7 +28,7 @@ struct FBinkMediaPlayerEditorTexture : public FTickableGameObject
 			{
 				SlateTexture->InitResource(RHICmdList);
 			}
-			FTexture2DRHIRef tex = SlateTexture->GetTypedResource();
+			FTextureRHIRef tex = SlateTexture->GetTypedResource();
 			if (!tex.GetReference())
 			{
 				return;
diff --git a/Engine/Plugins/Media/BlackmagicMedia/Source/BlackmagicMediaOutput/Private/BlackmagicMediaCapture.cpp b/Engine/Plugins/Media/BlackmagicMedia/Source/BlackmagicMediaOutput/Private/BlackmagicMediaCapture.cpp
--- a/Engine/Plugins/Media/BlackmagicMedia/Source/BlackmagicMediaOutput/Private/BlackmagicMediaCapture.cpp
+++ b/Engine/Plugins/Media/BlackmagicMedia/Source/BlackmagicMediaOutput/Private/BlackmagicMediaCapture.cpp
@@ -755,7 +755,7 @@ void UBlackmagicMediaCapture::LockDMATexture_RenderThread(FTextureRHIRef InTextu
 		{
 			TexturesToRelease.Add(InTexture);
 
-			FRHITexture2D* Texture = InTexture->GetTexture2D();
+			FRHITexture* Texture = InTexture->GetTexture2D();
 			UE::GPUTextureTransfer::FRegisterDMATextureArgs Args;
 			Args.RHITexture = Texture;
 
diff --git a/Engine/Plugins/Media/ElectraUtil/Source/ElectraSamples/Private/ElectraSamplesTextureManager.cpp b/Engine/Plugins/Media/ElectraUtil/Source/ElectraSamples/Private/ElectraSamplesTextureManager.cpp
--- a/Engine/Plugins/Media/ElectraUtil/Source/ElectraSamples/Private/ElectraSamplesTextureManager.cpp
+++ b/Engine/Plugins/Media/ElectraUtil/Source/ElectraSamples/Private/ElectraSamplesTextureManager.cpp
@@ -62,7 +62,7 @@ void FElectraSamplesTextureManager::PlatformRemoveSharedTextureRef(TSharedPtr<IV
 }
 
 
-FTexture2DRHIRef FElectraSamplesTextureManager::GetRHITextureFromPlatformTexture(void* PlatformTexture)
+FTextureRHIRef FElectraSamplesTextureManager::GetRHITextureFromPlatformTexture(void* PlatformTexture)
 {
 	FScopeLock Lock(&AccessCS);
 
@@ -75,5 +75,5 @@ FTexture2DRHIRef FElectraSamplesTextureManager::GetRHITextureFromPlatformTexture
 			return PlatformGetTextureRHI(PinnedTexture.Get());
 		}
 	}
-	return FTexture2DRHIRef();
+	return FTextureRHIRef();
 }
diff --git a/Engine/Plugins/Media/ElectraUtil/Source/ElectraSamples/Private/ElectraTextureSampleAVF.cpp b/Engine/Plugins/Media/ElectraUtil/Source/ElectraSamples/Private/ElectraTextureSampleAVF.cpp
--- a/Engine/Plugins/Media/ElectraUtil/Source/ElectraSamples/Private/ElectraTextureSampleAVF.cpp
+++ b/Engine/Plugins/Media/ElectraUtil/Source/ElectraSamples/Private/ElectraTextureSampleAVF.cpp
@@ -157,7 +157,7 @@ public:
 } // namespace anonymous
 
 
-void FElectraMediaTexConvApple::ConvertTexture(FTexture2DRHIRef & InDstTexture, CVImageBufferRef InImageBufferRef, bool bFullRange, EMediaTextureSampleFormat Format, const FMatrix44f& YUVMtx, const FMatrix44d& GamutToXYZMtx, UE::Color::EEncoding EncodingType, float NormalizationFactor)
+void FElectraMediaTexConvApple::ConvertTexture(FTextureRHIRef & InDstTexture, CVImageBufferRef InImageBufferRef, bool bFullRange, EMediaTextureSampleFormat Format, const FMatrix44f& YUVMtx, const FMatrix44d& GamutToXYZMtx, UE::Color::EEncoding EncodingType, float NormalizationFactor)
 {
 	FRHICommandListImmediate& RHICmdList = FRHICommandListImmediate::Get();
 
@@ -165,7 +165,7 @@ void FElectraMediaTexConvApple::ConvertTexture(FTexture2DRHIRef & InDstTexture,
 	const int32 FrameWidth = CVPixelBufferGetWidth(InImageBufferRef);
 	const int32 FrameStride = CVPixelBufferGetBytesPerRow(InImageBufferRef);
 
-	TRefCountPtr<FRHITexture2D> ShaderResource;
+	TRefCountPtr<FRHITexture> ShaderResource;
 
 	// We have to support Metal for this object now
 	check(COREVIDEO_SUPPORTS_METAL);
@@ -444,7 +444,7 @@ uint32 FElectraTextureSample::GetConverterInfoFlags() const
 	return ConverterInfoFlags_Default;
 }
 
-bool FElectraTextureSample::Convert(FTexture2DRHIRef & InDstTexture, const FConversionHints & Hints)
+bool FElectraTextureSample::Convert(FTextureRHIRef & InDstTexture, const FConversionHints & Hints)
 {
 	if (VideoDecoderOutput)
 	{
diff --git a/Engine/Plugins/Media/ElectraUtil/Source/ElectraSamples/Private/ElectraTextureSampleAndroid.cpp b/Engine/Plugins/Media/ElectraUtil/Source/ElectraSamples/Private/ElectraTextureSampleAndroid.cpp
--- a/Engine/Plugins/Media/ElectraUtil/Source/ElectraSamples/Private/ElectraTextureSampleAndroid.cpp
+++ b/Engine/Plugins/Media/ElectraUtil/Source/ElectraSamples/Private/ElectraTextureSampleAndroid.cpp
@@ -483,7 +483,7 @@ uint32 FElectraTextureSample::GetStride() const
 }
 
 
-bool FElectraTextureSample::Convert(FTexture2DRHIRef& InDstTexture, const FConversionHints& Hints)
+bool FElectraTextureSample::Convert(FTextureRHIRef& InDstTexture, const FConversionHints& Hints)
 {
 	FRHICommandListImmediate& RHICmdList = FRHICommandListImmediate::Get();
 
@@ -492,7 +492,7 @@ bool FElectraTextureSample::Convert(FTexture2DRHIRef& InDstTexture, const FConve
 		return false;
 	}
 
-	TRefCountPtr<FRHITexture2D> InputTexture;
+	TRefCountPtr<FRHITexture> InputTexture;
 
 	// Either use a texture we have around as a payload or make a temporary one from buffer contents...
 	if (!Texture.IsValid())
diff --git a/Engine/Plugins/Media/ElectraUtil/Source/ElectraSamples/Private/ElectraTextureSamplePC.cpp b/Engine/Plugins/Media/ElectraUtil/Source/ElectraSamples/Private/ElectraTextureSamplePC.cpp
--- a/Engine/Plugins/Media/ElectraUtil/Source/ElectraSamples/Private/ElectraTextureSamplePC.cpp
+++ b/Engine/Plugins/Media/ElectraUtil/Source/ElectraSamples/Private/ElectraTextureSamplePC.cpp
@@ -245,10 +245,10 @@ IMediaTextureSampleConverter* FElectraTextureSample::GetMediaTextureSampleConver
 struct FRHICommandCopyResourceDX11 final : public FRHICommand<FRHICommandCopyResourceDX11>
 {
 	TRefCountPtr<ID3D11Texture2D> SampleTexture;
-	FTexture2DRHIRef SampleDestinationTexture;
+	FTextureRHIRef SampleDestinationTexture;
 	bool bCrossDevice;
 
-	FRHICommandCopyResourceDX11(ID3D11Texture2D* InSampleTexture, FRHITexture2D* InSampleDestinationTexture, bool bInCrossDevice)
+	FRHICommandCopyResourceDX11(ID3D11Texture2D* InSampleTexture, FRHITexture* InSampleDestinationTexture, bool bInCrossDevice)
 		: SampleTexture(InSampleTexture)
 		, SampleDestinationTexture(InSampleDestinationTexture)
 		, bCrossDevice(bInCrossDevice)
@@ -334,7 +334,7 @@ struct FRHICommandCopyResourceDX11 final : public FRHICommand<FRHICommandCopyRes
 /**
  * "Converter" for textures - here: a copy from the decoder owned texture (possibly in another device) into a RHI one (as prep for the real conversion to RGB etc.)
  */
-bool FElectraTextureSample::Convert(FTexture2DRHIRef& InDstTexture, const FConversionHints& Hints)
+bool FElectraTextureSample::Convert(FTextureRHIRef& InDstTexture, const FConversionHints& Hints)
 {
 	LLM_SCOPE(ELLMTag::MediaStreaming);
 
diff --git a/Engine/Plugins/Media/ElectraUtil/Source/ElectraSamples/Public/Android/ElectraTextureSample.h b/Engine/Plugins/Media/ElectraUtil/Source/ElectraSamples/Public/Android/ElectraTextureSample.h
--- a/Engine/Plugins/Media/ElectraUtil/Source/ElectraSamples/Public/Android/ElectraTextureSample.h
+++ b/Engine/Plugins/Media/ElectraUtil/Source/ElectraSamples/Public/Android/ElectraTextureSample.h
@@ -77,7 +77,7 @@ private:
 	void InitializeTexture(EPixelFormat PixelFormat);
 	void SetupFromBuffer(const void* InBuffer, int32 InBufferSize);
 
-	virtual bool Convert(FTexture2DRHIRef& InDstTexture, const FConversionHints& Hints) override;
+	virtual bool Convert(FTextureRHIRef& InDstTexture, const FConversionHints& Hints) override;
 	virtual uint32 GetConverterInfoFlags() const
 	{
 		return ConverterInfoFlags_Default;
@@ -91,7 +91,7 @@ private:
 	int32 BufferSize;
 
 	/** Texture resource. */
-	TRefCountPtr<FRHITexture2D> Texture;
+	TRefCountPtr<FRHITexture> Texture;
 };
 
 
diff --git a/Engine/Plugins/Media/ElectraUtil/Source/ElectraSamples/Public/Apple/ElectraTextureSample.h b/Engine/Plugins/Media/ElectraUtil/Source/ElectraSamples/Public/Apple/ElectraTextureSample.h
--- a/Engine/Plugins/Media/ElectraUtil/Source/ElectraSamples/Public/Apple/ElectraTextureSample.h
+++ b/Engine/Plugins/Media/ElectraUtil/Source/ElectraSamples/Public/Apple/ElectraTextureSample.h
@@ -50,7 +50,7 @@ public:
 
 private:
 	/** The sample's texture resource. */
-	TRefCountPtr<FRHITexture2D> Texture;
+	TRefCountPtr<FRHITexture> Texture;
 
 	/** Output data from video decoder. */
 	FVideoDecoderOutputApple* VideoDecoderOutputApple;
@@ -58,7 +58,7 @@ private:
 	TWeakPtr<FElectraMediaTexConvApple, ESPMode::ThreadSafe> TexConv;
 
 	virtual uint32 GetConverterInfoFlags() const override;
-    virtual bool Convert(FTexture2DRHIRef & InDstTexture, const FConversionHints & Hints) override;
+    virtual bool Convert(FTextureRHIRef & InDstTexture, const FConversionHints & Hints) override;
 };
 
 
@@ -73,7 +73,7 @@ public:
     ~FElectraMediaTexConvApple();
 
 #if WITH_ENGINE
-    void ConvertTexture(FTexture2DRHIRef & InDstTexture, CVImageBufferRef InImageBufferRef, bool bFullRange, EMediaTextureSampleFormat Format, const FMatrix44f& YUVMtx, const FMatrix44d& GamutToXYZMtx, UE::Color::EEncoding EncodingType, float NormalizationFactor);
+    void ConvertTexture(FTextureRHIRef & InDstTexture, CVImageBufferRef InImageBufferRef, bool bFullRange, EMediaTextureSampleFormat Format, const FMatrix44f& YUVMtx, const FMatrix44d& GamutToXYZMtx, UE::Color::EEncoding EncodingType, float NormalizationFactor);
 #endif
 
 private:
diff --git a/Engine/Plugins/Media/ElectraUtil/Source/ElectraSamples/Public/Linux/ElectraTextureSample.h b/Engine/Plugins/Media/ElectraUtil/Source/ElectraSamples/Public/Linux/ElectraTextureSample.h
--- a/Engine/Plugins/Media/ElectraUtil/Source/ElectraSamples/Public/Linux/ElectraTextureSample.h
+++ b/Engine/Plugins/Media/ElectraUtil/Source/ElectraSamples/Public/Linux/ElectraTextureSample.h
@@ -48,7 +48,7 @@ public:
 
 private:
 	/** The sample's texture resource. */
-	TRefCountPtr<FRHITexture2D> Texture;
+	TRefCountPtr<FRHITexture> Texture;
 
 	/** Output data from video decoder. */
 	FVideoDecoderOutputLinux* VideoDecoderOutputLinux;
diff --git a/Engine/Plugins/Media/ElectraUtil/Source/ElectraSamples/Public/Windows/ElectraTextureSample.h b/Engine/Plugins/Media/ElectraUtil/Source/ElectraSamples/Public/Windows/ElectraTextureSample.h
--- a/Engine/Plugins/Media/ElectraUtil/Source/ElectraSamples/Public/Windows/ElectraTextureSample.h
+++ b/Engine/Plugins/Media/ElectraUtil/Source/ElectraSamples/Public/Windows/ElectraTextureSample.h
@@ -56,7 +56,7 @@ public:
 private:
 	virtual float GetSampleDataScale(bool b10Bit) const override;
 
-	virtual bool Convert(FTexture2DRHIRef & InDstTexture, const FConversionHints & Hints) override;
+	virtual bool Convert(FTextureRHIRef & InDstTexture, const FConversionHints & Hints) override;
 	virtual uint32 GetConverterInfoFlags() const
 	{
 		return ConverterInfoFlags_PreprocessOnly;
@@ -66,7 +66,7 @@ private:
 	EMediaTextureSampleFormat SampleFormat;
 
 	/** Destination Texture resource (from Rendering device) */
-	mutable FTexture2DRHIRef Texture;
+	mutable FTextureRHIRef Texture;
 
 	/** True if texture format could support sRGB conversion in HW */
 	bool bCanUseSRGB;
diff --git a/Engine/Plugins/Media/ImgMedia/Source/ImgMedia/Private/Readers/ExrImgMediaReaderGpu.cpp b/Engine/Plugins/Media/ImgMedia/Source/ImgMedia/Private/Readers/ExrImgMediaReaderGpu.cpp
--- a/Engine/Plugins/Media/ImgMedia/Source/ImgMedia/Private/Readers/ExrImgMediaReaderGpu.cpp
+++ b/Engine/Plugins/Media/ImgMedia/Source/ImgMedia/Private/Readers/ExrImgMediaReaderGpu.cpp
@@ -468,7 +468,7 @@ SIZE_T FExrImgMediaReaderGpu::GetBufferSize(const FIntPoint& Dim, int32 NumChann
 
 void FExrImgMediaReaderGpu::CreateSampleConverterCallback(TSharedPtr<FExrMediaTextureSampleConverter, ESPMode::ThreadSafe> SampleConverter)
 {
-	auto RenderThreadSwizzler = [] (FRHICommandListImmediate& RHICmdList, FTexture2DRHIRef RenderTargetTextureRHI, TMap<int32, FStructuredBufferPoolItemSharedPtr>& MipBuffers, const FSampleConverterParameters ConverterParams)->bool
+	auto RenderThreadSwizzler = [] (FRHICommandListImmediate& RHICmdList, FTextureRHIRef RenderTargetTextureRHI, TMap<int32, FStructuredBufferPoolItemSharedPtr>& MipBuffers, const FSampleConverterParameters ConverterParams)->bool
 	{
 		TRACE_CPUPROFILER_EVENT_SCOPE_TEXT(*FString::Printf(TEXT("ExrReaderGpu.Convert %d"), ConverterParams.FrameId));
 		SCOPED_DRAW_EVENT(RHICmdList, FExrImgMediaReaderGpu_Convert);
@@ -476,7 +476,7 @@ void FExrImgMediaReaderGpu::CreateSampleConverterCallback(TSharedPtr<FExrMediaTe
 
 		auto RenderMip = []
 			( FRHICommandListImmediate& RHICmdList
-			, FTexture2DRHIRef RenderTargetTextureRHI
+			, FTextureRHIRef RenderTargetTextureRHI
 			, const FSampleConverterParameters & ConverterParams
 			, int32 SampleMipLevel
 			, int32 TextureMipLevel
@@ -691,7 +691,7 @@ void FExrImgMediaReaderGpu::ReturnGpuBufferToPool(uint32 AllocSize, FStructuredB
 /* FExrMediaTextureSampleConverter implementation
  *****************************************************************************/
 
-bool FExrMediaTextureSampleConverter::Convert(FTexture2DRHIRef& InDstTexture, const FConversionHints& Hints)
+bool FExrMediaTextureSampleConverter::Convert(FTextureRHIRef& InDstTexture, const FConversionHints& Hints)
 {
 	FScopeLock ScopeLock(&ConverterCallbacksCriticalSection);
 	bool bExecutionSuccessful = false;
diff --git a/Engine/Plugins/Media/ImgMedia/Source/ImgMedia/Private/Readers/ExrImgMediaReaderGpu.h b/Engine/Plugins/Media/ImgMedia/Source/ImgMedia/Private/Readers/ExrImgMediaReaderGpu.h
--- a/Engine/Plugins/Media/ImgMedia/Source/ImgMedia/Private/Readers/ExrImgMediaReaderGpu.h
+++ b/Engine/Plugins/Media/ImgMedia/Source/ImgMedia/Private/Readers/ExrImgMediaReaderGpu.h
@@ -179,13 +179,13 @@ struct FSampleConverterParameters
 };
 
 
-FUNC_DECLARE_DELEGATE(FExrConvertBufferCallback, bool, FRHICommandListImmediate& /*RHICmdList*/, FTexture2DRHIRef /*RenderTargetTextureRHI*/, TMap<int32, FStructuredBufferPoolItemSharedPtr>& /*MipBuffers*/, const FSampleConverterParameters /*ConverterParams*/)
+FUNC_DECLARE_DELEGATE(FExrConvertBufferCallback, bool, FRHICommandListImmediate& /*RHICmdList*/, FTextureRHIRef /*RenderTargetTextureRHI*/, TMap<int32, FStructuredBufferPoolItemSharedPtr>& /*MipBuffers*/, const FSampleConverterParameters /*ConverterParams*/)
 
 class FExrMediaTextureSampleConverter: public IMediaTextureSampleConverter
 {
 
 public:
-	virtual bool Convert(FTexture2DRHIRef& InDstTexture, const FConversionHints& Hints) override;
+	virtual bool Convert(FTextureRHIRef& InDstTexture, const FConversionHints& Hints) override;
 	virtual ~FExrMediaTextureSampleConverter() {};
 	
 	void AddCallback(FExrConvertBufferCallback&& Callback) 
diff --git a/Engine/Plugins/Media/MediaIOFramework/Source/MediaIOCore/Private/MediaCapture.cpp b/Engine/Plugins/Media/MediaIOFramework/Source/MediaIOCore/Private/MediaCapture.cpp
--- a/Engine/Plugins/Media/MediaIOFramework/Source/MediaIOCore/Private/MediaCapture.cpp
+++ b/Engine/Plugins/Media/MediaIOFramework/Source/MediaIOCore/Private/MediaCapture.cpp
@@ -1052,7 +1052,7 @@ void UMediaCapture::PrepareAndDispatchCapture_GameThread(const TSharedPtr<UE::Me
 				{
 					FRDGBuilder GraphBuilder(RHICmdList);
 					
-					FTexture2DRHIRef SourceTexture = InMediaCapture->CaptureSource->GetSourceTextureForInput_RenderThread(RHICmdList);
+					FTextureRHIRef SourceTexture = InMediaCapture->CaptureSource->GetSourceTextureForInput_RenderThread(RHICmdList);
 					
 					FCaptureFrameArgs CaptureArgs{ GraphBuilder };
 					CaptureArgs.MediaCapture = InMediaCapture;
diff --git a/Engine/Plugins/Media/MediaIOFramework/Source/MediaIOCore/Private/MediaCaptureSources.h b/Engine/Plugins/Media/MediaIOFramework/Source/MediaIOCore/Private/MediaCaptureSources.h
--- a/Engine/Plugins/Media/MediaIOFramework/Source/MediaIOCore/Private/MediaCaptureSources.h
+++ b/Engine/Plugins/Media/MediaIOFramework/Source/MediaIOCore/Private/MediaCaptureSources.h
@@ -167,7 +167,7 @@ namespace UE::MediaCapture::Private
 		}
 		
 		/** Only returns a valid texture for scene viewports and render targets as the RHI Sources don't keep a pointer to the RHI texture. */
-		virtual FTexture2DRHIRef GetSourceTextureForInput_RenderThread(FRHICommandListImmediate& RHICmdList)
+		virtual FTextureRHIRef GetSourceTextureForInput_RenderThread(FRHICommandListImmediate& RHICmdList)
 		{
 			return nullptr;
 		}
@@ -278,9 +278,9 @@ namespace UE::MediaCapture::Private
 			return false;
 		}
 
-		virtual FTexture2DRHIRef GetSourceTextureForInput_RenderThread(FRHICommandListImmediate& RHICmdList) override
+		virtual FTextureRHIRef GetSourceTextureForInput_RenderThread(FRHICommandListImmediate& RHICmdList) override
 		{
-			FTexture2DRHIRef SourceTexture;
+			FTextureRHIRef SourceTexture;
 			if (TSharedPtr<FSceneViewport> Viewport = WeakViewport.Pin())
 			{
 #if WITH_EDITOR
@@ -402,7 +402,7 @@ namespace UE::MediaCapture::Private
 			return false;
 		}
 
-		virtual FTexture2DRHIRef GetSourceTextureForInput_RenderThread(FRHICommandListImmediate& RHICmdList)
+		virtual FTextureRHIRef GetSourceTextureForInput_RenderThread(FRHICommandListImmediate& RHICmdList)
 		{
 			constexpr bool bEvenIfPendingKill = false;
 			constexpr bool bThreadSafeTest = true;
diff --git a/Engine/Plugins/Media/MediaIOFramework/Source/MediaIOCore/Private/MediaIOCoreTextureSampleConverter.cpp b/Engine/Plugins/Media/MediaIOFramework/Source/MediaIOCore/Private/MediaIOCoreTextureSampleConverter.cpp
--- a/Engine/Plugins/Media/MediaIOFramework/Source/MediaIOCore/Private/MediaIOCoreTextureSampleConverter.cpp
+++ b/Engine/Plugins/Media/MediaIOFramework/Source/MediaIOCore/Private/MediaIOCoreTextureSampleConverter.cpp
@@ -12,7 +12,7 @@ void FMediaIOCoreTextureSampleConverter::Setup(const TSharedPtr<FMediaIOCoreText
 	JITRProxySample = InSample;
 }
 
-bool FMediaIOCoreTextureSampleConverter::Convert(FTexture2DRHIRef& InDestinationTexture, const FConversionHints& Hints)
+bool FMediaIOCoreTextureSampleConverter::Convert(FTextureRHIRef& InDestinationTexture, const FConversionHints& Hints)
 {
 	TRACE_CPUPROFILER_EVENT_SCOPE(JITR);
 
diff --git a/Engine/Plugins/Media/MediaIOFramework/Source/MediaIOCore/Private/MediaIOFrameManager.h b/Engine/Plugins/Media/MediaIOFramework/Source/MediaIOCore/Private/MediaIOFrameManager.h
--- a/Engine/Plugins/Media/MediaIOFramework/Source/MediaIOCore/Private/MediaIOFrameManager.h
+++ b/Engine/Plugins/Media/MediaIOFramework/Source/MediaIOCore/Private/MediaIOFrameManager.h
@@ -22,7 +22,7 @@ namespace UE::MediaCaptureData
 	{
 		FRDGBuilder& GraphBuilder;
 		TObjectPtr<UMediaCapture> MediaCapture = nullptr;
-		FTexture2DRHIRef ResourceToCapture;
+		FTextureRHIRef ResourceToCapture;
 		FRDGTextureRef RDGResourceToCapture = nullptr;
 		FIntPoint DesiredSize = FIntPoint::ZeroValue;
 		FIntRect SourceViewRect{ 0,0,0,0 };
diff --git a/Engine/Plugins/Media/MediaIOFramework/Source/MediaIOCore/Private/Shared/MediaIOCoreTextureSampleBase.cpp b/Engine/Plugins/Media/MediaIOFramework/Source/MediaIOCore/Private/Shared/MediaIOCoreTextureSampleBase.cpp
--- a/Engine/Plugins/Media/MediaIOFramework/Source/MediaIOCore/Private/Shared/MediaIOCoreTextureSampleBase.cpp
+++ b/Engine/Plugins/Media/MediaIOFramework/Source/MediaIOCore/Private/Shared/MediaIOCoreTextureSampleBase.cpp
@@ -207,7 +207,7 @@ void FMediaIOCoreTextureSampleBase::CopyConfiguration(const TSharedPtr<FMediaIOC
 	OriginalSample = SourceSample;
 }
 
-bool FMediaIOCoreTextureSampleBase::ApplyColorConversion(FTexture2DRHIRef& InSrcTexture, FTexture2DRHIRef& InDstTexture)
+bool FMediaIOCoreTextureSampleBase::ApplyColorConversion(FTextureRHIRef& InSrcTexture, FTextureRHIRef& InDstTexture)
 {
 	FRHICommandListImmediate& RHICmdList = FRHICommandListExecutor::GetImmediateCommandList();
 
diff --git a/Engine/Plugins/Media/MediaIOFramework/Source/MediaIOCore/Public/MediaIOCoreTextureSampleBase.h b/Engine/Plugins/Media/MediaIOFramework/Source/MediaIOCore/Public/MediaIOCoreTextureSampleBase.h
--- a/Engine/Plugins/Media/MediaIOFramework/Source/MediaIOCore/Public/MediaIOCoreTextureSampleBase.h
+++ b/Engine/Plugins/Media/MediaIOFramework/Source/MediaIOCore/Public/MediaIOCoreTextureSampleBase.h
@@ -340,7 +340,7 @@ public:
 	}
 
 	//~ IMediaTextureSampleColorConverter interface
-	virtual bool ApplyColorConversion(FTexture2DRHIRef& InSrcTexture, FTexture2DRHIRef& InDstTexture) override;
+	virtual bool ApplyColorConversion(FTextureRHIRef& InSrcTexture, FTextureRHIRef& InDstTexture) override;
 
 	void* GetMutableBuffer()
 	{
diff --git a/Engine/Plugins/Media/MediaIOFramework/Source/MediaIOCore/Public/MediaIOCoreTextureSampleConverter.h b/Engine/Plugins/Media/MediaIOFramework/Source/MediaIOCore/Public/MediaIOCoreTextureSampleConverter.h
--- a/Engine/Plugins/Media/MediaIOFramework/Source/MediaIOCore/Public/MediaIOCoreTextureSampleConverter.h
+++ b/Engine/Plugins/Media/MediaIOFramework/Source/MediaIOCore/Public/MediaIOCoreTextureSampleConverter.h
@@ -29,7 +29,7 @@ public:
 
 public:
 	//~ Begin IMediaTextureSampleConverter interface
-	virtual bool Convert(FTexture2DRHIRef& InDstTexture, const FConversionHints& Hints) override;
+	virtual bool Convert(FTextureRHIRef& InDstTexture, const FConversionHints& Hints) override;
 	virtual uint32 GetConverterInfoFlags() const override;
 	//~ End IMediaTextureSampleConverter interface
 
diff --git a/Engine/Plugins/Media/MediaMovieStreamer/Source/MediaMovieStreamer/Private/MediaMovieStreamer.cpp b/Engine/Plugins/Media/MediaMovieStreamer/Source/MediaMovieStreamer/Private/MediaMovieStreamer.cpp
--- a/Engine/Plugins/Media/MediaMovieStreamer/Source/MediaMovieStreamer/Private/MediaMovieStreamer.cpp
+++ b/Engine/Plugins/Media/MediaMovieStreamer/Source/MediaMovieStreamer/Private/MediaMovieStreamer.cpp
@@ -148,7 +148,7 @@ bool FMediaMovieStreamer::Tick(float DeltaTime)
 			FTextureResource* TextureResource = MediaTexture->GetResource();
 			if (TextureResource != nullptr)
 			{
-				FRHITexture2D* RHITexture2D = TextureResource->GetTexture2DRHI();
+				FRHITexture* RHITexture2D = TextureResource->GetTexture2DRHI();
 				if (RHITexture2D != nullptr)
 				{
 					// Get slate texture.
@@ -161,7 +161,7 @@ bool FMediaMovieStreamer::Tick(float DeltaTime)
 						}
 
 						// Update the slate texture.
-						FTexture2DRHIRef ref = RHITexture2D;
+						FTextureRHIRef ref = RHITexture2D;
 						CurrentTexture->SetRHIRef(ref, MediaTexture->GetWidth(), MediaTexture->GetHeight());
 
 						// Update viewport.
@@ -241,7 +241,7 @@ void FMediaMovieStreamer::Cleanup()
 	}
 }
 
-FTexture2DRHIRef FMediaMovieStreamer::GetTexture()
+FTextureRHIRef FMediaMovieStreamer::GetTexture()
 {
 	return nullptr;
 }
diff --git a/Engine/Plugins/Media/MediaMovieStreamer/Source/MediaMovieStreamer/Public/MediaMovieStreamer.h b/Engine/Plugins/Media/MediaMovieStreamer/Source/MediaMovieStreamer/Public/MediaMovieStreamer.h
--- a/Engine/Plugins/Media/MediaMovieStreamer/Source/MediaMovieStreamer/Public/MediaMovieStreamer.h
+++ b/Engine/Plugins/Media/MediaMovieStreamer/Source/MediaMovieStreamer/Public/MediaMovieStreamer.h
@@ -91,7 +91,7 @@ public:
 	virtual FString GetMovieName() override;
 	virtual bool IsLastMovieInPlaylist() override;
 	virtual void Cleanup() override;
-	virtual FTexture2DRHIRef GetTexture() override;
+	virtual FTextureRHIRef GetTexture() override;
 	virtual FOnCurrentMovieClipFinished& OnCurrentMovieClipFinished() override;
 
 	virtual void TickPreEngine() override;
diff --git a/Engine/Plugins/Media/PixelCapture/Source/PixelCapture/Private/PixelCaptureCapturerRHIToI420CPU.cpp b/Engine/Plugins/Media/PixelCapture/Source/PixelCapture/Private/PixelCaptureCapturerRHIToI420CPU.cpp
--- a/Engine/Plugins/Media/PixelCapture/Source/PixelCapture/Private/PixelCaptureCapturerRHIToI420CPU.cpp
+++ b/Engine/Plugins/Media/PixelCapture/Source/PixelCapture/Private/PixelCaptureCapturerRHIToI420CPU.cpp
@@ -67,7 +67,7 @@ void FPixelCaptureCapturerRHIToI420CPU::BeginProcess(const IPixelCaptureInputFra
 	MarkCPUWorkStart();
 
 	const FPixelCaptureInputFrameRHI& RHISourceFrame = StaticCast<const FPixelCaptureInputFrameRHI&>(InputFrame);
-	FTexture2DRHIRef SourceTexture = RHISourceFrame.FrameTexture;
+	FTextureRHIRef SourceTexture = RHISourceFrame.FrameTexture;
 
 	FRHICommandListImmediate& RHICmdList = FRHICommandListExecutor::GetImmediateCommandList();
 
diff --git a/Engine/Plugins/Media/PixelCapture/Source/PixelCapture/Private/PixelCaptureCapturerRHIToI420Compute.cpp b/Engine/Plugins/Media/PixelCapture/Source/PixelCapture/Private/PixelCaptureCapturerRHIToI420Compute.cpp
--- a/Engine/Plugins/Media/PixelCapture/Source/PixelCapture/Private/PixelCaptureCapturerRHIToI420Compute.cpp
+++ b/Engine/Plugins/Media/PixelCapture/Source/PixelCapture/Private/PixelCaptureCapturerRHIToI420Compute.cpp
@@ -111,7 +111,7 @@ void FPixelCaptureCapturerRHIToI420Compute::BeginProcess(const IPixelCaptureInpu
 	MarkCPUWorkStart();
 
 	const FPixelCaptureInputFrameRHI& RHISourceFrame = StaticCast<const FPixelCaptureInputFrameRHI&>(InputFrame);
-	FTexture2DRHIRef SourceTexture = RHISourceFrame.FrameTexture;
+	FTextureRHIRef SourceTexture = RHISourceFrame.FrameTexture;
 
 	FRHICommandListImmediate& RHICmdList = FRHICommandListExecutor::GetImmediateCommandList();
 
diff --git a/Engine/Plugins/Media/PixelCapture/Source/PixelCapture/Private/PixelCaptureInputFrameRHI.cpp b/Engine/Plugins/Media/PixelCapture/Source/PixelCapture/Private/PixelCaptureInputFrameRHI.cpp
--- a/Engine/Plugins/Media/PixelCapture/Source/PixelCapture/Private/PixelCaptureInputFrameRHI.cpp
+++ b/Engine/Plugins/Media/PixelCapture/Source/PixelCapture/Private/PixelCaptureInputFrameRHI.cpp
@@ -3,7 +3,7 @@
 #include "PixelCaptureInputFrameRHI.h"
 #include "PixelCaptureBufferFormat.h"
 
-FPixelCaptureInputFrameRHI::FPixelCaptureInputFrameRHI(FTexture2DRHIRef InFrameTexture)
+FPixelCaptureInputFrameRHI::FPixelCaptureInputFrameRHI(FTextureRHIRef InFrameTexture)
 	: FrameTexture(InFrameTexture)
 {
 	Metadata.SourceTime = rtc::TimeMillis();
diff --git a/Engine/Plugins/Media/PixelCapture/Source/PixelCapture/Public/PixelCaptureInputFrameRHI.h b/Engine/Plugins/Media/PixelCapture/Source/PixelCapture/Public/PixelCaptureInputFrameRHI.h
--- a/Engine/Plugins/Media/PixelCapture/Source/PixelCapture/Public/PixelCaptureInputFrameRHI.h
+++ b/Engine/Plugins/Media/PixelCapture/Source/PixelCapture/Public/PixelCaptureInputFrameRHI.h
@@ -11,12 +11,12 @@
 class PIXELCAPTURE_API FPixelCaptureInputFrameRHI : public IPixelCaptureInputFrame
 {
 public:
-	FPixelCaptureInputFrameRHI(FTexture2DRHIRef InFrameTexture);
+	FPixelCaptureInputFrameRHI(FTextureRHIRef InFrameTexture);
 	virtual ~FPixelCaptureInputFrameRHI() = default;
 
 	virtual int32 GetType() const override;
 	virtual int32 GetWidth() const override;
 	virtual int32 GetHeight() const override;
 
-	FTexture2DRHIRef FrameTexture;
+	FTextureRHIRef FrameTexture;
 };
diff --git a/Engine/Plugins/Media/PixelCapture/Source/PixelCapture/Public/PixelCaptureOutputFrameRHI.h b/Engine/Plugins/Media/PixelCapture/Source/PixelCapture/Public/PixelCaptureOutputFrameRHI.h
--- a/Engine/Plugins/Media/PixelCapture/Source/PixelCapture/Public/PixelCaptureOutputFrameRHI.h
+++ b/Engine/Plugins/Media/PixelCapture/Source/PixelCapture/Public/PixelCaptureOutputFrameRHI.h
@@ -24,12 +24,12 @@ public:
 		EPixelFormat Format;
 	};
 
-	FPixelCaptureOutputFrameRHI(FTexture2DRHIRef InFrameTexture)
+	FPixelCaptureOutputFrameRHI(FTextureRHIRef InFrameTexture)
 		: FrameTexture(InFrameTexture), bHasWrapperFormatData(false), WrapperData()
 	{
 	}
 
-	FPixelCaptureOutputFrameRHI(FTexture2DRHIRef InFrameTexture, WrapperFormatData InWrapperData)
+	FPixelCaptureOutputFrameRHI(FTextureRHIRef InFrameTexture, WrapperFormatData InWrapperData)
 		: FrameTexture(InFrameTexture), bHasWrapperFormatData(true), WrapperData(InWrapperData)
 	{
 	}
@@ -55,12 +55,12 @@ public:
 		return false;
 	}
 
-	FTexture2DRHIRef GetFrameTexture() const { return FrameTexture; }
+	FTextureRHIRef GetFrameTexture() const { return FrameTexture; }
 
-	void SetFrameTexture(FTexture2DRHIRef InFrameTexture) { FrameTexture = InFrameTexture; }
+	void SetFrameTexture(FTextureRHIRef InFrameTexture) { FrameTexture = InFrameTexture; }
 
 private:
-	FTexture2DRHIRef FrameTexture;
+	FTextureRHIRef FrameTexture;
 	bool bHasWrapperFormatData;
 	WrapperFormatData WrapperData;
 };
diff --git a/Engine/Plugins/Media/PixelStreaming/Source/PixelStreaming/Private/PixelStreamingVideoInputBackBuffer.cpp b/Engine/Plugins/Media/PixelStreaming/Source/PixelStreaming/Private/PixelStreamingVideoInputBackBuffer.cpp
--- a/Engine/Plugins/Media/PixelStreaming/Source/PixelStreaming/Private/PixelStreamingVideoInputBackBuffer.cpp
+++ b/Engine/Plugins/Media/PixelStreaming/Source/PixelStreaming/Private/PixelStreamingVideoInputBackBuffer.cpp
@@ -43,7 +43,7 @@ FPixelStreamingVideoInputBackBuffer::~FPixelStreamingVideoInputBackBuffer()
 	}
 }
 
-void FPixelStreamingVideoInputBackBuffer::OnBackBufferReady(SWindow& SlateWindow, const FTexture2DRHIRef& FrameBuffer)
+void FPixelStreamingVideoInputBackBuffer::OnBackBufferReady(SWindow& SlateWindow, const FTextureRHIRef& FrameBuffer)
 {
 	OnFrame(FPixelCaptureInputFrameRHI(FrameBuffer));
 }
diff --git a/Engine/Plugins/Media/PixelStreaming/Source/PixelStreaming/Private/PixelStreamingVideoInputRenderTarget.cpp b/Engine/Plugins/Media/PixelStreaming/Source/PixelStreaming/Private/PixelStreamingVideoInputRenderTarget.cpp
--- a/Engine/Plugins/Media/PixelStreaming/Source/PixelStreaming/Private/PixelStreamingVideoInputRenderTarget.cpp
+++ b/Engine/Plugins/Media/PixelStreaming/Source/PixelStreaming/Private/PixelStreamingVideoInputRenderTarget.cpp
@@ -25,7 +25,7 @@ void FPixelStreamingVideoInputRenderTarget::OnEndFrameRenderThread()
 {
 	if (Target)
 	{
-		if (FTexture2DRHIRef Texture = Target->GetResource()->GetTexture2DRHI())
+		if (FTextureRHIRef Texture = Target->GetResource()->GetTexture2DRHI())
 		{
 			OnFrame(FPixelCaptureInputFrameRHI(Texture));
 		}
diff --git a/Engine/Plugins/Media/PixelStreaming/Source/PixelStreaming/Public/PixelStreamingVideoInputBackBuffer.h b/Engine/Plugins/Media/PixelStreaming/Source/PixelStreaming/Public/PixelStreamingVideoInputBackBuffer.h
--- a/Engine/Plugins/Media/PixelStreaming/Source/PixelStreaming/Public/PixelStreamingVideoInputBackBuffer.h
+++ b/Engine/Plugins/Media/PixelStreaming/Source/PixelStreaming/Public/PixelStreamingVideoInputBackBuffer.h
@@ -20,7 +20,7 @@ public:
 private:
 	FPixelStreamingVideoInputBackBuffer() = default;
 
-	void OnBackBufferReady(SWindow& SlateWindow, const FTexture2DRHIRef& FrameBuffer);
+	void OnBackBufferReady(SWindow& SlateWindow, const FTextureRHIRef& FrameBuffer);
 
 	FDelegateHandle DelegateHandle;
 };
diff --git a/Engine/Plugins/Media/PixelStreaming/Source/PixelStreamingEditor/Private/PixelStreamingEditorUtils.cpp b/Engine/Plugins/Media/PixelStreaming/Source/PixelStreamingEditor/Private/PixelStreamingEditorUtils.cpp
--- a/Engine/Plugins/Media/PixelStreaming/Source/PixelStreamingEditor/Private/PixelStreamingEditorUtils.cpp
+++ b/Engine/Plugins/Media/PixelStreaming/Source/PixelStreamingEditor/Private/PixelStreamingEditorUtils.cpp
@@ -47,7 +47,7 @@ namespace UE::EditorPixelStreaming
 		}
 	}
 
-	const FString HashWindow(SWindow& SlateWindow, const FTexture2DRHIRef& FrameBuffer)
+	const FString HashWindow(SWindow& SlateWindow, const FTextureRHIRef& FrameBuffer)
 	{
 		return FString::Printf(TEXT("%s-%s-%dx%d"), ToString(SlateWindow.GetType()), *SlateWindow.GetTitle().ToString(), FrameBuffer->GetSizeXY().X, FrameBuffer->GetSizeXY().Y);
 	}
diff --git a/Engine/Plugins/Media/PixelStreaming/Source/PixelStreamingEditor/Private/PixelStreamingVideoInputBackBufferComposited.cpp b/Engine/Plugins/Media/PixelStreaming/Source/PixelStreamingEditor/Private/PixelStreamingVideoInputBackBufferComposited.cpp
--- a/Engine/Plugins/Media/PixelStreaming/Source/PixelStreamingEditor/Private/PixelStreamingVideoInputBackBufferComposited.cpp
+++ b/Engine/Plugins/Media/PixelStreaming/Source/PixelStreamingEditor/Private/PixelStreamingVideoInputBackBufferComposited.cpp
@@ -97,7 +97,7 @@ FPixelStreamingVideoInputBackBufferComposited::~FPixelStreamingVideoInputBackBuf
 	}
 }
 
-void FPixelStreamingVideoInputBackBufferComposited::OnBackBufferReady(SWindow& SlateWindow, const FTexture2DRHIRef& FrameBuffer)
+void FPixelStreamingVideoInputBackBufferComposited::OnBackBufferReady(SWindow& SlateWindow, const FTextureRHIRef& FrameBuffer)
 {
 	/**
 	 * When we receive a texture from this delegate, the texture will undergo a two copy process.
diff --git a/Engine/Plugins/Media/PixelStreaming/Source/PixelStreamingEditor/Public/PixelStreamingEditorUtils.h b/Engine/Plugins/Media/PixelStreaming/Source/PixelStreamingEditor/Public/PixelStreamingEditorUtils.h
--- a/Engine/Plugins/Media/PixelStreaming/Source/PixelStreamingEditor/Public/PixelStreamingEditorUtils.h
+++ b/Engine/Plugins/Media/PixelStreaming/Source/PixelStreamingEditor/Public/PixelStreamingEditorUtils.h
@@ -18,5 +18,5 @@ namespace UE::EditorPixelStreaming
 
 	FString ToString(EStreamTypes StreamType);
 	const TCHAR* ToString(EWindowType Type);
-	const FString HashWindow(SWindow& SlateWindow, const FTexture2DRHIRef& FrameBuffer);
+	const FString HashWindow(SWindow& SlateWindow, const FTextureRHIRef& FrameBuffer);
 }
\ No newline at end of file
diff --git a/Engine/Plugins/Media/PixelStreaming/Source/PixelStreamingEditor/Public/PixelStreamingVideoInputBackBufferComposited.h b/Engine/Plugins/Media/PixelStreaming/Source/PixelStreamingEditor/Public/PixelStreamingVideoInputBackBufferComposited.h
--- a/Engine/Plugins/Media/PixelStreaming/Source/PixelStreamingEditor/Public/PixelStreamingVideoInputBackBufferComposited.h
+++ b/Engine/Plugins/Media/PixelStreaming/Source/PixelStreamingEditor/Public/PixelStreamingVideoInputBackBufferComposited.h
@@ -54,7 +54,7 @@ private:
     FPixelStreamingVideoInputBackBufferComposited();
     void CompositeWindows();
 
-    void OnBackBufferReady(SWindow &SlateWindow, const FTexture2DRHIRef &FrameBuffer);
+    void OnBackBufferReady(SWindow &SlateWindow, const FTextureRHIRef &FrameBuffer);
 
     FDelegateHandle DelegateHandle;
 
diff --git a/Engine/Plugins/Media/WebMMedia/Source/WebMMedia/Private/Player/WebMMediaTextureSample.cpp b/Engine/Plugins/Media/WebMMedia/Source/WebMMedia/Private/Player/WebMMediaTextureSample.cpp
--- a/Engine/Plugins/Media/WebMMedia/Source/WebMMedia/Private/Player/WebMMediaTextureSample.cpp
+++ b/Engine/Plugins/Media/WebMMedia/Source/WebMMedia/Private/Player/WebMMediaTextureSample.cpp
@@ -93,7 +93,7 @@ void FWebMMediaTextureSample::ShutdownPoolable()
 	Time = FTimespan::Zero();
 }
 
-TRefCountPtr<FRHITexture2D> FWebMMediaTextureSample::GetTextureRef() const
+TRefCountPtr<FRHITexture> FWebMMediaTextureSample::GetTextureRef() const
 {
 	return Texture;
 }
diff --git a/Engine/Plugins/Media/WebMMedia/Source/WebMMedia/Private/Player/WebMVideoDecoder.cpp b/Engine/Plugins/Media/WebMMedia/Source/WebMMedia/Private/Player/WebMVideoDecoder.cpp
--- a/Engine/Plugins/Media/WebMMedia/Source/WebMMedia/Private/Player/WebMVideoDecoder.cpp
+++ b/Engine/Plugins/Media/WebMMedia/Source/WebMMedia/Private/Player/WebMVideoDecoder.cpp
@@ -225,7 +225,7 @@ void FWebMVideoDecoder::ConvertYUVToRGBAndSubmit(const FConvertParams& Params)
 	{
 		const auto CopyTextureMemory = [Image](
 			FRHICommandListImmediate& InCommandList,
-			FRHITexture2D* RHITexture,
+			FRHITexture* RHITexture,
 			int ImageIndex,
 			int CopyHeight)
 		{
diff --git a/Engine/Plugins/Media/WebMMedia/Source/WebMMedia/Public/WebMMediaTextureSample.h b/Engine/Plugins/Media/WebMMedia/Source/WebMMedia/Public/WebMMediaTextureSample.h
--- a/Engine/Plugins/Media/WebMMedia/Source/WebMMedia/Public/WebMMediaTextureSample.h
+++ b/Engine/Plugins/Media/WebMMedia/Source/WebMMedia/Public/WebMMediaTextureSample.h
@@ -39,10 +39,10 @@ public:
 	virtual void ShutdownPoolable() override;
 
 public:
-	TRefCountPtr<FRHITexture2D> GetTextureRef() const;
+	TRefCountPtr<FRHITexture> GetTextureRef() const;
 
 private:
-	TRefCountPtr<FRHITexture2D> Texture;
+	TRefCountPtr<FRHITexture> Texture;
 	FTimespan Time;
 	FTimespan Duration;
 	FIntPoint TotalSize;
diff --git a/Engine/Plugins/Media/WebMMedia/Source/WebMMedia/Public/WebMVideoDecoder.h b/Engine/Plugins/Media/WebMMedia/Source/WebMMedia/Public/WebMVideoDecoder.h
--- a/Engine/Plugins/Media/WebMMedia/Source/WebMMedia/Public/WebMVideoDecoder.h
+++ b/Engine/Plugins/Media/WebMMedia/Source/WebMMedia/Public/WebMVideoDecoder.h
@@ -37,9 +37,9 @@ private:
 
 	vpx_codec_ctx_t Context;
 	TUniquePtr<FWebMMediaTextureSamplePool> VideoSamplePool;
-	TRefCountPtr<FRHITexture2D> DecodedY;
-	TRefCountPtr<FRHITexture2D> DecodedU;
-	TRefCountPtr<FRHITexture2D> DecodedV;
+	TRefCountPtr<FRHITexture> DecodedY;
+	TRefCountPtr<FRHITexture> DecodedU;
+	TRefCountPtr<FRHITexture> DecodedV;
 	FGraphEventRef VideoDecodingTask;
 	IWebMSamplesSink& Samples;
 	bool bTexturesCreated;
diff --git a/Engine/Plugins/Media/WmfMedia/Source/WmfMedia/Private/Wmf/WmfMediaHardwareVideoDecodingRendering.cpp b/Engine/Plugins/Media/WmfMedia/Source/WmfMedia/Private/Wmf/WmfMediaHardwareVideoDecodingRendering.cpp
--- a/Engine/Plugins/Media/WmfMedia/Source/WmfMedia/Private/Wmf/WmfMediaHardwareVideoDecodingRendering.cpp
+++ b/Engine/Plugins/Media/WmfMedia/Source/WmfMedia/Private/Wmf/WmfMediaHardwareVideoDecodingRendering.cpp
@@ -36,9 +36,9 @@ DECLARE_GPU_STAT_NAMED(MediaTextureConversion, TEXT("MediaTextureConversion"));
 struct FRHICommandCopyResource final : public FRHICommand<FRHICommandCopyResource>
 {
 	TComPtr<ID3D11Texture2D> SampleTexture;
-	FTexture2DRHIRef SampleDestinationTexture;
+	FTextureRHIRef SampleDestinationTexture;
 
-	FRHICommandCopyResource(ID3D11Texture2D* InSampleTexture, FRHITexture2D* InSampleDestinationTexture)
+	FRHICommandCopyResource(ID3D11Texture2D* InSampleTexture, FRHITexture* InSampleDestinationTexture)
 		: SampleTexture(InSampleTexture)
 		, SampleDestinationTexture(InSampleDestinationTexture)
 	{
@@ -104,7 +104,7 @@ struct FRHICommandCopyResource final : public FRHICommand<FRHICommandCopyResourc
 	}
 };
 
-bool FWmfMediaHardwareVideoDecodingParameters::ConvertTextureFormat_RenderThread(FWmfMediaHardwareVideoDecodingTextureSample* InSample, FTexture2DRHIRef InDstTexture)
+bool FWmfMediaHardwareVideoDecodingParameters::ConvertTextureFormat_RenderThread(FWmfMediaHardwareVideoDecodingTextureSample* InSample, FTextureRHIRef InDstTexture)
 {
 	LLM_SCOPE(ELLMTag::MediaStreaming);
 
@@ -140,7 +140,7 @@ bool FWmfMediaHardwareVideoDecodingParameters::ConvertTextureFormat_RenderThread
 		GraphicsPSOInit.BoundShaderState.VertexDeclarationRHI = GetVertexDeclarationFVector4();
 
 		// Copy buffer.
-		FTexture2DRHIRef SampleDestinationTexture = InSample->GetOrCreateDestinationTexture();
+		FTextureRHIRef SampleDestinationTexture = InSample->GetOrCreateDestinationTexture();
 		if (InSample->IsBufferExternal() || !SampleTexture.IsValid())
 		{
 			const uint8* Data = (const uint8*)InSample->GetBuffer();
@@ -244,7 +244,7 @@ bool FWmfMediaHardwareVideoDecodingParameters::ConvertTextureFormat_RenderThread
 		else if (InSample->GetFormat() == EMediaTextureSampleFormat::YCoCg_DXT5_Alpha_BC4)
 		{
 			// Set up alpha texture.
-			FTexture2DRHIRef SampleDestinationAlphaTexture = InSample->GetOrCreateDestinationAlphaTexture();
+			FTextureRHIRef SampleDestinationAlphaTexture = InSample->GetOrCreateDestinationAlphaTexture();
 			{
 				const uint8* Data = ((const uint8*)InSample->GetBuffer()) + InSample->GetDim().X * InSample->GetDim().Y;
 
diff --git a/Engine/Plugins/Media/WmfMedia/Source/WmfMedia/Public/WmfMediaHardwareVideoDecodingRendering.h b/Engine/Plugins/Media/WmfMedia/Source/WmfMedia/Public/WmfMediaHardwareVideoDecodingRendering.h
--- a/Engine/Plugins/Media/WmfMedia/Source/WmfMedia/Public/WmfMediaHardwareVideoDecodingRendering.h
+++ b/Engine/Plugins/Media/WmfMedia/Source/WmfMedia/Public/WmfMediaHardwareVideoDecodingRendering.h
@@ -9,5 +9,5 @@ class FWmfMediaHardwareVideoDecodingTextureSample;
 class FWmfMediaHardwareVideoDecodingParameters
 {
 public:
-	WMFMEDIA_API static bool ConvertTextureFormat_RenderThread(FWmfMediaHardwareVideoDecodingTextureSample* InSample, FTexture2DRHIRef InDstTexture);
+	WMFMEDIA_API static bool ConvertTextureFormat_RenderThread(FWmfMediaHardwareVideoDecodingTextureSample* InSample, FTextureRHIRef InDstTexture);
 };
diff --git a/Engine/Plugins/Media/WmfMedia/Source/WmfMedia/Public/WmfMediaHardwareVideoDecodingTextureSample.h b/Engine/Plugins/Media/WmfMedia/Source/WmfMedia/Public/WmfMediaHardwareVideoDecodingTextureSample.h
--- a/Engine/Plugins/Media/WmfMedia/Source/WmfMedia/Public/WmfMediaHardwareVideoDecodingTextureSample.h
+++ b/Engine/Plugins/Media/WmfMedia/Source/WmfMedia/Public/WmfMediaHardwareVideoDecodingTextureSample.h
@@ -94,7 +94,7 @@ public:
 	/**
 	 * Texture sample convert using hardware video decoding.
 	 */
-	virtual bool Convert(FTexture2DRHIRef & InDstTexture, const FConversionHints & Hints) override
+	virtual bool Convert(FTextureRHIRef & InDstTexture, const FConversionHints & Hints) override
 	{
 		FWmfMediaHardwareVideoDecodingParameters::ConvertTextureFormat_RenderThread(this, InDstTexture);
 		return true;
diff --git a/Engine/Plugins/MovieScene/MovieRenderPipeline/Source/MovieRenderPipelineCore/Private/MoviePipelineSurfaceReader.cpp b/Engine/Plugins/MovieScene/MovieRenderPipeline/Source/MovieRenderPipelineCore/Private/MoviePipelineSurfaceReader.cpp
--- a/Engine/Plugins/MovieScene/MovieRenderPipeline/Source/MovieRenderPipelineCore/Private/MoviePipelineSurfaceReader.cpp
+++ b/Engine/Plugins/MovieScene/MovieRenderPipeline/Source/MovieRenderPipelineCore/Private/MoviePipelineSurfaceReader.cpp
@@ -74,7 +74,7 @@ void FMoviePipelineSurfaceReader::Reset()
 	bQueuedForCapture = false;
 }
 
-void FMoviePipelineSurfaceReader::ResolveSampleToReadbackTexture_RenderThread(const FTexture2DRHIRef& SourceSurfaceSample)
+void FMoviePipelineSurfaceReader::ResolveSampleToReadbackTexture_RenderThread(const FTextureRHIRef& SourceSurfaceSample)
 {
 	// We use GetModuleChecked here to avoid accidentally loading the module from the non-main thread.
 	static const FName RendererModuleName("Renderer");
@@ -323,7 +323,7 @@ void FMoviePipelineSurfaceQueue::Shutdown()
 	}
 }
 
-void FMoviePipelineSurfaceQueue::OnRenderTargetReady_RenderThread(const FTexture2DRHIRef InRenderTarget, TSharedRef<FImagePixelDataPayload, ESPMode::ThreadSafe> InFramePayload, TUniqueFunction<void(TUniquePtr<FImagePixelData>&&)>&& InFunctionCallback)
+void FMoviePipelineSurfaceQueue::OnRenderTargetReady_RenderThread(const FTextureRHIRef InRenderTarget, TSharedRef<FImagePixelDataPayload, ESPMode::ThreadSafe> InFramePayload, TUniqueFunction<void(TUniquePtr<FImagePixelData>&&)>&& InFunctionCallback)
 {
 	ensure(IsInRenderingThread());
 
diff --git a/Engine/Plugins/MovieScene/MovieRenderPipeline/Source/MovieRenderPipelineCore/Public/MoviePipelineSurfaceReader.h b/Engine/Plugins/MovieScene/MovieRenderPipeline/Source/MovieRenderPipelineCore/Public/MoviePipelineSurfaceReader.h
--- a/Engine/Plugins/MovieScene/MovieRenderPipeline/Source/MovieRenderPipelineCore/Public/MoviePipelineSurfaceReader.h
+++ b/Engine/Plugins/MovieScene/MovieRenderPipeline/Source/MovieRenderPipelineCore/Public/MoviePipelineSurfaceReader.h
@@ -42,7 +42,7 @@ struct MOVIERENDERPIPELINECORE_API FMoviePipelineSurfaceReader
 	/**
 	* Issues a command to the GPU to copy the given SourceSurfaceSample to our local ReadbackTexture for this surface.
 	*/
-	void ResolveSampleToReadbackTexture_RenderThread(const FTexture2DRHIRef& SourceSurfaceSample);
+	void ResolveSampleToReadbackTexture_RenderThread(const FTextureRHIRef& SourceSurfaceSample);
 
 	/**
 	* Maps the ReadbackTexture to the CPU (which should have been resolved to before this point) and copies the data over.
@@ -86,7 +86,7 @@ public:
 	FMoviePipelineSurfaceQueue& operator=(const FMoviePipelineSurfaceQueue&) = delete;
 public:
 
-	void OnRenderTargetReady_RenderThread(const FTexture2DRHIRef InRenderTarget, TSharedRef<FImagePixelDataPayload, ESPMode::ThreadSafe> InPayload, TUniqueFunction<void(TUniquePtr<FImagePixelData>&&)>&& InFunctionCallback);
+	void OnRenderTargetReady_RenderThread(const FTextureRHIRef InRenderTarget, TSharedRef<FImagePixelDataPayload, ESPMode::ThreadSafe> InPayload, TUniqueFunction<void(TUniquePtr<FImagePixelData>&&)>&& InFunctionCallback);
 
 	void BlockUntilAnyAvailable();
 	void Shutdown();
diff --git a/Engine/Plugins/Online/EOSShared/Source/EOSShared/Private/EOSSDKManager.cpp b/Engine/Plugins/Online/EOSShared/Source/EOSShared/Private/EOSSDKManager.cpp
--- a/Engine/Plugins/Online/EOSShared/Source/EOSShared/Private/EOSSDKManager.cpp
+++ b/Engine/Plugins/Online/EOSShared/Source/EOSShared/Private/EOSSDKManager.cpp
@@ -751,7 +751,7 @@ void FEOSSDKManager::SetupTicker()
 }
 
 #if WITH_ENGINE
-void FEOSSDKManager::OnBackBufferReady_RenderThread(SWindow& SlateWindow, const FTexture2DRHIRef& InBackBuffer)
+void FEOSSDKManager::OnBackBufferReady_RenderThread(SWindow& SlateWindow, const FTextureRHIRef& InBackBuffer)
 {
 	UE_CALL_ONCE([]() {	UE_LOG(LogEOSSDK, VeryVerbose, TEXT("[%hs] The method is not implemented for this platform."), __FUNCTION__) });
 }
diff --git a/Engine/Plugins/Online/EOSShared/Source/EOSShared/Private/EOSSDKManager.h b/Engine/Plugins/Online/EOSShared/Source/EOSShared/Private/EOSSDKManager.h
--- a/Engine/Plugins/Online/EOSShared/Source/EOSShared/Private/EOSSDKManager.h
+++ b/Engine/Plugins/Online/EOSShared/Source/EOSShared/Private/EOSSDKManager.h
@@ -83,7 +83,7 @@ protected:
 
 #if WITH_ENGINE
 	/** Provided to `OnBackBufferReadyToPresent` to get access to the render thread. */
-	virtual void OnBackBufferReady_RenderThread(SWindow& SlateWindow, const FTexture2DRHIRef& BackBuffer);
+	virtual void OnBackBufferReady_RenderThread(SWindow& SlateWindow, const FTextureRHIRef& BackBuffer);
 	/**
 	 * Check that the overlay is ready to be rendered.
 	 * This will also add the Back Buffer Ready To Present handler.
diff --git a/Engine/Plugins/Runtime/AR/AppleAR/AppleARKit/Source/AppleARKit/Private/AppleARKitTextures.cpp b/Engine/Plugins/Runtime/AR/AppleAR/AppleARKit/Source/AppleARKit/Private/AppleARKitTextures.cpp
--- a/Engine/Plugins/Runtime/AR/AppleAR/AppleARKit/Source/AppleARKit/Private/AppleARKitTextures.cpp
+++ b/Engine/Plugins/Runtime/AR/AppleAR/AppleARKit/Source/AppleARKit/Private/AppleARKitTextures.cpp
@@ -630,7 +630,7 @@ public:
 		SamplerStateRHI = RHICreateSamplerState(SamplerStateInitializer);
 	}
 
-	void CopyCubeFace(id<MTLTexture> MetalTexture, FTextureCubeRHIRef Cubemap, uint32 Rotation, int32 MetalCubeIndex, int32 OurCubeIndex)
+	void CopyCubeFace(id<MTLTexture> MetalTexture, FTextureRHIRef Cubemap, uint32 Rotation, int32 MetalCubeIndex, int32 OurCubeIndex)
 	{
 		// Rotate the image we need to get a view into the face as a new slice
 		id<MTLTexture> CubeFaceMetalTexture = [MetalTexture newTextureViewWithPixelFormat: MTLPixelFormatBGRA8Unorm textureType: MTLTextureType2D levels: NSMakeRange(0, 1) slices: NSMakeRange(MetalCubeIndex, 1)];
@@ -686,7 +686,7 @@ public:
 private:
 	FIntPoint Size;
 	
-	FTextureCubeRHIRef EnvCubemapTextureRHIRef;
+	FTextureRHIRef EnvCubemapTextureRHIRef;
 	
 	const UAppleARKitEnvironmentCaptureProbeTexture* Owner;
 	
@@ -1083,7 +1083,7 @@ private:
 	FIntPoint Size;
 
 	/** The texture that we actually render with which is populated via the GPU conversion process */
-	FTexture2DRHIRef DecodedTextureRef;
+	FTextureRHIRef DecodedTextureRef;
 	
 	FCriticalSection DecodedTextureLock;
 	
diff --git a/Engine/Plugins/Runtime/AR/Google/GoogleARCore/Source/GoogleARCoreBase/Private/GoogleARCoreTexture.cpp b/Engine/Plugins/Runtime/AR/Google/GoogleARCore/Source/GoogleARCoreBase/Private/GoogleARCoreTexture.cpp
--- a/Engine/Plugins/Runtime/AR/Google/GoogleARCore/Source/GoogleARCoreBase/Private/GoogleARCoreTexture.cpp
+++ b/Engine/Plugins/Runtime/AR/Google/GoogleARCore/Source/GoogleARCoreBase/Private/GoogleARCoreTexture.cpp
@@ -143,7 +143,7 @@ public:
 		
 		if (TextureRHI)
 		{
-			auto Texture2D = (FRHITexture2D*)(TextureRHI.GetReference());
+			auto Texture2D = (FRHITexture*)(TextureRHI.GetReference());
 			FUpdateTextureRegion2D Region(0, 0, 0, 0, ImageWidth, ImageHeight);
 			RHICmdList.UpdateTexture2D(Texture2D, 0, Region, sizeof(float) * ImageWidth, (const uint8*)DepthData);
 		}
diff --git a/Engine/Plugins/Runtime/HairStrands/Source/HairStrandsCore/Private/GroomTextureBuilder.cpp b/Engine/Plugins/Runtime/HairStrands/Source/HairStrandsCore/Private/GroomTextureBuilder.cpp
--- a/Engine/Plugins/Runtime/HairStrands/Source/HairStrandsCore/Private/GroomTextureBuilder.cpp
+++ b/Engine/Plugins/Runtime/HairStrands/Source/HairStrandsCore/Private/GroomTextureBuilder.cpp
@@ -242,7 +242,7 @@ static void InternalGenerateFollicleTexture_GPU(
 	{
 		if (OutTexture && OutTexture->GetResource() && OutTexture->GetResource()->GetTexture2DRHI())
 		{
-			FRHITexture2D* DstTexture = OutTexture->GetResource()->GetTexture2DRHI();
+			FRHITexture* DstTexture = OutTexture->GetResource()->GetTexture2DRHI();
 			RHICmdList.Transition(FRHITransitionInfo(DstTexture, ERHIAccess::SRVMask, ERHIAccess::CopyDest));
 
 			FRHICopyTextureInfo CopyInfo;
diff --git a/Engine/Plugins/Runtime/OpenXR/Source/OpenXRHMD/Internal/OpenXRHMD_RenderBridge.h b/Engine/Plugins/Runtime/OpenXR/Source/OpenXRHMD/Internal/OpenXRHMD_RenderBridge.h
--- a/Engine/Plugins/Runtime/OpenXR/Source/OpenXRHMD/Internal/OpenXRHMD_RenderBridge.h
+++ b/Engine/Plugins/Runtime/OpenXR/Source/OpenXRHMD/Internal/OpenXRHMD_RenderBridge.h
@@ -27,7 +27,7 @@ public:
 
 	virtual FXRSwapChainPtr CreateSwapchain(XrSession InSession, uint8 Format, uint8& OutActualFormat, uint32 SizeX, uint32 SizeY, uint32 ArraySize, uint32 NumMips, uint32 NumSamples, ETextureCreateFlags CreateFlags, const FClearValueBinding& ClearValueBinding, ETextureCreateFlags AuxiliaryCreateFlags = ETextureCreateFlags::None) = 0;
 
-	FXRSwapChainPtr CreateSwapchain(XrSession InSession, FRHITexture2D* Template, ETextureCreateFlags CreateFlags)
+	FXRSwapChainPtr CreateSwapchain(XrSession InSession, FRHITexture* Template, ETextureCreateFlags CreateFlags)
 	{
 		if (!Template)
 		{
