Fixed issue: Texture exporting fails after adding an additional Output node to the Graph

Fix/How: After debugging found out that thumbnail jobs are getting stuck during export. For now, we disabled the thumbnail job creation by adding a flag for export in invalidation details.

Removed finalise usage in case of export panel. Now we try to find a valid thumbnail in cache.
Removed rendering invalidation when enabling/disabling export of the output.

diff --git a/Engine/Plugins/Experimental/TextureGraph/Source/TextureGraph/Private/Expressions/Output/TG_Expression_Output.cpp b/Engine/Plugins/Experimental/TextureGraph/Source/TextureGraph/Private/Expressions/Output/TG_Expression_Output.cpp
--- a/Engine/Plugins/Experimental/TextureGraph/Source/TextureGraph/Private/Expressions/Output/TG_Expression_Output.cpp
+++ b/Engine/Plugins/Experimental/TextureGraph/Source/TextureGraph/Private/Expressions/Output/TG_Expression_Output.cpp
@@ -128,5 +128,8 @@ void UTG_Expression_Output::SetExport(bool bExport)
 {
 	OutputSettings.bExport = bExport;
 	UTG_Pin* Settings = GetParentNode()->GetPin(GET_MEMBER_NAME_CHECKED(UTG_Expression_Output, OutputSettings));
-	Settings->SetValue(OutputSettings.ToString());
+
+	// We are calling this instead of SetValue() to avoid rendering trigger.
+	// TODO: Better solution is that we should have NoInvalidate meta in UProperty or come up with a better solution.
+	Settings->EditSelfVar()->SetValueFromStr(OutputSettings.ToString());
 }
\ No newline at end of file
diff --git a/Engine/Plugins/Experimental/TextureGraph/Source/TextureGraph/Private/TG_HelperFunctions.cpp b/Engine/Plugins/Experimental/TextureGraph/Source/TextureGraph/Private/TG_HelperFunctions.cpp
--- a/Engine/Plugins/Experimental/TextureGraph/Source/TextureGraph/Private/TG_HelperFunctions.cpp
+++ b/Engine/Plugins/Experimental/TextureGraph/Source/TextureGraph/Private/TG_HelperFunctions.cpp
@@ -96,6 +96,7 @@ AsyncBool FTG_HelperFunctions::ExportAsync(UTextureGraph* InTextureGraph, FStrin
 	FInvalidationDetails Details;
 	Details.All();
 	Details.Mix = InTextureGraph;
+	Details.bExporting = true;
 	auto Batch = JobBatch::Create(Details);
 	/// Update the mix so that the rendering Cycle gets populated
 	MixUpdateCyclePtr Cycle = Batch->GetCycle();
diff --git a/Engine/Plugins/Experimental/TextureGraph/Source/TextureGraphEditor/Private/EdGraph/TG_EdGraph.cpp b/Engine/Plugins/Experimental/TextureGraph/Source/TextureGraphEditor/Private/EdGraph/TG_EdGraph.cpp
--- a/Engine/Plugins/Experimental/TextureGraph/Source/TextureGraphEditor/Private/EdGraph/TG_EdGraph.cpp
+++ b/Engine/Plugins/Experimental/TextureGraph/Source/TextureGraphEditor/Private/EdGraph/TG_EdGraph.cpp
@@ -134,31 +134,34 @@ void UTG_EdGraph::RefreshEditorDetails() const
 
 void UTG_EdGraph::OnNodeCreateThumbnail(UTG_Node* InNode, const FTG_EvaluationContext* InContext)
 {
-	UTG_EdGraphNode* EdGraphNode = GetViewModelNode(InNode->GetId());
-	if (EdGraphNode)
+	if(!InContext->Cycle->GetDetails().bExporting)
 	{
-		TArray<UTG_Pin*> OutputPins;
-		InNode->GetOutputPins(OutputPins);
-		// Loop over all Output pins
-		for(UTG_Pin* Pin : OutputPins)
+		UTG_EdGraphNode* EdGraphNode = GetViewModelNode(InNode->GetId());
+		if (EdGraphNode)
 		{
-			check (Pin);
-			if (Pin->IsArgTexture())
+			TArray<UTG_Pin*> OutputPins;
+			InNode->GetOutputPins(OutputPins);
+			// Loop over all Output pins
+			for(UTG_Pin* Pin : OutputPins)
 			{
-				FTG_Texture OutTexture;
-				if(Pin->GetValue(OutTexture))
+				check (Pin);
+				if (Pin->IsArgTexture())
 				{
-					UMixInterface* Mix = InContext->Cycle->GetMix();
-					auto TargetId = InContext->TargetId;
-
-					if (!OutTexture.RasterBlob)
+					FTG_Texture OutTexture;
+					if(Pin->GetValue(OutTexture))
 					{
-						OutTexture = FTG_Texture::GetBlack();
-					}
+						UMixInterface* Mix = InContext->Cycle->GetMix();
+						auto TargetId = InContext->TargetId;
 
-					TiledBlobPtr ThumbBlob = T_Thumbnail::Bind(Mix, Pin, OutTexture.RasterBlob, TargetId);
+						if (!OutTexture.RasterBlob)
+						{
+							OutTexture = FTG_Texture::GetBlack();
+						}
 
-					CacheThumbBlob(Pin->GetId(), ThumbBlob);
+						TiledBlobPtr ThumbBlob = T_Thumbnail::Bind(Mix, Pin, OutTexture.RasterBlob, TargetId);
+
+						CacheThumbBlob(Pin->GetId(), ThumbBlob);
+					}
 				}
 			}
 		}
diff --git a/Engine/Plugins/Experimental/TextureGraph/Source/TextureGraphEditor/Private/STG_OutputSelectionDlg.cpp b/Engine/Plugins/Experimental/TextureGraph/Source/TextureGraphEditor/Private/STG_OutputSelectionDlg.cpp
--- a/Engine/Plugins/Experimental/TextureGraph/Source/TextureGraphEditor/Private/STG_OutputSelectionDlg.cpp
+++ b/Engine/Plugins/Experimental/TextureGraph/Source/TextureGraphEditor/Private/STG_OutputSelectionDlg.cpp
@@ -90,24 +90,21 @@ void STG_OutputSelectionDlg::AddExportItems()
 				TSharedPtr<SWidget> ThumbnailWidget;
 				if (Variant.IsTexture())
 				{
-					TiledBlobPtr ThumbBlob = EdGraph->GetCachedThumbBlob(Id);
+					TSharedPtr<STG_NodeThumbnail> NodeThumbnail = SNew(STG_NodeThumbnail);
+					TiledBlobPtr CachedThumb = EdGraph->GetCachedThumbBlob(Id);
 
-					if (!ThumbBlob)
+					if(!CachedThumb)
 					{
-						ThumbBlob = TextureHelper::GetBlack();
+						CachedThumb = TextureHelper::GetBlack();						
 					}
-
-					TSharedPtr<STG_NodeThumbnail> NodeThumbnail = SNew(STG_NodeThumbnail);
-					ThumbBlob->OnFinalise()
-					.then([ThumbBlob, NodeThumbnail]
+					
+					if (CachedThumb->IsFinalised())
+					{
+						if (NodeThumbnail.IsValid())
 						{
-							// NOTE: If later, "this" were to be captured here, we should check DoesSharedInstanceExist()
-							// as there is a chance this might be invoked when the slate widgets have already been destroyed
-							if (NodeThumbnail.IsValid())
-							{
-								NodeThumbnail->UpdateBlob(ThumbBlob);
-							}
-						});
+							NodeThumbnail->UpdateBlob(CachedThumb);
+						}
+					}
 
 					ThumbnailWidget = NodeThumbnail;
 				}
diff --git a/Engine/Plugins/Experimental/TextureGraph/Source/TextureGraphEngine/Model/ModelObject.h b/Engine/Plugins/Experimental/TextureGraph/Source/TextureGraphEngine/Model/ModelObject.h
--- a/Engine/Plugins/Experimental/TextureGraph/Source/TextureGraphEngine/Model/ModelObject.h
+++ b/Engine/Plugins/Experimental/TextureGraph/Source/TextureGraphEngine/Model/ModelObject.h
@@ -27,7 +27,8 @@ struct TEXTUREGRAPHENGINE_API FInvalidationDetails
 	bool								bReload = false;							/// Whether to reload the sources or not
 	bool								bSelective = true;							/// Selective invalidation [tile based]
 	bool								bTweaking = false;							/// Whether being tweaked or not
-
+	bool								bExporting = false;							/// Should only be true in case of exporting.
+	
 	/// The list of things that have been invalidated
 	bool								bForceInvalidateParent = false;				/// Force invalidates the parent
 	bool								bRender = false;							/// Re-render
