Fixed issue: Preview doesn't render in the Viewport after deleting one of the Output nodes

diff --git a/Engine/Plugins/Experimental/TextureGraph/Source/TextureGraphEditor/Private/TG_Editor.cpp b/Engine/Plugins/Experimental/TextureGraph/Source/TextureGraphEditor/Private/TG_Editor.cpp
--- a/Engine/Plugins/Experimental/TextureGraph/Source/TextureGraphEditor/Private/TG_Editor.cpp
+++ b/Engine/Plugins/Experimental/TextureGraph/Source/TextureGraphEditor/Private/TG_Editor.cpp
@@ -1812,8 +1812,34 @@ void FTG_Editor::OnNodeRemoved(UTG_Node* InNode, FName Title)
 			}
 		}
 
+		FViewportSettings& ViewportSettings = EditedTextureGraph->GetSettings()->GetViewportSettings();
+		
 		// Update Viewport material mapping for deleted output node
-		EditedTextureGraph->GetSettings()->GetViewportSettings().RemoveMaterialMappingForTarget(Title);
+		ViewportSettings.RemoveMaterialMappingForTarget(Title);
+
+		if(ViewportSettings.MaterialMappingInfos.Num() > 0)
+		{
+			// We should have atleast one target assigned.
+			const int AssignedTargets = ViewportSettings.NumAssignedTargets();
+
+			// if We dont have a valid target then we will try to assign 
+			if(AssignedTargets <= 0)
+			{
+				// Check if we have output node available
+				const FName FirstOutputName;
+				EditedTextureGraph->Graph()->ForEachParams([&](const UTG_Pin* Pin, uint32 index)
+					{
+						if(Pin->IsOutput() && (Pin->IsArgTexture() || Pin->IsArgColor()))
+						{
+							if (FirstOutputName.IsNone())
+							{
+								ViewportSettings.SetDefaultTarget(Pin->GetAliasName());
+								GetEditorViewport()->UpdateRenderMode();
+							}
+						}
+					});
+			}
+		}
 	}
 }
 
diff --git a/Engine/Plugins/Experimental/TextureGraph/Source/TextureGraphEngine/Model/Mix/ViewportSettings.cpp b/Engine/Plugins/Experimental/TextureGraph/Source/TextureGraphEngine/Model/Mix/ViewportSettings.cpp
--- a/Engine/Plugins/Experimental/TextureGraph/Source/TextureGraphEngine/Model/Mix/ViewportSettings.cpp
+++ b/Engine/Plugins/Experimental/TextureGraph/Source/TextureGraphEngine/Model/Mix/ViewportSettings.cpp
@@ -133,3 +133,17 @@ void FViewportSettings::OnTargetRename(const FName OldName,const FName NewName)
 		OnMaterialMappingChangedEvent.Broadcast();
 	}
 }
+
+int FViewportSettings::NumAssignedTargets()
+{
+	int ValidTargets = 0;
+	for(const FMaterialMappingInfo& MappingInfo : MaterialMappingInfos)
+	{
+		if(MappingInfo.HasTarget())
+		{
+			ValidTargets++;
+		}
+	}
+
+	return ValidTargets;
+}
diff --git a/Engine/Plugins/Experimental/TextureGraph/Source/TextureGraphEngine/Model/Mix/ViewportSettings.h b/Engine/Plugins/Experimental/TextureGraph/Source/TextureGraphEngine/Model/Mix/ViewportSettings.h
--- a/Engine/Plugins/Experimental/TextureGraph/Source/TextureGraphEngine/Model/Mix/ViewportSettings.h
+++ b/Engine/Plugins/Experimental/TextureGraph/Source/TextureGraphEngine/Model/Mix/ViewportSettings.h
@@ -43,6 +43,8 @@ struct TEXTUREGRAPHENGINE_API FViewportSettings
 	bool RemoveMaterialMappingForTarget(FName OutputNode);
 	void OnMaterialUpdate();
 	void OnTargetRename(const FName OldName, const FName NewName);
+
+	int  NumAssignedTargets();
 	
 	DECLARE_MULTICAST_DELEGATE(FViewportSettingsUpdateEvent)
 	FViewportSettingsUpdateEvent OnViewportMaterialChangeEvent;
