Fixing some minor issues with FxMaterial_DrawMaterial.h

diff --git a/Engine/Plugins/Experimental/TextureGraph/Source/TextureGraphEngine/FxMat/FxMaterial_DrawMaterial.h b/Engine/Plugins/Experimental/TextureGraph/Source/TextureGraphEngine/FxMat/FxMaterial_DrawMaterial.h
--- a/Engine/Plugins/Experimental/TextureGraph/Source/TextureGraphEngine/FxMat/FxMaterial_DrawMaterial.h
+++ b/Engine/Plugins/Experimental/TextureGraph/Source/TextureGraphEngine/FxMat/FxMaterial_DrawMaterial.h
@@ -83,9 +83,8 @@ public:
 		return IsFeatureLevelSupported(Platform, ERHIFeatureLevel::SM5);
 	}
 
-	void SetParameters(FRHIBatchedShaderParameters& BatchedParameters, const FSceneView& View, const FMaterialRenderProxy* MaterialProxy)
+	void SetParameters(FRHIBatchedShaderParameters& BatchedParameters, const FSceneView& View, const FMaterialRenderProxy* MaterialProxy, const FMaterial& Material)
 	{
-		const FMaterial& Material = MaterialProxy->GetMaterialWithFallback(View.GetFeatureLevel(), MaterialProxy);
 		FMaterialShader::SetViewParameters(BatchedParameters, View, View.ViewUniformBuffer);
 		FMaterialShader::SetParameters(BatchedParameters, MaterialProxy, Material, View);
 	}
@@ -172,10 +171,10 @@ public:
 
 		// Pixel shader combined with Material
 		const FMaterialRenderProxy* MaterialProxy = Material->GetRenderProxy();
-		MaterialProxy->UpdateUniformExpressionCacheIfNeeded(GMaxRHIFeatureLevel);
+		MaterialProxy->UpdateUniformExpressionCacheIfNeeded(SceneView->GetFeatureLevel());
 		FMaterialRenderProxy::UpdateDeferredCachedUniformExpressions();
 
-		const FMaterial& RenderMaterial = MaterialProxy->GetMaterialWithFallback(GMaxRHIFeatureLevel, MaterialProxy);
+		const FMaterial& RenderMaterial = MaterialProxy->GetMaterialWithFallback(SceneView->GetFeatureLevel(), MaterialProxy);
 		const FMaterialShaderMap* const MaterialShaderMap = RenderMaterial.GetRenderingThreadShaderMap();
 
 		TShaderRef<VSH_Type> VertexShader = MaterialShaderMap->GetShader<VSH_Type>();
@@ -196,11 +195,9 @@ public:
 		RHI.ApplyCachedRenderTargets(PSO);
 
 		SetGraphicsPipelineState(RHI, PSO, 0);
-		SetShaderParametersLegacyVS(RHI, VertexShader, *SceneView);
-		SetShaderParametersLegacyPS(RHI, PixelShader, *SceneView, MaterialProxy);
-		
-		SetShaderParameters(RHI, VertexShader, RHIVertexShader, VSHParams);
-		SetShaderParameters(RHI, PixelShader, RHIPixelShader, FSHParams);
+
+		SetShaderParametersMixedVS(RHI, VertexShader, VSHParams, *SceneView);
+		SetShaderParametersMixedPS(RHI, PixelShader, FSHParams, *SceneView, MaterialProxy, RenderMaterial);
 
 		RHI.SetStreamSource(0, GQuadBuffer.VertexBufferRHI, 0);
 		RHI.DrawPrimitive(0, 2, 1);
