Removing unnecessary register assignment that ends up breaking Bindless conversion.

diff --git a/Engine/Plugins/Experimental/TextureGraph/Shaders/Utils/Histogram_comp.usf b/Engine/Plugins/Experimental/TextureGraph/Shaders/Utils/Histogram_comp.usf
--- a/Engine/Plugins/Experimental/TextureGraph/Shaders/Utils/Histogram_comp.usf
+++ b/Engine/Plugins/Experimental/TextureGraph/Shaders/Utils/Histogram_comp.usf
@@ -8,7 +8,7 @@
 Texture2D SourceTiles;
 
 RWTexture2D<float4> Result;
-RWStructuredBuffer<uint4> Histogram : register(u0);
+RWStructuredBuffer<uint4> Histogram;
 
 #ifndef THREADGROUPSIZE_X
 #define THREADGROUPSIZE_X 16
