TextureGraph: Analytics schematization updates

diff --git a/Engine/Plugins/Experimental/TextureGraph/Source/TextureGraph/Private/TG_HelperFunctions.cpp b/Engine/Plugins/Experimental/TextureGraph/Source/TextureGraph/Private/TG_HelperFunctions.cpp
--- a/Engine/Plugins/Experimental/TextureGraph/Source/TextureGraph/Private/TG_HelperFunctions.cpp
+++ b/Engine/Plugins/Experimental/TextureGraph/Source/TextureGraph/Private/TG_HelperFunctions.cpp
@@ -1,6 +1,8 @@
 // Copyright Epic Games, Inc. All Rights Reserved.
 
 #include "TG_HelperFunctions.h"
+
+#include "EngineAnalytics.h"
 #include "Job/Scheduler.h"
 #include "TG_Texture.h"
 #include "TG_Graph.h"
@@ -102,8 +104,6 @@ AsyncBool FTG_HelperFunctions::ExportAsync(UTextureGraph* InTextureGraph, FStrin
 	MixUpdateCyclePtr Cycle = Batch->GetCycle();
 	InTextureGraph->Update(Cycle);
 	
-	//auto SettingsSet = InTextureGraph->GetOutputSettingsSet();
-
 	InTextureGraph->Graph()->ForEachNodes([=,&TargetExportSettings](const UTG_Node* Node, uint32 Index)
 	{
 		UTG_Expression_Output* TargetExpression = Cast<UTG_Expression_Output>(Node->GetExpression());
@@ -162,9 +162,22 @@ AsyncBool FTG_HelperFunctions::ExportAsync(UTextureGraph* InTextureGraph, FStrin
 	});
 
 	return RenderAsync(InTextureGraph, Batch)
-	.then([InTextureGraph, TargetExportSettings, ExportPath](auto result) 
+	.then([InTextureGraph, &TargetExportSettings, ExportPath](auto result) 
 	{
 		return TextureExporter::ExportAsUAsset(InTextureGraph, TargetExportSettings, ExportPath);
+	})
+	.then([InTextureGraph, &TargetExportSettings]()
+	{
+		// Add analytics tag
+		if (FEngineAnalytics::IsAvailable())
+		{
+			TArray<FAnalyticsEventAttribute> Attributes;
+			Attributes.Add(FAnalyticsEventAttribute(TEXT("NumExports"),  TargetExportSettings.MapsExported));
+					
+			// Send Analytics event 
+			FEngineAnalytics::GetProvider().RecordEvent(TEXT("Editor.Usage.TextureGraph.Export"), Attributes);
+		}
+		return cti::make_ready_continuable(true);
 	});
 }
 
diff --git a/Engine/Plugins/Experimental/TextureGraph/Source/TextureGraphEditor/Private/TG_Editor.cpp b/Engine/Plugins/Experimental/TextureGraph/Source/TextureGraphEditor/Private/TG_Editor.cpp
--- a/Engine/Plugins/Experimental/TextureGraph/Source/TextureGraphEditor/Private/TG_Editor.cpp
+++ b/Engine/Plugins/Experimental/TextureGraph/Source/TextureGraphEditor/Private/TG_Editor.cpp
@@ -356,7 +356,7 @@ void FTG_Editor::InitEditor(const EToolkitMode::Type Mode, const TSharedPtr< cla
 	if (FEngineAnalytics::IsAvailable())
 	{
 		SessionStartTime = FDateTime::Now();
-		FEngineAnalytics::GetProvider().RecordEvent(TEXT("Editor.Usage.TextureGraph"), TEXT("SessionStarted"), OriginalTextureGraph.GetName());
+		FEngineAnalytics::GetProvider().RecordEvent(TEXT("Editor.Usage.TextureGraph.SessionStarted"));
 	}
 }
 
@@ -1567,8 +1567,7 @@ void FTG_Editor::OnClose()
 	if (FEngineAnalytics::IsAvailable())
 	{
 		TArray<FAnalyticsEventAttribute> Attributes;
-		Attributes.Add(FAnalyticsEventAttribute(TEXT("GraphName"), OriginalTextureGraph.GetName()));
-		Attributes.Add(FAnalyticsEventAttribute(TEXT("TimeActive"),  (FDateTime::Now() - SessionStartTime).ToString()));
+		Attributes.Add(FAnalyticsEventAttribute(TEXT("TimeActive.Seconds"),  (FDateTime::Now() - SessionStartTime).GetTotalSeconds()));
 				
 		// Send Analytics event 
 		FEngineAnalytics::GetProvider().RecordEvent(TEXT("Editor.Usage.TextureGraph.SessionEnded"), Attributes);
