Some Material require binding the FPrimitive Constant Buffer
we were missing this in the Material draw pass when setting up the shader's parameters

This has also been pushed in upcoming release branch with commit   31810539
The ROBOMERGE was having a conflict, so i skipped the robomerge CL and instead do the change manually through this cl

diff --git a/Engine/Plugins/Experimental/TextureGraph/Source/TextureGraphEngine/FxMat/FxMaterial_DrawMaterial.h b/Engine/Plugins/Experimental/TextureGraph/Source/TextureGraphEngine/FxMat/FxMaterial_DrawMaterial.h
--- a/Engine/Plugins/Experimental/TextureGraph/Source/TextureGraphEngine/FxMat/FxMaterial_DrawMaterial.h
+++ b/Engine/Plugins/Experimental/TextureGraph/Source/TextureGraphEngine/FxMat/FxMaterial_DrawMaterial.h
@@ -85,6 +85,8 @@ public:
 
 	void SetParameters(FRHIBatchedShaderParameters& BatchedParameters, const FSceneView& View, const FMaterialRenderProxy* MaterialProxy, const FMaterial& Material)
 	{
+		auto& PrimitivePS = GetUniformBufferParameter<FPrimitiveUniformShaderParameters>();
+		SetUniformBufferParameter(BatchedParameters, PrimitivePS, GIdentityPrimitiveUniformBuffer);
 		FMaterialShader::SetViewParameters(BatchedParameters, View, View.ViewUniformBuffer);
 		FMaterialShader::SetParameters(BatchedParameters, MaterialProxy, Material, View);
 	}
