Fixing the Deferred Attribute material value displyed when reseting the material property
#tests 65dd46b58a0d4d5ee1dd4da7

diff --git a/Engine/Plugins/Experimental/TextureGraph/Source/TextureGraph/Private/Expressions/Input/TG_Expression_Material.cpp b/Engine/Plugins/Experimental/TextureGraph/Source/TextureGraph/Private/Expressions/Input/TG_Expression_Material.cpp
--- a/Engine/Plugins/Experimental/TextureGraph/Source/TextureGraph/Private/Expressions/Input/TG_Expression_Material.cpp
+++ b/Engine/Plugins/Experimental/TextureGraph/Source/TextureGraph/Private/Expressions/Input/TG_Expression_Material.cpp
@@ -10,6 +10,7 @@ void UTG_Expression_Material::PostEditChangeProperty(FPropertyChangedEvent& Prop
 	{
 		UE_LOG(LogTextureGraph, Log, TEXT("Material Expression PostEditChangeProperty."));
 		SetMaterialInternal(Material);
+		FeedbackPinValue(GET_MEMBER_NAME_CHECKED(UTG_Expression_Material, RenderedAttribute), RenderedAttribute);
 	}
 	// Second catch if AttributeName changes
 	if (PropertyChangedEvent.GetPropertyName() == GET_MEMBER_NAME_CHECKED(UTG_Expression_Material, RenderedAttribute))
diff --git a/Engine/Plugins/Experimental/TextureGraph/Source/TextureGraph/Public/Expressions/TG_Expression.h b/Engine/Plugins/Experimental/TextureGraph/Source/TextureGraph/Public/Expressions/TG_Expression.h
--- a/Engine/Plugins/Experimental/TextureGraph/Source/TextureGraph/Public/Expressions/TG_Expression.h
+++ b/Engine/Plugins/Experimental/TextureGraph/Source/TextureGraph/Public/Expressions/TG_Expression.h
@@ -128,6 +128,22 @@ protected:
 	virtual void NotifyCommonVariantTypeChanged(FTG_Variant::EType NewType) const {}
 	FTG_Variant::EType EvalExpressionCommonVariantType() const { return GetParentNode()->EvalExpressionCommonVariantType(); }
 
+	// In some cases, evaluation or change in the Expression needs to be feedback to the matching pin's value
+	// This is not needed for the standard flow of evaluation but is sometime required for coupled member.
+	template<typename T> bool FeedbackPinValue(const FName& InPinName, const T& InValue)
+	{
+		UTG_Node* ParentNode = GetParentNode();
+		if (ParentNode)
+		{
+			UTG_Pin* Pin = ParentNode->GetPin(InPinName);
+			if (Pin)
+			{
+				return Pin->SetValue(InValue);
+			}
+		}
+		return false;
+	}
+
 #if WITH_EDITOR
 	virtual void PostEditChangeProperty(FPropertyChangedEvent& PropertyChangedEvent) override;
 	virtual void PostEditUndo() override;
