TextureGraphEditor: Add node preview using the SImageViewport widget

diff --git a/Engine/Plugins/Experimental/TextureGraph/Source/TextureGraphEditor/Private/STG_NodePreview.cpp b/Engine/Plugins/Experimental/TextureGraph/Source/TextureGraphEditor/Private/STG_NodePreview.cpp
new file mode 100644
--- /dev/null
+++ b/Engine/Plugins/Experimental/TextureGraph/Source/TextureGraphEditor/Private/STG_NodePreview.cpp
@@ -0,0 +1,461 @@
+﻿// Copyright Epic Games, Inc. All Rights Reserved.
+
+#if TEXTUREGRAPHEDITOR_ENABLE_NEW_NODE_PREVIEW
+
+#include "STG_NodePreview.h"
+
+#include "CanvasItem.h"
+#include "CanvasTypes.h"
+#include "SImageViewport.h"
+#include "Texture2DPreview.h"
+#include "TextureResource.h"
+#include "TG_Node.h"
+#include "2D/Tex.h"
+#include "Device/FX/DeviceBuffer_FX.h"
+#include "EdGraph/TG_EdGraphNode.h"
+#include "Framework/MultiBox/MultiBoxBuilder.h"
+#include "Widgets/Text/STextBlock.h"
+#include "Widgets/Input/SCheckBox.h"
+
+#define LOCTEXT_NAMESPACE "STG_NodePreview"
+
+UE::ImageWidgets::IImageViewer::FImageInfo FNodeViewer::GetCurrentImageInfo() const
+{
+	if (NodeTexture)
+	{
+		return {{}, {NodeTexture->SizeX, NodeTexture->SizeY}, NodeTexture->GetNumMips(), true};
+	}
+
+	return {{}, FIntPoint::ZeroValue, 0, false};
+}
+
+void FNodeViewer::DrawCurrentImage(FViewport*, FCanvas* Canvas, const FDrawProperties& Properties)
+{
+	if (NodeTexture)
+	{
+		if (FTextureResource* TextureResource = NodeTexture->GetResource(); TextureResource != nullptr)
+		{
+			TextureResource->bGreyScaleFormat = IsSingleChannel();
+			TextureResource->bSRGB = IsSRGB();
+
+			DrawTexture(TextureResource, Canvas, Properties.Placement, Properties.Mip);
+		}
+	}
+}
+
+TOptional<TVariant<FColor, FLinearColor>> FNodeViewer::GetCurrentImagePixelColor(FIntPoint PixelCoords, int32 MipIndex) const
+{
+	const uint64 PixelIndex = PixelCoords.Y * NodeDescriptor.Width + PixelCoords.X;
+	if (PixelIndex < NodePixels.Num())
+	{
+		const FLinearColor LinearColor = NodePixels[PixelIndex];
+
+		if (NodeDescriptor.Format == BufferFormat::Byte)
+		{
+			return TVariant<FColor, FLinearColor>(TInPlaceType<FColor>(), LinearColor.ToFColor(IsSRGB()));
+		}
+		return TVariant<FColor, FLinearColor>(TInPlaceType<FLinearColor>(), LinearColor);
+	}
+	return {};
+}
+
+FText FNodeViewer::GetFormatLabelText() const
+{
+	if (NodeTexture != nullptr)
+	{
+		return FText::Format(FTextFormat::FromString("{0}_{1} {2}"),
+		                     FText::FromString(TextureHelper::GetChannelsTextFromItemsPerPoint(NodeDescriptor.ItemsPerPoint)),
+		                     FText::FromString(BufferDescriptor::FormatToString(NodeDescriptor.Format)),
+		                     FText::FromString(NodeDescriptor.bIsSRGB ? "(sRGB)" : "(Linear)"));
+	}
+
+	return FText();
+}
+
+bool FNodeViewer::IsSingleChannel() const
+{
+	return NodeDescriptor.ItemsPerPoint == 1;
+}
+
+void FNodeViewer::SetTexture(const BlobPtr& InBlob)
+{
+	NodeTexture = nullptr;
+
+	NodePixels.SetNumUnsafeInternal(0);
+	NodeDescriptor = {};
+
+	if (InBlob)
+	{
+		if (const DeviceBufferPtr Buffer = InBlob->GetBufferRef().GetPtr())
+		{
+			NodeTexture = GetTextureFromBuffer(Buffer);
+
+			Buffer->Raw()
+			      .then([this](const RawBufferPtr& RawBuffer)
+			      {
+				      RawBuffer->GetAsLinearColor(NodePixels);
+				      NodeDescriptor = RawBuffer->GetDescriptor();
+			      });
+		}
+	}
+}
+
+void FNodeViewer::SetRGBA(bool bR, bool bG, bool bB, bool bA)
+{
+	bRGBA[0] = bR;
+	bRGBA[1] = bG;
+	bRGBA[2] = bB;
+	bRGBA[3] = bA;
+}
+
+void FNodeViewer::DrawTexture(const FTextureResource* TextureResource, FCanvas* Canvas, const FDrawProperties::FPlacement& Placement,
+                              const FDrawProperties::FMip& Mip) const
+{
+	check(GMaxRHIFeatureLevel >= ERHIFeatureLevel::SM5);
+
+	const bool bIsNormalMap = NodeTexture->IsNormalMap();
+	const bool bIsVirtualTexture = NodeTexture->IsCurrentlyVirtualTextured();
+
+	const TRefCountPtr BatchedElementParameters = new FBatchedElementTexture2DPreviewParameters(
+		Mip.MipLevel, 0, 0, bIsNormalMap, false, false, bIsVirtualTexture, false, Placement.ZoomFactor >= 2.0/*Mip.bUsePointSampling*/);
+
+	FCanvasTileItem Tile(Placement.Offset, TextureResource, Placement.Size, FLinearColor::White);
+	Tile.BatchedElementParameters = BatchedElementParameters.GetReference();
+	Tile.BlendMode = GetBlendMode();
+	Canvas->DrawItem(Tile);
+}
+
+ESimpleElementBlendMode FNodeViewer::GetBlendMode() const
+{
+	if (NodeTexture)
+	{
+		const TEnumAsByte<TextureCompressionSettings>& CompressionSettings = NodeTexture->CompressionSettings;
+		if (CompressionSettings == TC_Grayscale || CompressionSettings == TC_Alpha)
+		{
+			return SE_BLEND_Opaque;
+		}
+	}
+
+	int32 Result = SE_BLEND_RGBA_MASK_START;
+
+	if (IsSingleChannel())
+	{
+		Result += bRGBA[0] * 0b00111;
+	}
+	else
+	{
+		Result += bRGBA[0] * 0b00001;
+		Result += bRGBA[1] * 0b00010;
+		Result += bRGBA[2] * 0b00100;
+		Result += bRGBA[3] * 0b01000;
+	}
+
+	return static_cast<ESimpleElementBlendMode>(Result);
+}
+
+UTextureRenderTarget2D* FNodeViewer::GetTextureFromBuffer(const DeviceBufferPtr& Buffer) const
+{
+	UTextureRenderTarget2D* OutTexture = nullptr;
+	if (const std::shared_ptr<DeviceBuffer_FX> FXBuffer = std::static_pointer_cast<DeviceBuffer_FX>(Buffer))
+	{
+		if (!FXBuffer->IsNull())
+		{
+			const TexPtr Tex = FXBuffer->GetTexture();
+			check(Tex);
+			UTexture* BufferTexture = Tex->GetTexture();
+			if (UTextureRenderTarget2D* RenderTarget = Cast<UTextureRenderTarget2D>(BufferTexture))
+			{
+				OutTexture = RenderTarget;
+			}
+			else
+			{
+				UE_LOG(LogTemp, Warning, TEXT("Texture is not a UTextureRenderTarget2D."));
+			}
+		}
+		else
+		{
+			UE_LOG(LogTemp, Warning, TEXT("Buffer is not an FXBuffer."));
+		}
+	}
+	else
+	{
+		UE_LOG(LogTemp, Warning, TEXT("BlobTexture failed to find the buffer."));
+	}
+
+	return OutTexture;
+}
+
+bool FNodeViewer::IsSRGB() const
+{
+	return NodeDescriptor.bIsSRGB;
+}
+
+class FPreviewViewerCommands : public TCommands<FPreviewViewerCommands>
+{
+public:
+	FPreviewViewerCommands()
+		: TCommands(TEXT("NodePreview"), LOCTEXT("ContextDescription", "Node Preview"), NAME_None, FAppStyle::Get().GetStyleSetName())
+	{
+	}
+
+	virtual void RegisterCommands() override
+	{
+		UI_COMMAND(ToggleLock, "Toggle Node Preview Lock", "Toggles the node preview lock.", EUserInterfaceActionType::ToggleButton, FInputChord(EKeys::L));
+	}
+
+	TSharedPtr<FUICommandInfo> ToggleLock;
+};
+
+STG_NodePreviewWidget::~STG_NodePreviewWidget()
+{
+	FPreviewViewerCommands::Unregister();
+}
+
+void STG_NodePreviewWidget::Construct(const FArguments& InArgs)
+{
+	NodeViewer = MakeShared<FNodeViewer>();
+
+	OnNodeBlobChanged = InArgs._OnNodeBlobChanged;
+
+	FPreviewViewerCommands::Register();
+	const FPreviewViewerCommands Commands = FPreviewViewerCommands::Get();
+	CommandList = MakeShared<FUICommandList>();
+	CommandList->MapAction(Commands.ToggleLock,
+	                       FExecuteAction::CreateSP(this, &STG_NodePreviewWidget::ToggleLock),
+	                       FCanExecuteAction::CreateLambda([this] { return NodeViewer->GetCurrentImageInfo().bIsValid; }),
+	                       FIsActionChecked::CreateLambda([this] { return LockedNode != nullptr; }));
+
+	const TSharedPtr<FExtender> ToolbarExtender = MakeShared<FExtender>();
+	ToolbarExtender->AddToolBarExtension("ToolbarLeft", EExtensionHook::Before, CommandList,
+	                                     FToolBarExtensionDelegate::CreateSP(this, &STG_NodePreviewWidget::AddLockButton));
+
+	ToolbarExtender->AddToolBarExtension("ToolbarRight", EExtensionHook::After, CommandList,
+	                                     FToolBarExtensionDelegate::CreateSP(this, &STG_NodePreviewWidget::AddRGBAButtons));
+
+	const TSharedPtr<UE::ImageWidgets::SImageViewport::FStatusBarExtender> StatusBarExtender = MakeShared<
+		UE::ImageWidgets::SImageViewport::FStatusBarExtender>();
+	StatusBarExtender->AddExtension("StatusBarLeft", EExtensionHook::After, CommandList,
+	                                UE::ImageWidgets::SImageViewport::FStatusBarExtender::FDelegate::CreateSP(this, &STG_NodePreviewWidget::AddFormatLabel));
+
+	ChildSlot
+		.VAlign(VAlign_Fill)
+		.HAlign(HAlign_Fill)
+		[
+			SAssignNew(Viewport, UE::ImageWidgets::SImageViewport, NodeViewer.ToSharedRef())
+				.ToolbarExtender(ToolbarExtender)
+				.StatusBarExtender(StatusBarExtender)
+				.DrawSettings(UE::ImageWidgets::SImageViewport::FDrawSettings{
+					.ClearColor = FVector3f(0.1f),
+					.bBorderEnabled = false,
+					.bBackgroundColorEnabled = true,
+					.BackgroundColor = FLinearColor::Black,
+					.bBackgroundCheckerEnabled = true
+				})
+				.ControllerSettings(UE::ImageWidgets::SImageViewport::FControllerSettings{
+					.DefaultZoomMode = UE::ImageWidgets::SImageViewport::FControllerSettings::EDefaultZoomMode::Fill
+				})
+		];
+}
+
+void STG_NodePreviewWidget::SelectionChanged(UTG_EdGraphNode* Node)
+{
+	const bool bUpdatePreview = LockedNode == nullptr && SelectedNode != Node;
+
+	SelectedNode = Node;
+
+	if (bUpdatePreview)
+	{
+		Update();
+	}
+}
+
+void STG_NodePreviewWidget::NodeDeleted(const UTG_EdGraphNode* Node)
+{
+	const bool bPreviewNodeDeleted = LockedNode == Node || (LockedNode == nullptr && SelectedNode == Node);
+
+	if (LockedNode == Node)
+	{
+		LockedNode = nullptr;
+	}
+
+	if (SelectedNode == Node)
+	{
+		SelectedNode = nullptr;
+	}
+
+	if (bPreviewNodeDeleted)
+	{
+		Update();
+	}
+}
+
+void STG_NodePreviewWidget::Update() const
+{
+	const UTG_EdGraphNode* PreviewNode = LockedNode ? LockedNode : SelectedNode;
+
+	// Get blob if preview node is valid or assign a nullptr if not.
+	const BlobPtr Blob = [PreviewNode]() -> BlobPtr
+	{
+		if (PreviewNode)
+		{
+			TArray<FTG_Texture> OutTextures;
+			const UTG_Node* TGNode = PreviewNode->GetNode();
+			check(TGNode);
+			TGNode->GetAllOutputValues(OutTextures);
+
+			if (!OutTextures.IsEmpty())
+			{
+				return OutTextures[0].RasterBlob;
+			}
+		}
+
+		return nullptr;
+	}();
+
+	// Lambda for actually updating the preview.
+	auto UpdatePreview = [this, Blob]
+	{
+		NodeViewer->SetTexture(Blob);
+		Viewport->ResetZoom(NodeViewer->GetCurrentImageInfo().Size);
+	};
+
+	if (Blob)
+	{
+		// Set non-empty preview if blob is valid.
+		// Note that if the blob is tiled, it needs to be combined first.
+		if (Blob->IsTiled())
+		{
+			Blob->OnFinalise()
+			    .then([Blob]()
+			    {
+				    const TiledBlobPtr BlobTiled = std::static_pointer_cast<TiledBlob>(Blob);
+				    return BlobTiled->CombineTiles(false, false);
+			    })
+			    .then(UpdatePreview);
+		}
+		else
+		{
+			Blob->OnFinalise()
+			    .then(UpdatePreview);
+		}
+	}
+	else
+	{
+		// Set black preview, i.e. Blob is nullptr.
+		UpdatePreview();
+	}
+
+	// Update preview blob and trigger related external updates.
+	[[maybe_unused]] bool Result = OnNodeBlobChanged.ExecuteIfBound(Blob);
+}
+
+FReply STG_NodePreviewWidget::OnKeyDown(const FGeometry& InGeometry, const FKeyEvent& InKeyEvent)
+{
+	if (CommandList.IsValid() && CommandList->ProcessCommandBindings(InKeyEvent))
+	{
+		return FReply::Handled();
+	}
+
+	return FReply::Unhandled();
+}
+
+void STG_NodePreviewWidget::AddFormatLabel(SHorizontalBox& HorizontalBox)
+{
+	HorizontalBox.AddSlot()
+				 .VAlign(VAlign_Center)
+	[
+		SNew(STextBlock)
+			.Text_Raw(NodeViewer.Get(), &FNodeViewer::GetFormatLabelText)
+	];
+}
+
+void STG_NodePreviewWidget::AddLockButton(FToolBarBuilder& ToolbarBuilder) const
+{
+	auto GetLockButtonImage = [this]
+	{
+		return FSlateIcon(FAppStyle::Get().GetStyleSetName(), LockedNode ? "PropertyWindow.Locked" : "PropertyWindow.Unlocked");
+	};
+
+	ToolbarBuilder.AddToolBarButton(FPreviewViewerCommands::Get().ToggleLock, NAME_None, TAttribute<FText>(), TAttribute<FText>(),
+	                                TAttribute<FSlateIcon>::CreateLambda(GetLockButtonImage));
+
+	ToolbarBuilder.AddSeparator();
+}
+
+void STG_NodePreviewWidget::AddRGBAButtons(FToolBarBuilder& ToolbarBuilder)
+{
+	ToolbarBuilder.AddSeparator();
+
+	auto GetTextButton = [this](const FString& Label, const FCheckBoxStyle* ButtonStyle, bool& bChecked)
+	{
+		return SNew(SCheckBox)
+			.Style(ButtonStyle)
+			.IsEnabled_Lambda([this, &bChecked]
+		                      {
+			                      return NodeViewer->GetCurrentImageInfo().bIsValid && (&bChecked == &bRGBA[0] ||
+				                      !NodeViewer->IsSingleChannel());
+		                      })
+			.IsChecked(bChecked ? ECheckBoxState::Checked : ECheckBoxState::Unchecked)
+			.OnCheckStateChanged_Lambda([&bChecked, this](const ECheckBoxState State)
+		                      {
+			                      bChecked = State == ECheckBoxState::Checked;
+			                      NodeViewer->SetRGBA(bRGBA[0], bRGBA[1], bRGBA[2], bRGBA[3]);
+		                      })
+		[
+			SNew(STextBlock)
+					.Font(FAppStyle::GetFontStyle("EditorViewportToolBar.Font"))
+					.Text(FText::FromString(Label))
+		];
+	};
+
+	const FCheckBoxStyle* ButtonStyleStart = &FAppStyle::Get().GetWidgetStyle<FCheckBoxStyle>("EditorViewportToolBar.ToggleButton.Start");
+	const FCheckBoxStyle* ButtonStyleMiddle = &FAppStyle::Get().GetWidgetStyle<FCheckBoxStyle>("EditorViewportToolBar.ToggleButton.Middle");
+	const FCheckBoxStyle* ButtonStyleEnd = &FAppStyle::Get().GetWidgetStyle<FCheckBoxStyle>("EditorViewportToolBar.ToggleButton.End");
+
+	const TSharedRef<SHorizontalBox> RGBA = SNew(SHorizontalBox)
+		+ SHorizontalBox::Slot()
+		.AutoWidth()
+		[
+			GetTextButton("R", ButtonStyleStart, bRGBA[0])
+		]
+		+ SHorizontalBox::Slot()
+		.AutoWidth()
+		[
+			GetTextButton("G", ButtonStyleMiddle, bRGBA[1])
+		]
+		+ SHorizontalBox::Slot()
+		.AutoWidth()
+		[
+			GetTextButton("B", ButtonStyleMiddle, bRGBA[2])
+		]
+		+ SHorizontalBox::Slot()
+		.AutoWidth()
+		[
+			GetTextButton("A", ButtonStyleEnd, bRGBA[3])
+		];
+
+	ToolbarBuilder.AddToolBarWidget(RGBA);
+}
+
+void STG_NodePreviewWidget::ToggleLock()
+{
+	if (LockedNode)
+	{
+		const bool bUpdatePreview = SelectedNode != LockedNode;
+
+		LockedNode = nullptr;
+
+		if (bUpdatePreview)
+		{
+			Update();
+		}
+	}
+	else
+	{
+		LockedNode = SelectedNode;
+	}
+}
+
+#undef LOCTEXT_NAMESPACE
+
+#endif // TEXTUREGRAPHEDITOR_ENABLE_NEW_NODE_PREVIEW
diff --git a/Engine/Plugins/Experimental/TextureGraph/Source/TextureGraphEditor/Private/STG_NodePreview.h b/Engine/Plugins/Experimental/TextureGraph/Source/TextureGraphEditor/Private/STG_NodePreview.h
new file mode 100644
--- /dev/null
+++ b/Engine/Plugins/Experimental/TextureGraph/Source/TextureGraphEditor/Private/STG_NodePreview.h
@@ -0,0 +1,142 @@
+﻿// Copyright Epic Games, Inc. All Rights Reserved.
+
+#pragma once
+
+#if TEXTUREGRAPHEDITOR_ENABLE_NEW_NODE_PREVIEW
+
+#include "IImageViewer.h"
+#include "Data/Blob.h"
+#include "Widgets/SCompoundWidget.h"
+
+namespace UE::ImageWidgets
+{
+	class SImageViewport;
+}
+
+class FToolBarBuilder;
+class FUICommandList;
+class SHorizontalBox;
+class UTG_EdGraphNode;
+
+/**
+ * Image viewer implementation that holds the node buffer information and draws the node texture.
+ */
+class FNodeViewer final : public UE::ImageWidgets::IImageViewer
+{
+public:
+	// IImageViewer overrides - begin
+	virtual FImageInfo GetCurrentImageInfo() const override;
+	virtual void DrawCurrentImage(FViewport* Viewport, FCanvas* Canvas, const FDrawProperties& Properties) override;
+	virtual TOptional<TVariant<FColor, FLinearColor>> GetCurrentImagePixelColor(FIntPoint PixelCoords, int32 MipIndex) const override;
+#if IMAGE_WIDGETS_WITH_CATALOG
+	virtual void OnImageSelected(const FGuid& ImageGuid) override {};
+#endif
+#if IMAGE_WIDGETS_WITH_AB_COMPARISON
+	virtual bool IsValidImage(const FGuid& Guid) const override { return true; }
+#endif
+	// IImageViewer overrides - end
+
+	/** Returns the format label for the status bar */
+	FText GetFormatLabelText() const;
+
+	/** Indicates that a node is single channel, i.e. grayscale. */
+	bool IsSingleChannel() const;
+
+	/** Sets the node buffer to a given node. */
+	void SetTexture(const BlobPtr& InBlob);
+
+	/** Toggles RGBA components for display. */
+	void SetRGBA(bool bR, bool bG, bool bB, bool bA);
+
+private:
+	/** Draws the node texture in the viewport */
+	void DrawTexture(const FTextureResource* TextureResource, FCanvas* Canvas, const FDrawProperties::FPlacement& TilePlacementInfo,
+	                 const FDrawProperties::FMip& Mip) const;
+
+	/** Determines the blend mode based on the node texture and the RGBA toggles. */
+	ESimpleElementBlendMode GetBlendMode() const;
+
+	/** Retrieves the render target texture  from a given node buffer. */
+	UTextureRenderTarget2D* GetTextureFromBuffer(const DeviceBufferPtr& Buffer) const;
+
+	/** Indicates that a node is in sRGB format. */
+	bool IsSRGB() const;
+
+	/** The render target texture for the node. */
+	UTextureRenderTarget2D* NodeTexture = nullptr;
+
+	/** Toggles for enabling RGBA components for drawing. */
+	bool bRGBA[4] = {true, true, true, true};
+
+	/** Duplicate CPU buffer with the node's pixel data for quick lookup of the pixel under the cursor. */
+	TArray<FLinearColor> NodePixels;
+
+	/** Node meta data.  */
+	BufferDescriptor NodeDescriptor;
+};
+
+/**
+ * Widget for the Node Preview tab containing the image viewport.
+ */
+class STG_NodePreviewWidget : public SCompoundWidget
+{
+public:
+	DECLARE_DELEGATE_OneParam(FOnNodeBlobChanged, BlobPtr Blob);
+
+	SLATE_BEGIN_ARGS(STG_NodePreviewWidget)
+		{
+		}
+
+		/** Callback to notify about the node preview texture having changed. */
+		SLATE_EVENT(FOnNodeBlobChanged, OnNodeBlobChanged)
+	SLATE_END_ARGS()
+
+	virtual ~STG_NodePreviewWidget() override;
+
+	void Construct(const FArguments& InArgs);
+
+	/** Notify the preview about the node selection having changed. */
+	void SelectionChanged(UTG_EdGraphNode* Node);
+
+	/** Notify the preview about a node being deleted. */
+	void NodeDeleted(const UTG_EdGraphNode* Node);
+
+	/** Trigger an update of the node preview after the displayed contents changed. */
+	void Update() const;
+
+	// SWidget overrides - begin
+	virtual FReply OnKeyDown(const FGeometry& InGeometry, const FKeyEvent& InKeyEvent) override;
+	// SWidget overrides - end
+
+private:
+	// Add toolbar and status bar extensions to the viewport.
+	void AddFormatLabel(SHorizontalBox& HorizontalBox);
+	void AddLockButton(FToolBarBuilder& ToolbarBuilder) const;
+	void AddRGBAButtons(FToolBarBuilder& ToolbarBuilder);
+
+	/** Toggles the node preview lock based on the lock button extension in the viewport toolbar. */
+	void ToggleLock();
+
+	/** The viewport widget for displaying the node preview. */
+	TSharedPtr<UE::ImageWidgets::SImageViewport> Viewport;
+
+	/** The image viewer implementation holding and drawing the actual image. */
+	TSharedPtr<FNodeViewer> NodeViewer;
+
+	/** The additional commands used for this widget. */
+	TSharedPtr<FUICommandList> CommandList;
+
+	/** Pointer to the currently selected node, which might be different if the node preview is locked to another node. */
+	UTG_EdGraphNode* SelectedNode = nullptr;
+
+	/** Pointer to the node the node preview is currently locked on. */
+	UTG_EdGraphNode* LockedNode = nullptr;
+
+	/** Callback to notify about the node preview texture having changed to make sure that for example the histogram is synced with the node preview. */
+	FOnNodeBlobChanged OnNodeBlobChanged;
+
+	/** Flags for toggling RGBA channels in the preview, which are hooked up to the RGBA buttons in the viewport toolbar extension. */
+	bool bRGBA[4] = {true, true, true, true};
+};
+
+#endif // TEXTUREGRAPHEDITOR_ENABLE_NEW_NODE_PREVIEW
diff --git a/Engine/Plugins/Experimental/TextureGraph/Source/TextureGraphEditor/Private/TG_Editor.cpp b/Engine/Plugins/Experimental/TextureGraph/Source/TextureGraphEditor/Private/TG_Editor.cpp
--- a/Engine/Plugins/Experimental/TextureGraph/Source/TextureGraphEditor/Private/TG_Editor.cpp
+++ b/Engine/Plugins/Experimental/TextureGraph/Source/TextureGraphEditor/Private/TG_Editor.cpp
@@ -69,6 +69,9 @@
 #include "Model/Mix/ViewportSettings.h"
 #include "UObject/MetaData.h"
 #include "UObject/ObjectSaveContext.h"
+#if TEXTUREGRAPHEDITOR_ENABLE_NEW_NODE_PREVIEW
+#include "STG_NodePreview.h"
+#endif // TEXTUREGRAPHEDITOR_ENABLE_NEW_NODE_PREVIEW
 
 #include "Expressions/Input/TG_Expression_InputParam.h"
 #include "STG_OutputSelectionDlg.h"
@@ -118,10 +121,19 @@ void FTG_Editor::RegisterTabSpawners(const TSharedRef<class FTabManager>& InTabM
 		.SetGroup(WorkspaceMenuCategoryRef)
 		.SetIcon(FSlateIcon(FAppStyle::GetAppStyleSetName(), "LevelEditor.Tabs.Details"));
 
+#if TEXTUREGRAPHEDITOR_ENABLE_OLD_SELECTION_PREVIEW
 	InTabManager->RegisterTabSpawner(FTG_EditorTabs::SelectionPreviewTabId, FOnSpawnTab::CreateSP(this, &FTG_Editor::SpawnTab_SelectionPreview))
 		.SetDisplayName(LOCTEXT("SelectionPreviewTab", "Node Preview"))
 		.SetGroup(WorkspaceMenuCategoryRef)
+	.SetIcon(FSlateIcon(FAppStyle::GetAppStyleSetName(), "LevelEditor.Tabs.Viewports"));
+#endif // TEXTUREGRAPHEDITOR_ENABLE_OLD_SELECTION_PREVIEW
+
+#if TEXTUREGRAPHEDITOR_ENABLE_NEW_NODE_PREVIEW
+	InTabManager->RegisterTabSpawner(FTG_EditorTabs::NodePreviewTabId, FOnSpawnTab::CreateSP(this, &FTG_Editor::SpawnTab_NodePreview))
+		.SetDisplayName(LOCTEXT("NodePreviewTab", "Node Preview"))
+		.SetGroup(WorkspaceMenuCategoryRef)
 		.SetIcon(FSlateIcon(FAppStyle::GetAppStyleSetName(), "LevelEditor.Tabs.Viewports"));
+#endif // TEXTUREGRAPHEDITOR_ENABLE_NEW_NODE_PREVIEW
 
 	/*InTabManager->RegisterTabSpawner(FTG_EditorTabs::OutputTabId, FOnSpawnTab::CreateSP(this, &FTG_Editor::SpawnTab_Output))
 		.SetDisplayName(LOCTEXT("OutputTab", "Output"))
@@ -154,7 +166,12 @@ void FTG_Editor::UnregisterTabSpawners(const TSharedRef<class FTabManager>& InTa
 	//InTabManager->UnregisterTabSpawner(FTG_EditorTabs::FindTabId);
 	InTabManager->UnregisterTabSpawner(FTG_EditorTabs::PreviewSceneSettingsTabId);
 	InTabManager->UnregisterTabSpawner(FTG_EditorTabs::ParameterDefaultsTabId);
+#if TEXTUREGRAPHEDITOR_ENABLE_OLD_SELECTION_PREVIEW
 	InTabManager->UnregisterTabSpawner(FTG_EditorTabs::SelectionPreviewTabId);
+#endif // TEXTUREGRAPHEDITOR_ENABLE_OLD_SELECTION_PREVIEW
+#if TEXTUREGRAPHEDITOR_ENABLE_NEW_NODE_PREVIEW
+	InTabManager->UnregisterTabSpawner(FTG_EditorTabs::NodePreviewTabId);
+#endif // TEXTUREGRAPHEDITOR_ENABLE_NEW_NODE_PREVIEW
 	InTabManager->UnregisterTabSpawner(FTG_EditorTabs::OutputTabId);
 	InTabManager->UnregisterTabSpawner(FTG_EditorTabs::PreviewSettingsTabId);
 	InTabManager->UnregisterTabSpawner(FTG_EditorTabs::ErrorsTabId);
@@ -207,7 +224,12 @@ void FTG_Editor::InitEditor(const EToolkitMode::Type Mode, const TSharedPtr< cla
 	EditedTextureGraph->Graph()->OnTGNodeRenamedDelegate.AddSP(this, &FTG_Editor::OnNodeRenamed);
 
 	GraphEditorWidget = CreateGraphEditorWidget();
+#if TEXTUREGRAPHEDITOR_ENABLE_OLD_SELECTION_PREVIEW
 	SelectionPreview = CreateSelectionViewWidget();
+#endif // TEXTUREGRAPHEDITOR_ENABLE_OLD_SELECTION_PREVIEW
+#if TEXTUREGRAPHEDITOR_ENABLE_NEW_NODE_PREVIEW
+	NodePreview = CreateNodePreviewWidget();
+#endif // TEXTUREGRAPHEDITOR_ENABLE_NEW_NODE_PREVIEW
 	TextureDetails = CreateTextureDetailsWidget();
 	Palette = SNew(STG_Palette, SharedThis(this));
 
@@ -308,9 +330,18 @@ void FTG_Editor::InitEditor(const EToolkitMode::Type Mode, const TSharedPtr< cla
 					(
 						FTabManager::NewStack()
 						->SetHideTabWell(true)
+#if TEXTUREGRAPHEDITOR_ENABLE_OLD_SELECTION_PREVIEW
 						->AddTab(FTG_EditorTabs::SelectionPreviewTabId, ETabState::OpenedTab)
+#endif // TEXTUREGRAPHEDITOR_ENABLE_OLD_SELECTION_PREVIEW
+#if TEXTUREGRAPHEDITOR_ENABLE_NEW_NODE_PREVIEW
+						->AddTab(FTG_EditorTabs::NodePreviewTabId, ETabState::OpenedTab)
+#endif // TEXTUREGRAPHEDITOR_ENABLE_NEW_NODE_PREVIEW
 						->AddTab(FTG_EditorTabs::TextureDetailsTabId, ETabState::OpenedTab)
-						->SetForegroundTab(FTG_EditorTabs::SelectionPreviewTabId)
+#if TEXTUREGRAPHEDITOR_ENABLE_NEW_NODE_PREVIEW
+						->SetForegroundTab(FTG_EditorTabs::NodePreviewTabId)
+#else
+						->SetForegroundTab(FTG_EditorTabs::SelectionPreviewTabId)						
+#endif // TEXTUREGRAPHEDITOR_ENABLE_NEW_NODE_PREVIEW
 					)
 					->Split
 					(
@@ -637,7 +668,7 @@ void FTG_Editor::OnRenderingDone(UMixInterface* TextureGraph, const FInvalidatio
 {
 	if (TextureGraph != nullptr)
 	{
-		RefreshSelectionPreview(GraphEditorWidget->GetSelectedNodes(), Details);
+		RefreshNodePreview(GraphEditorWidget->GetSelectedNodes(), Details, true);
 	}
 }
 
@@ -837,6 +868,7 @@ bool FTG_Editor::OnVerifyNodeTextCommit(const FText& NewText, UEdGraphNode* Node
 	}
 	return bValid;
 }
+
 TSharedRef<class SGraphEditor> FTG_Editor::CreateGraphEditorWidget()
 {
 	if (!GraphEditorCommands)
@@ -909,6 +941,14 @@ TSharedRef<class STG_SelectionPreview> FTG_Editor::CreateSelectionViewWidget()
 		.OnBlobSelectionChanged(this, &FTG_Editor::OnSelectedBlobChanged);
 }
 
+#if TEXTUREGRAPHEDITOR_ENABLE_NEW_NODE_PREVIEW
+TSharedRef<STG_NodePreviewWidget> FTG_Editor::CreateNodePreviewWidget()
+{
+	return SNew(STG_NodePreviewWidget)
+		.OnNodeBlobChanged(this, &FTG_Editor::OnSelectedBlobChanged);
+}
+#endif // TEXTUREGRAPHEDITOR_ENABLE_NEW_NODE_PREVIEW
+
 TSharedRef<class STG_TextureDetails> FTG_Editor::CreateTextureDetailsWidget()
 {
 	return SNew(STG_TextureDetails);
@@ -1039,6 +1079,7 @@ TSharedRef<SDockTab> FTG_Editor::SpawnTab_ParameterDefaults(const FSpawnTabArgs&
 		];
 }
 
+#if TEXTUREGRAPHEDITOR_ENABLE_OLD_SELECTION_PREVIEW
 TSharedRef<SDockTab> FTG_Editor::SpawnTab_SelectionPreview(const FSpawnTabArgs& Args)
 {
 	check(Args.GetTabId() == FTG_EditorTabs::SelectionPreviewTabId);
@@ -1053,6 +1094,21 @@ TSharedRef<SDockTab> FTG_Editor::SpawnTab_SelectionPreview(const FSpawnTabArgs&
 	return Tab.ToSharedRef();
 
 }
+#endif // TEXTUREGRAPHEDITOR_ENABLE_OLD_SELECTION_PREVIEW
+
+#if TEXTUREGRAPHEDITOR_ENABLE_NEW_NODE_PREVIEW
+TSharedRef<SDockTab> FTG_Editor::SpawnTab_NodePreview(const FSpawnTabArgs& Args)
+{
+	check(Args.GetTabId() == FTG_EditorTabs::NodePreviewTabId);
+
+	return SNew(SDockTab)
+		// .Label(LOCTEXT("NodeViewer", "Node Viewer"))
+		.TabColorScale(GetTabColorScale())
+		[
+			NodePreview.ToSharedRef()
+		];
+}
+#endif // TEXTUREGRAPHEDITOR_ENABLE_NEW_NODE_PREVIEW
 
 TSharedRef<SDockTab> FTG_Editor::SpawnTab_Output(const FSpawnTabArgs& Args)
 {
@@ -1201,7 +1257,7 @@ void FTG_Editor::OnSelectedNodesChanged(const TSet<class UObject*>& NewSelection
 	FocusDetailsPanel();
 
 	RefreshPreviewViewport();
-	RefreshSelectionPreview(NewSelection, nullptr);
+	RefreshNodePreview(NewSelection, nullptr, false);
 }
 
 void FTG_Editor::OnNodeDoubleClicked(UEdGraphNode* Node)
@@ -1264,10 +1320,16 @@ bool FTG_Editor::DeleteNodes(TArray<UEdGraphNode*> NodesToDelete, bool ForceDele
 					UTG_Node* TGNode = TGEdGraphNode->GetNode();
 					check(TGNode);
 
+#if TEXTUREGRAPHEDITOR_ENABLE_OLD_SELECTION_PREVIEW
 					if (SelectionPreview->GetSelectedNode() && SelectionPreview->GetSelectedNode()->GetName() == TGEdGraphNode->GetName())
 					{
 						SelectionPreview->OnSelectedNodeDeleted();
 					}
+#endif // TEXTUREGRAPHEDITOR_ENABLE_OLD_SELECTION_PREVIEW
+#if TEXTUREGRAPHEDITOR_ENABLE_NEW_NODE_PREVIEW
+					NodePreview->NodeDeleted(TGEdGraphNode);
+#endif // TEXTUREGRAPHEDITOR_ENABLE_NEW_NODE_PREVIEW
+
 					TextureGraph->RemoveNode(TGNode);
 				}
 				EdGraphNode->DestroyNode();
@@ -1903,7 +1965,7 @@ void FTG_Editor::RefreshPreviewViewport()
 	//TODO: refresh the viewport here
 }
 
-void FTG_Editor::RefreshSelectionPreview(const TSet<class UObject*>& NewSelection, const FInvalidationDetails* Details)
+void FTG_Editor::RefreshNodePreview(const TSet<class UObject*>& NewSelection, const FInvalidationDetails* Details, bool bUpdateOnly)
 {
 	TArray<UTG_EdGraphNode*> NodesForSelectionPreview;
 	for (TSet<class UObject*>::TConstIterator SetIt(NewSelection); SetIt; ++SetIt)
@@ -1916,7 +1978,20 @@ void FTG_Editor::RefreshSelectionPreview(const TSet<class UObject*>& NewSelectio
 
 	//TODO: Handle multiple selection case
 	UTG_EdGraphNode* SelectedNode = NodesForSelectionPreview.Num() > 0 ? NodesForSelectionPreview[0] : nullptr;
+#if TEXTUREGRAPHEDITOR_ENABLE_OLD_SELECTION_PREVIEW
 	SelectionPreview->OnSelectionChanged(SelectedNode);
+#endif // TEXTUREGRAPHEDITOR_ENABLE_OLD_SELECTION_PREVIEW
+
+#if TEXTUREGRAPHEDITOR_ENABLE_NEW_NODE_PREVIEW
+	if (bUpdateOnly)
+	{
+		NodePreview->Update();
+	}
+	else
+	{
+		NodePreview->SelectionChanged(SelectedNode);
+	}
+#endif // TEXTUREGRAPHEDITOR_ENABLE_NEW_NODE_PREVIEW
 }
 
 void FTG_Editor::SetMesh(class UMeshComponent* InPreviewMesh, class UWorld* InWorld)
diff --git a/Engine/Plugins/Experimental/TextureGraph/Source/TextureGraphEditor/Private/TG_Editor.h b/Engine/Plugins/Experimental/TextureGraph/Source/TextureGraphEditor/Private/TG_Editor.h
--- a/Engine/Plugins/Experimental/TextureGraph/Source/TextureGraphEditor/Private/TG_Editor.h
+++ b/Engine/Plugins/Experimental/TextureGraph/Source/TextureGraphEditor/Private/TG_Editor.h
@@ -103,7 +103,7 @@ public:
 	void											RefreshPreviewViewport();
 	
 	/** Called to update the selection view */
-	void											RefreshSelectionPreview(const TSet<class UObject*>& NewSelection, const FInvalidationDetails* Details);
+	void											RefreshNodePreview(const TSet<class UObject*>& NewSelection, const FInvalidationDetails* Details, bool bUpdateOnly);
 
 	void											SetMesh(class UMeshComponent* InPreviewMesh, class UWorld* InWorld) override;
 	bool 											SetPreviewAsset(UObject* InAsset);
@@ -192,8 +192,13 @@ private:
 
 	FActionMenuContent								OnCreateGraphActionMenu(UEdGraph* InGraph, const FVector2D& InNodePosition, const TArray<UEdGraphPin*>& InDraggedPins, bool bAutoExpand, SGraphEditor::FActionMenuClosed InOnMenuClosed);
 
-	/** Create Selection view widget */
+	/** Create Node preview widget */
+#if TEXTUREGRAPHEDITOR_ENABLE_OLD_SELECTION_PREVIEW
 	TSharedRef<class STG_SelectionPreview>			CreateSelectionViewWidget();
+#endif // TEXTUREGRAPHEDITOR_ENABLE_OLD_SELECTION_PREVIEW
+#if TEXTUREGRAPHEDITOR_ENABLE_NEW_NODE_PREVIEW
+	TSharedRef<class STG_NodePreviewWidget>			CreateNodePreviewWidget();
+#endif // TEXTUREGRAPHEDITOR_ENABLE_NEW_NODE_PREVIEW
 
 	/** Create Texture details widget */
 	TSharedRef<class STG_TextureDetails>			CreateTextureDetailsWidget();
@@ -253,7 +258,12 @@ private:
 	TSharedRef<SDockTab>							SpawnTab_Find(const FSpawnTabArgs& Args);
 	TSharedRef<SDockTab>							SpawnTab_PreviewSettings(const FSpawnTabArgs& Args);
 	TSharedRef<SDockTab>							SpawnTab_ParameterDefaults(const FSpawnTabArgs& Args);
+#if TEXTUREGRAPHEDITOR_ENABLE_OLD_SELECTION_PREVIEW
 	TSharedRef<SDockTab>							SpawnTab_SelectionPreview(const FSpawnTabArgs& Args);
+#endif // TEXTUREGRAPHEDITOR_ENABLE_OLD_SELECTION_PREVIEW
+#if TEXTUREGRAPHEDITOR_ENABLE_NEW_NODE_PREVIEW
+	TSharedRef<SDockTab>							SpawnTab_NodePreview(const FSpawnTabArgs& Args);
+#endif // TEXTUREGRAPHEDITOR_ENABLE_NEW_NODE_PREVIEW
 	TSharedRef<SDockTab>							SpawnTab_Output(const FSpawnTabArgs& Args);
 	TSharedRef<SDockTab>							SpawnTab_Settings(const FSpawnTabArgs& Args);
 	TSharedRef<SDockTab>							SpawnTab_Errors(const FSpawnTabArgs& Args);
@@ -301,10 +311,15 @@ private:
 	/** Graph editor widget being displayed*/
 	TSharedPtr<SGraphEditor>						GraphEditorWidget;
 
-	/** Selection preview widget being displayed*/
+	/** Node preview widget being displayed*/
+#if TEXTUREGRAPHEDITOR_ENABLE_OLD_SELECTION_PREVIEW
 	TSharedPtr<class STG_SelectionPreview>			SelectionPreview;
+#endif // TEXTUREGRAPHEDITOR_ENABLE_OLD_SELECTION_PREVIEW
+#if TEXTUREGRAPHEDITOR_ENABLE_NEW_NODE_PREVIEW
+	TSharedPtr<STG_NodePreviewWidget>				NodePreview;
+#endif // TEXTUREGRAPHEDITOR_ENABLE_NEW_NODE_PREVIEW
 
-	/** Selection preview widget being displayed*/
+	/** Texture details widget being displayed*/
 	TSharedPtr<class STG_TextureDetails>			TextureDetails;
 
 	/** Command list for this editor */
diff --git a/Engine/Plugins/Experimental/TextureGraph/Source/TextureGraphEditor/Private/TG_EditorTabs.cpp b/Engine/Plugins/Experimental/TextureGraph/Source/TextureGraphEditor/Private/TG_EditorTabs.cpp
--- a/Engine/Plugins/Experimental/TextureGraph/Source/TextureGraphEditor/Private/TG_EditorTabs.cpp
+++ b/Engine/Plugins/Experimental/TextureGraph/Source/TextureGraphEditor/Private/TG_EditorTabs.cpp
@@ -12,7 +12,12 @@ const FName FTG_EditorTabs::FindTabId					(TEXT("TG_Editor_Find"));
 const FName FTG_EditorTabs::GraphEditorId				(TEXT("TG_Editor_GraphEditor"));
 const FName FTG_EditorTabs::PreviewSceneSettingsTabId	(TEXT("TG_Editor_PreviewSceneSettings"));
 const FName FTG_EditorTabs::ParameterDefaultsTabId		(TEXT("TG_Editor_ParameterDefaults"));
+#if TEXTUREGRAPHEDITOR_ENABLE_OLD_SELECTION_PREVIEW
 const FName FTG_EditorTabs::SelectionPreviewTabId		(TEXT("TG_Editor_SelectionPreview"));
+#endif // TEXTUREGRAPHEDITOR_ENABLE_OLD_SELECTION_PREVIEW
+#if TEXTUREGRAPHEDITOR_ENABLE_NEW_NODE_PREVIEW
+const FName FTG_EditorTabs::NodePreviewTabId			(TEXT("TG_Editor_NodePreview"));
+#endif // TEXTUREGRAPHEDITOR_ENABLE_NEW_NODE_PREVIEW
 const FName FTG_EditorTabs::OutputTabId					(TEXT("TG_Editor_Output"));
 const FName FTG_EditorTabs::PreviewSettingsTabId		(TEXT("TG_Editor_PreviewSettings"));
 const FName FTG_EditorTabs::ErrorsTabId					(TEXT("TG_Editor_Stats"));
diff --git a/Engine/Plugins/Experimental/TextureGraph/Source/TextureGraphEditor/Private/TG_EditorTabs.h b/Engine/Plugins/Experimental/TextureGraph/Source/TextureGraphEditor/Private/TG_EditorTabs.h
--- a/Engine/Plugins/Experimental/TextureGraph/Source/TextureGraphEditor/Private/TG_EditorTabs.h
+++ b/Engine/Plugins/Experimental/TextureGraph/Source/TextureGraphEditor/Private/TG_EditorTabs.h
@@ -15,7 +15,12 @@ struct FTG_EditorTabs
 	static const FName GraphEditorId;
 	static const FName PreviewSceneSettingsTabId;
 	static const FName ParameterDefaultsTabId;
+#if TEXTUREGRAPHEDITOR_ENABLE_OLD_SELECTION_PREVIEW
 	static const FName SelectionPreviewTabId;
+#endif // TEXTUREGRAPHEDITOR_ENABLE_OLD_SELECTION_PREVIEW
+#if TEXTUREGRAPHEDITOR_ENABLE_NEW_NODE_PREVIEW
+	static const FName NodePreviewTabId;
+#endif // TEXTUREGRAPHEDITOR_ENABLE_NEW_NODE_PREVIEW
 	static const FName OutputTabId;
 	static const FName PreviewSettingsTabId;
 	static const FName ErrorsTabId;
diff --git a/Engine/Plugins/Experimental/TextureGraph/Source/TextureGraphEditor/TextureGraphEditor.Build.cs b/Engine/Plugins/Experimental/TextureGraph/Source/TextureGraphEditor/TextureGraphEditor.Build.cs
--- a/Engine/Plugins/Experimental/TextureGraph/Source/TextureGraphEditor/TextureGraphEditor.Build.cs
+++ b/Engine/Plugins/Experimental/TextureGraph/Source/TextureGraphEditor/TextureGraphEditor.Build.cs
@@ -2,6 +2,11 @@
 using UnrealBuildTool;
 public class TextureGraphEditor : ModuleRules
 {
+	// Flag that enables the new node preview.
+	private const bool EnableNewNodePreview = true;
+	// Flag that disables the old selection preview. This is only used if the EnableNewNodePreview is set to true.
+	private const bool DisableOldSelectionPreview = false;
+
 	public TextureGraphEditor(ReadOnlyTargetRules Target) : base(Target)
 	{
 		PCHUsage = ModuleRules.PCHUsageMode.UseExplicitOrSharedPCHs;
@@ -67,5 +72,17 @@ public class TextureGraphEditor : ModuleRules
 				// ... add any modules that your module loads dynamically here ...
 			}
 			);
+
+		if (EnableNewNodePreview)
+		{
+			PrivateDependencyModuleNames.Add("ImageWidgets");
+		}
+
+		PublicDefinitions.Add(EnableNewNodePreview
+			? "TEXTUREGRAPHEDITOR_ENABLE_NEW_NODE_PREVIEW=1"
+			: "TEXTUREGRAPHEDITOR_ENABLE_NEW_NODE_PREVIEW=0");
+		PublicDefinitions.Add(EnableNewNodePreview && DisableOldSelectionPreview
+			? "TEXTUREGRAPHEDITOR_ENABLE_OLD_SELECTION_PREVIEW=0"
+			: "TEXTUREGRAPHEDITOR_ENABLE_OLD_SELECTION_PREVIEW=1");
 	}
 }
\ No newline at end of file
diff --git a/Engine/Plugins/Experimental/TextureGraph/TextureGraph.uplugin b/Engine/Plugins/Experimental/TextureGraph/TextureGraph.uplugin
--- a/Engine/Plugins/Experimental/TextureGraph/TextureGraph.uplugin
+++ b/Engine/Plugins/Experimental/TextureGraph/TextureGraph.uplugin
@@ -71,6 +71,10 @@
 		{
 			"Name": "RenderDocPlugin",
 			"Enabled": true
+		},
+		{
+			"Name": "ImageWidgets",
+			"Enabled": true
 		}
 	],
 	"SupportedTargetPlatforms": [
