Fixed the custom TG_Scalar widget used for the threshold value of the Threshold expression.
learning from SPropertyEditorNumeric, we cleaned up and simplify the detail view widget created for a scalar

#tests 65e6eff3925165fabe3570f7

diff --git a/Engine/Plugins/Experimental/TextureGraph/Source/TextureGraphEditor/Private/Customizations/TG_ScalarCustomization.h b/Engine/Plugins/Experimental/TextureGraph/Source/TextureGraphEditor/Private/Customizations/TG_ScalarCustomization.h
--- a/Engine/Plugins/Experimental/TextureGraph/Source/TextureGraphEditor/Private/Customizations/TG_ScalarCustomization.h
+++ b/Engine/Plugins/Experimental/TextureGraph/Source/TextureGraphEditor/Private/Customizations/TG_ScalarCustomization.h
@@ -26,9 +26,12 @@ public:
 
 class FTG_ScalarCustomization : public IPropertyTypeCustomization
 {
-	bool bIsUsingSlider = false;
+	/** Edited Property */
 	TSharedPtr<IPropertyHandle> ScalarHandle;
 
+	/** True if the slider is being used to change the value of the property */
+	bool bIsUsingSlider = false;
+
 public:
 	static TSharedRef<IPropertyTypeCustomization> Create()
 	{
@@ -38,33 +41,32 @@ public:
 	virtual void CustomizeHeader(TSharedRef<IPropertyHandle> PropertyHandle, FDetailWidgetRow& HeaderRow, IPropertyTypeCustomizationUtils& CustomizationUtils) override
 	{
 		ScalarHandle = PropertyHandle;
-		check(ScalarHandle.IsValid());
-
-		float FloatValue;
-		ScalarHandle->GetValue(FloatValue);
-
-		TWeakPtr<IPropertyHandle> ScalarHandlePtr = ScalarHandle;
 
 		// Get min/max metadata values if defined
 		float MinAllowedValue = 0.0f;
 		float MinUIValue = 0.0f;
 		float MaxAllowedValue = 1.0f;
 		float MaxUIValue = 1.0f;
-		if (ScalarHandle->HasMetaData(TEXT("ClampMin")))
-		{
-			MinAllowedValue = ScalarHandle->GetFloatMetaData(TEXT("ClampMin"));
-		}
-		if (ScalarHandle->HasMetaData(TEXT("UIMin")))
-		{
-			MinUIValue = ScalarHandle->GetFloatMetaData(TEXT("UIMin"));
-		}
-		if (ScalarHandle->HasMetaData(TEXT("ClampMax")))
+		FText DisplayName;
+		if (ScalarHandle.IsValid())
 		{
-			MaxAllowedValue = ScalarHandle->GetFloatMetaData(TEXT("ClampMax"));
-		}
-		if (ScalarHandle->HasMetaData(TEXT("UIMax")))
-		{
-			MaxUIValue = ScalarHandle->GetFloatMetaData(TEXT("UIMax"));
+			if (ScalarHandle->HasMetaData(TEXT("ClampMin")))
+			{
+				MinAllowedValue = ScalarHandle->GetFloatMetaData(TEXT("ClampMin"));
+			}
+			if (ScalarHandle->HasMetaData(TEXT("UIMin")))
+			{
+				MinUIValue = ScalarHandle->GetFloatMetaData(TEXT("UIMin"));
+			}
+			if (ScalarHandle->HasMetaData(TEXT("ClampMax")))
+			{
+				MaxAllowedValue = ScalarHandle->GetFloatMetaData(TEXT("ClampMax"));
+			}
+			if (ScalarHandle->HasMetaData(TEXT("UIMax")))
+			{
+				MaxUIValue = ScalarHandle->GetFloatMetaData(TEXT("UIMax"));
+			}
+			DisplayName = ScalarHandle->GetPropertyDisplayName();
 		}
 
 		// Build the  ui
@@ -72,51 +74,67 @@ public:
 			.NameContent()
 			[
 				SNew(STextBlock)
-				.Text(PropertyHandle->GetPropertyDisplayName())
+				.Text(DisplayName)
 				.Font(CustomizationUtils.GetRegularFont())
 			]
 			.ValueContent()
 			.MinDesiredWidth(STG_TextureHistogram::PreferredWidth)
 			[
 				SNew(SSlider)
-				.Value(this, &FTG_ScalarCustomization::OnGetValue, ScalarHandlePtr)
+				.Value(this, &FTG_ScalarCustomization::OnGetValue)
 				.MinValue(MinAllowedValue)
 				.MaxValue(MaxAllowedValue)
-				.OnValueChanged(this, &FTG_ScalarCustomization::OnValueChanged, ScalarHandlePtr)
+				.OnValueChanged(this, &FTG_ScalarCustomization::OnValueChanged)
 				.OnMouseCaptureBegin(this, &FTG_ScalarCustomization::OnBeginSliderMovement)
 				.OnMouseCaptureEnd(this, &FTG_ScalarCustomization::OnEndSliderMovement)
 			];
 	}
 
-	float OnGetValue(TWeakPtr<IPropertyHandle> HandleWeakPtr) const
+	float OnGetValue() const
 	{
-		auto ValueSharedPtr = HandleWeakPtr.Pin();
-		float Value = 0;
-		if (ValueSharedPtr.IsValid())
+		float NumericVal = 0;
+		if (ScalarHandle.IsValid())
 		{
-			ensure(ValueSharedPtr->GetValue(Value) == FPropertyAccess::Success);
+			if (ScalarHandle->GetValue(NumericVal) != FPropertyAccess::Fail)
+			{
+				return NumericVal;
+			}
 		}
-		return Value;
-	}
-
-	void OnBeginSliderMovement()
-	{
-		bIsUsingSlider = true;
-		GEditor->BeginTransaction(LOCTEXT("SetScalarProperty", "Set Scalar Property"));
+		return NumericVal;
 	}
 
-	void OnValueChanged(float NewValue, TWeakPtr<IPropertyHandle> HandleWeakPtr)
+	void OnValueChanged(float NewValue)
 	{
 		if (bIsUsingSlider)
 		{
-			auto HandleSharedPtr = HandleWeakPtr.Pin();
-			if (HandleSharedPtr.IsValid())
+			if (ScalarHandle.IsValid())
 			{
-				ensure(HandleSharedPtr->SetValue(NewValue, EPropertyValueSetFlags::InteractiveChange) == FPropertyAccess::Success);
+				float OrgValue(0);
+				if (ScalarHandle->GetValue(OrgValue) != FPropertyAccess::Fail)
+				{
+					// Value hasn't changed, so lets return now
+					if (OrgValue == NewValue)
+					{
+						return;
+					}
+				}
+
+				// We don't create a transaction for each property change when using the slider.  Only once when the slider first is moved
+				EPropertyValueSetFlags::Type Flags = (EPropertyValueSetFlags::InteractiveChange | EPropertyValueSetFlags::NotTransactable);
+				ScalarHandle->SetValue(NewValue, Flags);
 			}
 		}
 	}
 
+	void OnBeginSliderMovement()
+	{
+		bIsUsingSlider = true;
+		if (ScalarHandle.IsValid())
+		{
+			GEditor->BeginTransaction(NSLOCTEXT("GraphEditor", "ChangeNumberPinValueSlider", "Change Number Pin Value slider"));
+		}
+	}
+
 	void OnEndSliderMovement()
 	{
 		bIsUsingSlider = false;
