Remove reference to legacy float16 from Mutable in favor of the engine's standard FFloat16 type.
#rnx

diff --git a/Engine/Plugins/Experimental/TextureGraph/Source/TextureGraphEngine/2D/TextureHelper.cpp b/Engine/Plugins/Experimental/TextureGraph/Source/TextureGraphEngine/2D/TextureHelper.cpp
--- a/Engine/Plugins/Experimental/TextureGraph/Source/TextureGraphEngine/2D/TextureHelper.cpp
+++ b/Engine/Plugins/Experimental/TextureGraph/Source/TextureGraphEngine/2D/TextureHelper.cpp
@@ -36,7 +36,6 @@
 #include "RHIGPUReadback.h"
 #include "Data/Blobber.h"
 #include "TextureGraphEngine.h"
-#include "MuR/MutableMath.h"
 #include "Curves/CurveLinearColorAtlas.h"
 #include "Engine/TextureLightProfile.h"
 #include "Model/Mix/MixSettings.h"
@@ -963,10 +962,10 @@ void ShortToHalf_DataConverter(uint8* Dst, size_t DstLength, const BufferDescrip
 
 	for (size_t PixelIndex = 0; PixelIndex < DstCount; PixelIndex++)
 	{
-		mu::float16* DstPixel = reinterpret_cast<mu::float16*>(Dst + PixelIndex * DstFormatSize);
+		FFloat16* DstPixel = reinterpret_cast<FFloat16*>(Dst + PixelIndex * DstFormatSize);
 		const uint16_t* SrcPixel = reinterpret_cast<const uint16_t*>(Src + PixelIndex * SrcFormatSize);
 		float SrcPixelF = float(*SrcPixel) / 65536.0f;
-		*DstPixel = mu::floatToHalf(SrcPixelF);
+		*DstPixel = SrcPixelF;
 	}
 }
 
