Gradient base implementation

diff --git a/Engine/Plugins/Experimental/TextureGraph/Shaders/Expressions/Expression_Gradient.usf b/Engine/Plugins/Experimental/TextureGraph/Shaders/Expressions/Expression_Gradient.usf
new file mode 100644
--- /dev/null
+++ b/Engine/Plugins/Experimental/TextureGraph/Shaders/Expressions/Expression_Gradient.usf
@@ -0,0 +1,144 @@
+// Copyright Epic Games, Inc. All Rights Reserved.
+
+#include "/Engine/Public/Platform.ush"
+#include "/Plugin/TextureGraph/SamplerStates.ush"
+#include "/Plugin/TextureGraph/TileInfo.ush"
+#include "/Plugin/TextureGraph/Paint.ush"
+
+#define LINEAR 0
+#define EXP 1
+#define LOG 1
+
+#define ROT_0 0
+#define ROT_90 1
+#define ROT_180 2
+#define ROT_270 3
+
+#ifndef GRADIENT_INTERPOLATION
+#define GRADIENT_INTERPOLATION LINEAR
+#endif
+
+#ifndef GRADIENT_ROTATION
+#define GRADIENT_ROTATION ROT_0
+#endif
+
+float4 Center = float4(0, 0, 0, 0);
+float4 Line = float4(0, 0, 0, 0);
+float4 LineDir = float4(0, 0, 0, 0);
+
+float CalcExp(float Input)
+{
+	const float eMax = exp(1);
+	const float eMaxInverse = 1.0 / (eMax - 1.0);
+
+	return (exp(Input) - 1) * eMaxInverse;
+}
+
+float4 FSH_GradientLinear_1(float2 InUV : TEXCOORD0) : SV_Target0
+{
+	float Input = 0, Gray = 0;
+	float2 UV = TileInfo_fromCurrentTileToLayer(InUV);
+	
+#if GRADIENT_ROTATION == ROT_0
+	Input = UV.x;
+#elif GRADIENT_ROTATION == ROT_90
+	Input = UV.y;
+#elif GRADIENT_ROTATION == ROT_180
+	Input = 1.0 - UV.x;
+#elif GRADIENT_ROTATION == ROT_270
+	Input = 1.0 - UV.y;
+#endif 
+	
+#if GRADIENT_INTERPOLATION == LINEAR
+	Gray = Input;
+#elif GRADIENT_INTERPOLATION == EXP
+	Gray = CalcExp(Input);
+#endif 
+	
+	return float4(Gray, Gray, Gray, 1);
+}
+
+float4 FSH_GradientLinear_2(float2 InUV : TEXCOORD0) : SV_Target0
+{
+	float2 UV = TileInfo_fromCurrentTileToLayer(InUV);
+	float Input = 0;
+	
+#if GRADIENT_ROTATION == ROT_0
+	Input = UV.x;
+#elif GRADIENT_ROTATION == ROT_90
+	Input = UV.y;
+#elif GRADIENT_ROTATION == ROT_180
+	Input = 1.0 - UV.x;
+#elif GRADIENT_ROTATION == ROT_270
+	Input = 1.0 - UV.y;
+#endif 
+	
+	float Gray = 1 - abs(Input * 2 - 1);
+	
+#if GRADIENT_INTERPOLATION == LINEAR
+#elif GRADIENT_INTERPOLATION == EXP
+	Gray = CalcExp(Gray);
+#endif 
+	
+	return float4(Gray, Gray, Gray, 1);
+}
+
+float4 FSH_GradientRadial(float2 InUV : TEXCOORD0) : SV_Target0
+{
+	float2 UV = TileInfo_fromCurrentTileToLayer(InUV);
+	float Input = distance(UV.xy, Center.xy);
+	float Radius = Center.z;
+	float Gray = min(Input / Radius, 1);
+	
+#if GRADIENT_INTERPOLATION == LINEAR
+#elif GRADIENT_INTERPOLATION == EXP
+	Gray = CalcExp(Gray);
+#endif 
+	
+	/// Because center is white going onto gray
+	Gray = 1 - Gray;
+	
+	return float4(Gray, Gray, Gray, 1);
+}
+
+float GetAxialValue(float2 UV)
+{
+	float2 P0 = Line.xy;
+	float2 RayDir = Line.zw;
+	float RayLenSq = LineDir.z;
+	float RayLen = LineDir.w;
+	float2 Segment = UV.xy - P0;
+	float DP = dot(Segment, RayDir);
+	float2 PtOnSegment = (DP / RayLenSq) * RayDir;
+	float Input = DP <= 0 ? 0 : clamp(length(PtOnSegment) / RayLen, 0, 1);
+	
+	return Input;
+}
+
+float4 FSH_GradientAxial1(float2 InUV : TEXCOORD0) : SV_Target0
+{
+	float2 UV = TileInfo_fromCurrentTileToLayer(InUV);
+	float Input = GetAxialValue(UV);
+	float Gray = Input; 
+	
+#if GRADIENT_INTERPOLATION == LINEAR
+#elif GRADIENT_INTERPOLATION == EXP
+	Gray = CalcExp(Input);
+#endif 
+	
+	return float4(Gray, Gray, Gray, 1);
+}
+
+float4 FSH_GradientAxial2(float2 InUV : TEXCOORD0) : SV_Target0
+{
+	float2 UV = TileInfo_fromCurrentTileToLayer(InUV);
+	float Input = GetAxialValue(UV);
+	float Gray = 1 - abs(Input * 2 - 1);
+	
+#if GRADIENT_INTERPOLATION == LINEAR
+#elif GRADIENT_INTERPOLATION == EXP
+	Gray = CalcExp(Input);
+#endif 
+	
+	return float4(Gray, Gray, Gray, 1);
+}
diff --git a/Engine/Plugins/Experimental/TextureGraph/Source/TextureGraph/Private/Expressions/Procedural/TG_Expression_Gradient.cpp b/Engine/Plugins/Experimental/TextureGraph/Source/TextureGraph/Private/Expressions/Procedural/TG_Expression_Gradient.cpp
new file mode 100644
--- /dev/null
+++ b/Engine/Plugins/Experimental/TextureGraph/Source/TextureGraph/Private/Expressions/Procedural/TG_Expression_Gradient.cpp
@@ -0,0 +1,68 @@
+// Copyright Epic Games, Inc. All Rights Reserved.
+
+#include "Expressions/Procedural/TG_Expression_Gradient.h"
+#include "FxMat/MaterialManager.h"
+#include "Job/JobBatch.h"
+
+void UTG_Expression_Gradient::Evaluate(FTG_EvaluationContext* InContext)
+{
+	Super::Evaluate(InContext);
+
+	BufferDescriptor Desc = T_Gradient::InitOutputDesc( Output.GetBufferDescriptor());
+
+	int32 CurrentRotation = 0;
+
+	if (Type == EGradientType::GT_Linear_1)
+		CurrentRotation = (int32)Rotation;
+	else if (Type == EGradientType::GT_Linear_2)
+		CurrentRotation = (int32)RotationLimited;
+
+	T_Gradient::Params params =
+	{
+		.Type = Type,
+		.Interpolation = Interpolation,
+		.Rotation = CurrentRotation,
+		.Center = Center,
+		.Radius = Radius,
+		.Point1 = Point1,
+		.Point2 = Point2
+	};
+
+	Output = T_Gradient::Create(InContext->Cycle, Desc, params, InContext->TargetId);
+}
+
+#if WITH_EDITOR
+bool UTG_Expression_Gradient::CanEditChange(const FProperty* InProperty) const
+{
+	bool bEditCondition = Super::CanEditChange(InProperty);
+
+	const FName PropertyName = InProperty->GetFName();
+
+	// Specific logic associated with Property
+	if (PropertyName == GET_MEMBER_NAME_CHECKED(UTG_Expression_Gradient, Rotation))
+	{
+		bEditCondition = Type == EGradientType::GT_Linear_1;
+	}
+	else if (PropertyName == GET_MEMBER_NAME_CHECKED(UTG_Expression_Gradient, RotationLimited))
+	{
+		bEditCondition = Type == EGradientType::GT_Linear_2;
+	}
+	else if (PropertyName == GET_MEMBER_NAME_CHECKED(UTG_Expression_Gradient, Interpolation))
+	{
+		bEditCondition = (int32)Type <= (int32)EGradientType::GT_Linear_2;
+	}
+	else if (PropertyName == GET_MEMBER_NAME_CHECKED(UTG_Expression_Gradient, Center) || 
+		PropertyName == GET_MEMBER_NAME_CHECKED(UTG_Expression_Gradient, Radius))
+	{
+		bEditCondition = Type == EGradientType::GT_Radial;
+	}
+	else if (PropertyName == GET_MEMBER_NAME_CHECKED(UTG_Expression_Gradient, Point1) || 
+		PropertyName == GET_MEMBER_NAME_CHECKED(UTG_Expression_Gradient, Point2))
+	{
+		bEditCondition = (Type == EGradientType::GT_Axial_1 || Type == EGradientType::GT_Axial_2);
+	}
+
+	// Default behaviour
+	return bEditCondition;
+}
+#endif
\ No newline at end of file
diff --git a/Engine/Plugins/Experimental/TextureGraph/Source/TextureGraph/Public/Expressions/Procedural/TG_Expression_Gradient.h b/Engine/Plugins/Experimental/TextureGraph/Source/TextureGraph/Public/Expressions/Procedural/TG_Expression_Gradient.h
new file mode 100644
--- /dev/null
+++ b/Engine/Plugins/Experimental/TextureGraph/Source/TextureGraph/Public/Expressions/Procedural/TG_Expression_Gradient.h
@@ -0,0 +1,62 @@
+// Copyright Epic Games, Inc. All Rights Reserved.
+
+
+#pragma once
+
+#include "CoreMinimal.h"
+
+#include "Expressions/TG_Expression.h"
+#include "TG_Texture.h"
+#include "Transform/Mask/T_Gradient.h"
+
+#include "TG_Expression_Gradient.generated.h"
+
+UCLASS()
+class TEXTUREGRAPH_API UTG_Expression_Gradient : public UTG_Expression
+{
+	GENERATED_BODY()
+
+public:
+	TG_DECLARE_EXPRESSION(TG_Category::Procedural)
+	
+	virtual void						Evaluate(FTG_EvaluationContext* InContext) override;
+
+#if WITH_EDITOR
+	// Used to implement EditCondition logic for both Node UI and Details View
+	virtual bool						CanEditChange(const FProperty* InProperty) const override;
+#endif
+
+
+	// The type of the gradient function
+	UPROPERTY(EditAnywhere, BlueprintReadWrite, Category = NoCategory, meta = (TGType = "TG_Setting", RegenPinsOnChange))
+	EGradientType						Type = EGradientType::GT_Linear_1;
+
+	// Type of interpolation to use
+	UPROPERTY(EditAnywhere, BlueprintReadWrite, Category = NoCategory, meta = (TGType = "TG_Setting", EditConditionHides))
+	EGradientInterpolation				Interpolation = EGradientInterpolation::GTI_Linear;
+
+	// Type of interpolation to use
+	UPROPERTY(EditAnywhere, BlueprintReadWrite, Category = NoCategory, meta = (TGType = "TG_Setting", EditConditionHides))
+	EGradientRotation					Rotation = EGradientRotation::GTR_0;
+
+	UPROPERTY(EditAnywhere, BlueprintReadWrite, Category = NoCategory, meta = (TGType = "TG_Setting", DisplayName = "Rotation", EditConditionHides, PinDisplayName = "Rotation"))
+	EGradientRotationLimited			RotationLimited = EGradientRotationLimited::GTRL_0;
+
+	UPROPERTY(EditAnywhere, BlueprintReadWrite, Category = NoCategory, meta = (TGType = "TG_Input", EditConditionHides, UIMin = "0", ClampMin = "0", UIMax = "1", ClampMax = "1"))
+	FVector2f							Center = { 0.5f, 0.5f };
+
+	UPROPERTY(EditAnywhere, BlueprintReadWrite, Category = NoCategory, meta = (TGType = "TG_Input", EditConditionHides, UIMin = "0.001", ClampMin = "0.001", UIMax = "1", ClampMax = "1"))
+	float								Radius = 0.25f;
+
+	UPROPERTY(EditAnywhere, BlueprintReadWrite, Category = NoCategory, meta = (TGType = "TG_Input", EditConditionHides, UIMin = "0", ClampMin = "0", UIMax = "1", ClampMax = "1"))
+	FVector2f							Point1 = { 0.25f, 0.25f };
+
+	UPROPERTY(EditAnywhere, BlueprintReadWrite, Category = NoCategory, meta = (TGType = "TG_Input", EditConditionHides, UIMin = "0", ClampMin = "0", UIMax = "1", ClampMax = "1"))
+	FVector2f							Point2 = { 0.75f, 0.75f };
+
+	// The generated gradient texture
+	UPROPERTY(EditAnywhere, Category = NoCategory, meta = (TGType = "TG_Output", PinDisplayName = ""))
+	FTG_Texture							Output;
+	
+	virtual FText						GetTooltipText() const override { return FText::FromString(TEXT("Generates different types of gradients.")); } 
+};
diff --git a/Engine/Plugins/Experimental/TextureGraph/Source/TextureGraphEngine/Helper/GraphicsUtil.h b/Engine/Plugins/Experimental/TextureGraph/Source/TextureGraphEngine/Helper/GraphicsUtil.h
--- a/Engine/Plugins/Experimental/TextureGraph/Source/TextureGraphEngine/Helper/GraphicsUtil.h
+++ b/Engine/Plugins/Experimental/TextureGraph/Source/TextureGraphEngine/Helper/GraphicsUtil.h
@@ -44,6 +44,7 @@ public:
 	static FORCEINLINE Vector4	ConvertTangent(const Tangent& t) { return Vector4(t.TangentX.X, t.TangentX.Y, t.TangentX.Z, t.bFlipTangentY ? -1.0f : 1.0f); }
 	static FORCEINLINE FLinearColor TangentAsLinearColor(const Tangent& t) { return FLinearColor(ConvertTangent(t)); }
 	static FORCEINLINE FLinearColor Vec2AsLinearColor(const Vector2& v) { return FLinearColor(v.X, v.Y, 0.0f, 0.0f); }
+	static FORCEINLINE FLinearColor Vec2AsLinearColor(const FVector2f& v) { return FLinearColor(v.X, v.Y, 0.0f, 0.0f); }
 	static FORCEINLINE FLinearColor Vec3AsLinearColor(const Vector3& v) { return FLinearColor(v.X, v.Y, v.Z, 0.0f); }
 	static FORCEINLINE FLinearColor Vec4AsLinearColor(const Vector4& v) { return FLinearColor(v); }
 
diff --git a/Engine/Plugins/Experimental/TextureGraph/Source/TextureGraphEngine/Transform/Mask/T_Gradient.cpp b/Engine/Plugins/Experimental/TextureGraph/Source/TextureGraphEngine/Transform/Mask/T_Gradient.cpp
new file mode 100644
--- /dev/null
+++ b/Engine/Plugins/Experimental/TextureGraph/Source/TextureGraphEngine/Transform/Mask/T_Gradient.cpp
@@ -0,0 +1,179 @@
+// Copyright Epic Games, Inc. All Rights Reserved.
+#include "T_Gradient.h"
+#include "Job/JobArgs.h"
+#include "TextureGraphEngine.h"
+#include "2D/TargetTextureSet.h"
+#include "3D/RenderMesh.h"
+#include "FxMat/MaterialManager.h"
+#include "Job/JobBatch.h"
+#include "Helper/GraphicsUtil.h"
+#include "Model/Mix/MixInterface.h"
+#include "Model/Mix/MixSettings.h"
+#include "T_Gradient.h"
+
+IMPLEMENT_GLOBAL_SHADER(FSH_GradientLinear_1, "/Plugin/TextureGraph/Expressions/Expression_Gradient.usf", "FSH_GradientLinear_1", SF_Pixel);
+IMPLEMENT_GLOBAL_SHADER(FSH_GradientLinear_2, "/Plugin/TextureGraph/Expressions/Expression_Gradient.usf", "FSH_GradientLinear_2", SF_Pixel);
+IMPLEMENT_GLOBAL_SHADER(FSH_GradientRadial, "/Plugin/TextureGraph/Expressions/Expression_Gradient.usf", "FSH_GradientRadial", SF_Pixel);
+IMPLEMENT_GLOBAL_SHADER(FSH_GradientAxial1, "/Plugin/TextureGraph/Expressions/Expression_Gradient.usf", "FSH_GradientAxial1", SF_Pixel);
+IMPLEMENT_GLOBAL_SHADER(FSH_GradientAxial2, "/Plugin/TextureGraph/Expressions/Expression_Gradient.usf", "FSH_GradientAxial2", SF_Pixel);
+
+TiledBlobPtr Create_Linear_1(MixUpdateCyclePtr Cycle, BufferDescriptor DesiredOutputDesc, const T_Gradient::Params& InParams, int32 TargetId)
+{
+	RenderMaterial_FXPtr RenderMaterial;
+
+	FSH_GradientLinear_1::FPermutationDomain PermutationVector;
+	PermutationVector.Set<FVar_GradientInterpolation>((int32)InParams.Interpolation);
+	PermutationVector.Set<FVar_GradientRotation>(InParams.Rotation);
+	RenderMaterial = TextureGraphEngine::GetMaterialManager()->CreateMaterial_FX<VSH_Simple, FSH_GradientLinear_1>(TEXT("T_Gradient_Linear_1"), PermutationVector);
+
+	FString Name = FString::Printf(TEXT("%s-%d.[%llu]"), *RenderMaterial->GetName(), (int32)InParams.Type, Cycle->GetBatch()->GetBatchId());
+
+	check(RenderMaterial);
+
+	JobUPtr RenderJob = std::make_unique<Job>(Cycle->GetMix(), TargetId, std::static_pointer_cast<BlobTransform>(RenderMaterial));
+	FTileInfo TileInfo;
+	RenderJob
+		->AddArg(ARG_TILEINFO(TileInfo, "TileInfo"))
+		->AddArg(WithUnbounded(ARG_INT((int32)InParams.Interpolation, "Interpolation")))
+		->AddArg(WithUnbounded(ARG_INT((int32)InParams.Rotation, "Rotation")))
+		;
+
+	TiledBlobPtr Result = RenderJob->InitResult(Name, &DesiredOutputDesc);
+	Cycle->AddJob(TargetId, std::move(RenderJob));
+
+	return Result;
+}
+
+TiledBlobPtr Create_Linear_2(MixUpdateCyclePtr Cycle, BufferDescriptor DesiredOutputDesc, const T_Gradient::Params& InParams, int32 TargetId)
+{
+	RenderMaterial_FXPtr RenderMaterial;
+
+	FSH_GradientLinear_2::FPermutationDomain PermutationVector;
+	PermutationVector.Set<FVar_GradientInterpolation>((int32)InParams.Interpolation);
+	PermutationVector.Set<FVar_GradientRotation>(InParams.Rotation);
+	RenderMaterial = TextureGraphEngine::GetMaterialManager()->CreateMaterial_FX<VSH_Simple, FSH_GradientLinear_2>(TEXT("T_Gradient_Linear_2"), PermutationVector);
+
+	FString Name = FString::Printf(TEXT("%s-%d.[%llu]"), *RenderMaterial->GetName(), (int32)InParams.Type, Cycle->GetBatch()->GetBatchId());
+
+	check(RenderMaterial);
+
+	JobUPtr RenderJob = std::make_unique<Job>(Cycle->GetMix(), TargetId, std::static_pointer_cast<BlobTransform>(RenderMaterial));
+	FTileInfo TileInfo;
+	RenderJob
+		->AddArg(ARG_TILEINFO(TileInfo, "TileInfo"))
+		->AddArg(WithUnbounded(ARG_INT((int32)InParams.Interpolation, "Interpolation")))
+		->AddArg(WithUnbounded(ARG_INT((int32)InParams.Rotation, "Rotation")))
+		;
+
+	TiledBlobPtr Result = RenderJob->InitResult(Name, &DesiredOutputDesc);
+	Cycle->AddJob(TargetId, std::move(RenderJob));
+
+	return Result;
+}
+
+TiledBlobPtr Create_Radial(MixUpdateCyclePtr Cycle, BufferDescriptor DesiredOutputDesc, const T_Gradient::Params& InParams, int32 TargetId)
+{
+	RenderMaterial_FXPtr RenderMaterial;
+
+	FSH_GradientRadial::FPermutationDomain PermutationVector;
+	PermutationVector.Set<FVar_GradientInterpolation>((int32)InParams.Interpolation);
+	RenderMaterial = TextureGraphEngine::GetMaterialManager()->CreateMaterial_FX<VSH_Simple, FSH_GradientRadial>(TEXT("T_Gradient_Radial"), PermutationVector);
+
+	FString Name = FString::Printf(TEXT("%s-%d.[%llu]"), *RenderMaterial->GetName(), (int32)InParams.Type, Cycle->GetBatch()->GetBatchId());
+
+	check(RenderMaterial);
+
+	JobUPtr RenderJob = std::make_unique<Job>(Cycle->GetMix(), TargetId, std::static_pointer_cast<BlobTransform>(RenderMaterial));
+	FTileInfo TileInfo;
+	RenderJob
+		->AddArg(ARG_TILEINFO(TileInfo, "TileInfo"))
+		->AddArg(ARG_VECTOR(FLinearColor(InParams.Center.X, InParams.Center.Y, InParams.Radius, 0), "Center"))
+		;
+
+	TiledBlobPtr Result = RenderJob->InitResult(Name, &DesiredOutputDesc);
+	Cycle->AddJob(TargetId, std::move(RenderJob));
+
+	return Result;
+}
+
+TiledBlobPtr Create_Axial(MixUpdateCyclePtr Cycle, BufferDescriptor DesiredOutputDesc, const T_Gradient::Params& InParams, int32 TargetId)
+{
+	RenderMaterial_FXPtr RenderMaterial;
+
+	FSH_GradientRadial::FPermutationDomain PermutationVector;
+	PermutationVector.Set<FVar_GradientInterpolation>((int32)InParams.Interpolation);
+
+	if (InParams.Type == EGradientType::GT_Axial_1)
+		RenderMaterial = TextureGraphEngine::GetMaterialManager()->CreateMaterial_FX<VSH_Simple, FSH_GradientAxial1>(TEXT("T_Gradient_Axial1"), PermutationVector);
+	else
+		RenderMaterial = TextureGraphEngine::GetMaterialManager()->CreateMaterial_FX<VSH_Simple, FSH_GradientAxial2>(TEXT("T_Gradient_Axial2"), PermutationVector);
+
+	FString Name = FString::Printf(TEXT("%s-%d.[%llu]"), *RenderMaterial->GetName(), (int32)InParams.Type, Cycle->GetBatch()->GetBatchId());
+
+	check(RenderMaterial);
+
+	FVector2f Line = InParams.Point2 - InParams.Point1;
+	float LineLen = Line.Length();
+	FVector2f LineDir = Line;
+	LineDir.Normalize();
+
+	/// 0 length vector
+	if (abs(LineLen) < 0.0001f)
+		return TextureHelper::GetBlack();
+
+	JobUPtr RenderJob = std::make_unique<Job>(Cycle->GetMix(), TargetId, std::static_pointer_cast<BlobTransform>(RenderMaterial));
+	FTileInfo TileInfo;
+	RenderJob
+		->AddArg(ARG_TILEINFO(TileInfo, "TileInfo"))
+		->AddArg(ARG_VECTOR(FLinearColor(InParams.Point1.X, InParams.Point1.Y, Line.X, Line.Y), "Line"))
+		->AddArg(ARG_VECTOR(FLinearColor(LineDir.X, LineDir.Y, Line.SquaredLength(), Line.Length()), "LineDir"))
+		;
+
+	TiledBlobPtr Result = RenderJob->InitResult(Name, &DesiredOutputDesc);
+	Cycle->AddJob(TargetId, std::move(RenderJob));
+
+	return Result;
+}
+
+BufferDescriptor T_Gradient::InitOutputDesc(BufferDescriptor DesiredOutputDesc)
+{
+	int32 MaxSize = std::max(DesiredOutputDesc.Width, DesiredOutputDesc.Height);
+	if (DesiredOutputDesc.Width <= 0 || DesiredOutputDesc.Height <= 0)
+	{
+		MaxSize = (MaxSize > 0 ? MaxSize : DefaultSize);
+		DesiredOutputDesc.Width = DesiredOutputDesc.Height = MaxSize;
+	}
+
+	if (DesiredOutputDesc.Format == BufferFormat::Auto)
+		DesiredOutputDesc.Format = BufferFormat::Byte;
+
+	if (DesiredOutputDesc.ItemsPerPoint == 0)
+		DesiredOutputDesc.ItemsPerPoint = 4;
+
+	DesiredOutputDesc.bIsSRGB = true;
+	DesiredOutputDesc.DefaultValue = FLinearColor(1.0f, 1.0f, 1.0f, 1.0f);
+
+	return DesiredOutputDesc;
+}
+
+TiledBlobPtr T_Gradient::Create(MixUpdateCyclePtr Cycle, BufferDescriptor DesiredOutputDesc, const Params& InParams, int32 TargetId)
+{
+	RenderMaterial_FXPtr RenderMaterial;
+
+	DesiredOutputDesc = InitOutputDesc(DesiredOutputDesc);
+
+	switch (InParams.Type) 
+	{
+	case EGradientType::GT_Linear_1:
+		return Create_Linear_1(Cycle, DesiredOutputDesc, InParams, TargetId);
+	case EGradientType::GT_Linear_2:
+		return Create_Linear_2(Cycle, DesiredOutputDesc, InParams, TargetId);
+	case EGradientType::GT_Radial:
+		return Create_Radial(Cycle, DesiredOutputDesc, InParams, TargetId);
+	case EGradientType::GT_Axial_1:
+	case EGradientType::GT_Axial_2:
+		return Create_Axial(Cycle, DesiredOutputDesc, InParams, TargetId);
+	}
+
+	return nullptr;
+}
diff --git a/Engine/Plugins/Experimental/TextureGraph/Source/TextureGraphEngine/Transform/Mask/T_Gradient.h b/Engine/Plugins/Experimental/TextureGraph/Source/TextureGraphEngine/Transform/Mask/T_Gradient.h
new file mode 100644
--- /dev/null
+++ b/Engine/Plugins/Experimental/TextureGraph/Source/TextureGraphEngine/Transform/Mask/T_Gradient.h
@@ -0,0 +1,158 @@
+// Copyright Epic Games, Inc. All Rights Reserved.
+#pragma once
+
+#include "CoreMinimal.h"
+#include "FxMat/FxMaterial.h"
+#include "Job/Job.h"
+#include "Model/Mix/MixUpdateCycle.h"
+#include "UObject/NoExportTypes.h"
+#include <DataDrivenShaderPlatformInfo.h>
+
+UENUM(BlueprintType)
+enum class EGradientType : uint8
+{
+	GT_Linear_1 = 0				UMETA(DisplayName = "Simple Linear"),
+	GT_Linear_2 = 1				UMETA(DisplayName = "Linear Centered"),
+	GT_Radial = 2				UMETA(DisplayName = "Radial"),
+	GT_Axial_1 = 3				UMETA(DisplayName = "Axial Linear"),
+	GT_Axial_2 = 4				UMETA(DisplayName = "Axial Centered"),
+};
+
+UENUM(BlueprintType)
+enum class EGradientInterpolation : uint8
+{
+	GTI_Linear = 0				UMETA(DisplayName = "Linear"),
+	GTI_Exp = 1					UMETA(DisplayName = "Exponential"),
+	//GTI_Log = 2					UMETA(DisplayName = "Logarithmic"),
+};
+
+UENUM(BlueprintType)
+enum class EGradientRotation : uint8
+{
+	GTR_0 = 0					UMETA(DisplayName = "0"),
+	GTR_90 = 1					UMETA(DisplayName = "90"),
+	GTR_180 = 2					UMETA(DisplayName = "180"),
+	GTR_270 = 3					UMETA(DisplayName = "270"),
+};
+
+UENUM(BlueprintType)
+enum class EGradientRotationLimited : uint8
+{
+	GTRL_0 = 0					UMETA(DisplayName = "0"),
+	GTRL_90 = 1					UMETA(DisplayName = "90"),
+};
+
+namespace
+{
+	// Declare one Shader Permutation Var class per parameters
+	class FVar_GradientInterpolation : SHADER_PERMUTATION_INT("GRADIENT_INTERPOLATION", (int32)EGradientInterpolation::GTI_Exp + 1);
+	class FVar_GradientRotation : SHADER_PERMUTATION_INT("GRADIENT_ROTATION", (int32)360);
+}
+
+///////////////////////////////////////////////////////////////////////////////
+class TEXTUREGRAPHENGINE_API FSH_GradientLinear_1 : public FSH_Base
+{
+public:
+	DECLARE_GLOBAL_SHADER(FSH_GradientLinear_1);
+	SHADER_USE_PARAMETER_STRUCT(FSH_GradientLinear_1, FSH_Base);
+	BEGIN_SHADER_PARAMETER_STRUCT(FParameters, )
+		SHADER_PARAMETER_STRUCT(FTileInfo, TileInfo)
+	END_SHADER_PARAMETER_STRUCT()
+
+	using FPermutationDomain = TShaderPermutationDomain<FVar_GradientInterpolation, FVar_GradientRotation>;
+public:
+	TEXTURE_ENGINE_DEFAULT_PERMUTATION;
+	TEXTUREGRAPH_ENGINE_DEFAULT_COMPILATION_ENV;
+};
+
+///////////////////////////////////////////////////////////////////////////////
+class TEXTUREGRAPHENGINE_API FSH_GradientLinear_2 : public FSH_Base
+{
+public:
+	DECLARE_GLOBAL_SHADER(FSH_GradientLinear_2);
+	SHADER_USE_PARAMETER_STRUCT(FSH_GradientLinear_2, FSH_Base);
+	BEGIN_SHADER_PARAMETER_STRUCT(FParameters, )
+		SHADER_PARAMETER_STRUCT(FTileInfo, TileInfo)
+	END_SHADER_PARAMETER_STRUCT()
+
+	using FPermutationDomain = TShaderPermutationDomain<FVar_GradientInterpolation, FVar_GradientRotation>;
+public:
+	TEXTURE_ENGINE_DEFAULT_PERMUTATION;
+	TEXTUREGRAPH_ENGINE_DEFAULT_COMPILATION_ENV;
+};
+
+///////////////////////////////////////////////////////////////////////////////
+class TEXTUREGRAPHENGINE_API FSH_GradientRadial : public FSH_Base
+{
+public:
+	DECLARE_GLOBAL_SHADER(FSH_GradientRadial);
+	SHADER_USE_PARAMETER_STRUCT(FSH_GradientRadial, FSH_Base);
+	BEGIN_SHADER_PARAMETER_STRUCT(FParameters, )
+		SHADER_PARAMETER_STRUCT(FTileInfo, TileInfo)
+		SHADER_PARAMETER(FVector4f, Center)
+	END_SHADER_PARAMETER_STRUCT()
+
+	using FPermutationDomain = TShaderPermutationDomain<FVar_GradientInterpolation>;
+public:
+	TEXTURE_ENGINE_DEFAULT_PERMUTATION;
+	TEXTUREGRAPH_ENGINE_DEFAULT_COMPILATION_ENV;
+};
+
+///////////////////////////////////////////////////////////////////////////////
+class TEXTUREGRAPHENGINE_API FSH_GradientAxial1: public FSH_Base
+{
+public:
+	DECLARE_GLOBAL_SHADER(FSH_GradientAxial1);
+	SHADER_USE_PARAMETER_STRUCT(FSH_GradientAxial1, FSH_Base);
+	BEGIN_SHADER_PARAMETER_STRUCT(FParameters, )
+		SHADER_PARAMETER_STRUCT(FTileInfo, TileInfo)
+		SHADER_PARAMETER(FVector4f, Line)
+		SHADER_PARAMETER(FVector4f, LineDir)
+	END_SHADER_PARAMETER_STRUCT()
+
+	using FPermutationDomain = TShaderPermutationDomain<FVar_GradientInterpolation>;
+public:
+	TEXTURE_ENGINE_DEFAULT_PERMUTATION;
+	TEXTUREGRAPH_ENGINE_DEFAULT_COMPILATION_ENV;
+};
+
+///////////////////////////////////////////////////////////////////////////////
+class TEXTUREGRAPHENGINE_API FSH_GradientAxial2: public FSH_Base
+{
+public:
+	DECLARE_GLOBAL_SHADER(FSH_GradientAxial2);
+	SHADER_USE_PARAMETER_STRUCT(FSH_GradientAxial2, FSH_Base);
+	BEGIN_SHADER_PARAMETER_STRUCT(FParameters, )
+		SHADER_PARAMETER_STRUCT(FTileInfo, TileInfo)
+		SHADER_PARAMETER(FVector4f, Line)
+		SHADER_PARAMETER(FVector4f, LineDir)
+	END_SHADER_PARAMETER_STRUCT()
+
+	using FPermutationDomain = TShaderPermutationDomain<FVar_GradientInterpolation>;
+public:
+	TEXTURE_ENGINE_DEFAULT_PERMUTATION;
+	TEXTUREGRAPH_ENGINE_DEFAULT_COMPILATION_ENV;
+};
+
+///////////////////////////////////////////////////////////////////////////////
+/**
+ * 
+ */
+class TEXTUREGRAPHENGINE_API T_Gradient
+{
+public:
+	static constexpr int		DefaultSize = 1024;
+	static BufferDescriptor		InitOutputDesc(BufferDescriptor DesiredOutputDesc);
+
+	struct Params {
+		EGradientType			Type = EGradientType::GT_Linear_1;
+		EGradientInterpolation	Interpolation = EGradientInterpolation::GTI_Linear;
+		int32					Rotation = 0; 
+		FVector2f				Center;
+		float					Radius;
+		FVector2f				Point1;
+		FVector2f				Point2;
+	};
+
+	static TiledBlobPtr	Create(MixUpdateCyclePtr Cycle, BufferDescriptor DesiredOutputDesc, const Params& InParams, int32 TargetId = 0);
+};
