Reduce number of permutations of the gradient shader

diff --git a/Engine/Plugins/Experimental/TextureGraph/Source/TextureGraphEngine/Transform/Mask/T_Gradient.h b/Engine/Plugins/Experimental/TextureGraph/Source/TextureGraphEngine/Transform/Mask/T_Gradient.h
--- a/Engine/Plugins/Experimental/TextureGraph/Source/TextureGraphEngine/Transform/Mask/T_Gradient.h
+++ b/Engine/Plugins/Experimental/TextureGraph/Source/TextureGraphEngine/Transform/Mask/T_Gradient.h
@@ -46,7 +46,7 @@ namespace
 {
 	// Declare one Shader Permutation Var class per parameters
 	class FVar_GradientInterpolation : SHADER_PERMUTATION_INT("GRADIENT_INTERPOLATION", (int32)EGradientInterpolation::GTI_Exp + 1);
-	class FVar_GradientRotation : SHADER_PERMUTATION_INT("GRADIENT_ROTATION", (int32)360);
+	class FVar_GradientRotation : SHADER_PERMUTATION_INT("GRADIENT_ROTATION", (int32)EGradientRotation::GTR_270 + 1);
 }
 
 ///////////////////////////////////////////////////////////////////////////////
