Edge detection filter

diff --git a/Engine/Plugins/Experimental/TextureGraph/Shaders/Expressions/Expression_EdgeDetect.usf b/Engine/Plugins/Experimental/TextureGraph/Shaders/Expressions/Expression_EdgeDetect.usf
new file mode 100644
--- /dev/null
+++ b/Engine/Plugins/Experimental/TextureGraph/Shaders/Expressions/Expression_EdgeDetect.usf
@@ -0,0 +1,56 @@
+// Copyright Epic Games, Inc. All Rights Reserved.
+#include "/Engine/Public/Platform.ush"
+#include "/Plugin/TextureGraph/SamplerStates.ush"
+#include "/Plugin/TextureGraph/TileInfo.ush"
+
+Texture2D SourceTexture;
+float Thickness;
+
+float Intensity(float4 Color)
+{
+	return sqrt((Color.x * Color.x) + (Color.y * Color.y) + (Color.z * Color.z));
+}
+
+float3 Sobel(float StepX, float StepY, float2 Center)
+{
+	// get samples around pixel
+	float TopLeft = Intensity(SourceTexture.Sample(SamplerStates_Linear_Clamp, Center + float2(-StepX, StepY)));
+	float Left = Intensity(SourceTexture.Sample(SamplerStates_Linear_Clamp, Center + float2(-StepX, 0)));
+	float BottomLeft = Intensity(SourceTexture.Sample(SamplerStates_Linear_Clamp, Center + float2(-StepX, -StepY)));
+	float Top = Intensity(SourceTexture.Sample(SamplerStates_Linear_Clamp, Center + float2(0, StepY)));
+	float Bottom = Intensity(SourceTexture.Sample(SamplerStates_Linear_Clamp, Center + float2(0, -StepY)));
+	float TopRight = Intensity(SourceTexture.Sample(SamplerStates_Linear_Clamp, Center + float2(StepX, StepY)));
+	float Right = Intensity(SourceTexture.Sample(SamplerStates_Linear_Clamp, Center + float2(StepX, 0)));
+	float BottomRight = Intensity(SourceTexture.Sample(SamplerStates_Linear_Clamp, Center + float2(StepX, -StepY)));
+ 
+	// Sobel masks (see http://en.wikipedia.org/wiki/Sobel_operator)
+	//        1 0 -1     -1 -2 -1
+	//    X = 2 0 -2  Y = 0  0  0
+	//        1 0 -1      1  2  1
+	
+	// You could also use Scharr operator:
+	//        3 0 -3        3 10   3
+	//    X = 10 0 -10  Y = 0  0   0
+	//        3 0 -3        -3 -10 -3
+ 
+	float X = TopLeft + 2.0 * Left + BottomLeft - TopRight - 2.0 * Right - BottomRight;
+	float Y = -TopLeft - 2.0 * Top - TopRight + BottomLeft + 2.0 * Bottom + BottomRight;
+	float Color = sqrt((X * X) + (Y * Y));
+	return float3(Color, Color, Color);
+}
+
+float4 FSH_EdgeDetect(float2 InUV : TEXCOORD0) : SV_Target0
+{
+	int Width = TileInfo_TileCountX * TileInfo_TileWidth;
+	int Height = TileInfo_TileCountY * TileInfo_TileHeight;
+	float WidthInv = rcp(float(Width));
+	float HeightInv = rcp(float(Height));
+	
+	float StepX = Thickness * WidthInv;
+	float StepY = Thickness * HeightInv;
+	float2 UV = TileInfo_fromCurrentTileToLayer(InUV);
+	
+	float4 Center = SourceTexture.Sample(SamplerStates_Linear_Clamp, UV);
+	float3 Result = Sobel(StepX, StepY, UV) * Center.xyz; 
+	return float4(Result.x, Result.y, Result.z, Center.a);
+}
diff --git a/Engine/Plugins/Experimental/TextureGraph/Source/TextureGraph/Private/Expressions/FIlter/TG_Expression_EdgeDetect.cpp b/Engine/Plugins/Experimental/TextureGraph/Source/TextureGraph/Private/Expressions/FIlter/TG_Expression_EdgeDetect.cpp
new file mode 100644
--- /dev/null
+++ b/Engine/Plugins/Experimental/TextureGraph/Source/TextureGraph/Private/Expressions/FIlter/TG_Expression_EdgeDetect.cpp
@@ -0,0 +1,18 @@
+// Copyright Epic Games, Inc. All Rights Reserved.
+
+#include "Expressions/Filter/TG_Expression_EdgeDetect.h"
+
+#include "Transform/Expressions/T_EdgeDetect.h"
+
+void UTG_Expression_EdgeDetect::Evaluate(FTG_EvaluationContext* InContext)
+{
+	Super::Evaluate(InContext);
+
+	if (!Input)
+	{
+		Output = FTG_Texture::GetBlack();
+		return;
+	}
+
+	Output = T_EdgeDetect::Create(InContext->Cycle, Output.GetBufferDescriptor(), Input, Thickness, InContext->TargetId);
+}
diff --git a/Engine/Plugins/Experimental/TextureGraph/Source/TextureGraph/Public/Expressions/Filter/TG_Expression_EdgeDetect.h b/Engine/Plugins/Experimental/TextureGraph/Source/TextureGraph/Public/Expressions/Filter/TG_Expression_EdgeDetect.h
new file mode 100644
--- /dev/null
+++ b/Engine/Plugins/Experimental/TextureGraph/Source/TextureGraph/Public/Expressions/Filter/TG_Expression_EdgeDetect.h
@@ -0,0 +1,33 @@
+// Copyright Epic Games, Inc. All Rights Reserved.
+
+#pragma once
+
+#include "CoreMinimal.h"
+#include "Expressions/TG_Expression_MaterialBase.h"
+
+#include "TG_Expression_EdgeDetect.generated.h"
+
+UCLASS()
+class TEXTUREGRAPH_API UTG_Expression_EdgeDetect : public UTG_Expression
+{
+	GENERATED_BODY()
+
+public:
+	TG_DECLARE_EXPRESSION(TG_Category::Filter);
+	
+	// What is the thickness of the edge
+	UPROPERTY(EditAnywhere, BlueprintReadWrite, Category = NoCategory, meta = (TGType = "TG_Setting", ClampMin = 0.1, ClampMax = 30, UIMin = 0.1, UIMax = 1, EditConditionHides))
+	float Thickness = 1;
+
+	// The output texture having blurred effect
+	UPROPERTY(EditAnywhere, Category = NoCategory, meta = (TGType = "TG_Output", PinDisplayName = ""))
+	FTG_Texture Output;
+
+	// The input texture to apply the blur effect
+	UPROPERTY(meta = (TGType = "TG_Input"))
+	FTG_Texture Input;
+
+	virtual FText GetTooltipText() const override { return FText::FromString(TEXT("Applies a blur filter with a given strength.")); }
+	virtual void Evaluate(FTG_EvaluationContext* InContext) override;
+};
+
diff --git a/Engine/Plugins/Experimental/TextureGraph/Source/TextureGraphEngine/Transform/Expressions/T_EdgeDetect.cpp b/Engine/Plugins/Experimental/TextureGraph/Source/TextureGraphEngine/Transform/Expressions/T_EdgeDetect.cpp
new file mode 100644
--- /dev/null
+++ b/Engine/Plugins/Experimental/TextureGraph/Source/TextureGraphEngine/Transform/Expressions/T_EdgeDetect.cpp
@@ -0,0 +1,39 @@
+// Copyright Epic Games, Inc. All Rights Reserved.
+
+#include "T_EdgeDetect.h"
+#include "FxMat/MaterialManager.h"
+#include "Job/JobBatch.h"
+#include "Transform/Utility/T_CombineTiledBlob.h"
+
+IMPLEMENT_GLOBAL_SHADER(FSH_EdgeDetect, "/Plugin/TextureGraph/Expressions/Expression_EdgeDetect.usf", "FSH_EdgeDetect", SF_Pixel);
+
+TiledBlobPtr T_EdgeDetect::Create(MixUpdateCyclePtr InCycle, BufferDescriptor DesiredDesc, TiledBlobPtr SourceTexture, float Thickness, int32 InTargetId)
+{
+	if (!SourceTexture)
+		return TextureHelper::GetBlack();
+
+	RenderMaterial_FXPtr RenderMaterial = TextureGraphEngine::GetMaterialManager()->CreateMaterial_FX<VSH_Simple, FSH_EdgeDetect>(TEXT("T_EdgeDetect"));
+	check(RenderMaterial);
+
+	// Combine the tiled texture
+	TiledBlobPtr CombinedBlob = T_CombineTiledBlob::Create(InCycle, SourceTexture->GetDescriptor(), InTargetId, SourceTexture);
+	float StepX = Thickness / (float)SourceTexture->GetWidth();
+	float StepY = Thickness / (float)SourceTexture->GetHeight();
+
+	BufferDescriptor Desc = BufferDescriptor::Combine(DesiredDesc, SourceTexture->GetDescriptor());
+
+	FTileInfo TileInfo;
+	JobUPtr RenderJob = std::make_unique<Job>(InCycle->GetMix(), InTargetId, std::static_pointer_cast<BlobTransform>(RenderMaterial));
+	RenderJob
+		->AddArg(ARG_TILEINFO(TileInfo, "TileInfo"))
+		->AddArg(ARG_BLOB(CombinedBlob, "SourceTexture"))
+		->AddArg(ARG_FLOAT(Thickness, "Thickness"))
+		;
+
+	const FString Name = FString::Printf(TEXT("T_EdgeDetect.[%llu]"), InCycle->GetBatch()->GetBatchId());
+
+	TiledBlobPtr Result = RenderJob->InitResult(Name, &Desc);
+	InCycle->AddJob(InTargetId, std::move(RenderJob));
+
+	return Result;
+}
diff --git a/Engine/Plugins/Experimental/TextureGraph/Source/TextureGraphEngine/Transform/Expressions/T_EdgeDetect.h b/Engine/Plugins/Experimental/TextureGraph/Source/TextureGraphEngine/Transform/Expressions/T_EdgeDetect.h
new file mode 100644
--- /dev/null
+++ b/Engine/Plugins/Experimental/TextureGraph/Source/TextureGraphEngine/Transform/Expressions/T_EdgeDetect.h
@@ -0,0 +1,39 @@
+// Copyright Epic Games, Inc. All Rights Reserved.
+#pragma once
+
+#include <2D/BlendModes.h>
+
+#include "CoreMinimal.h"
+#include "Job/Job.h"
+#include "FxMat/FxMaterial.h"
+#include "Model/Mix/MixUpdateCycle.h"
+#include <DataDrivenShaderPlatformInfo.h>
+
+//////////////////////////////////////////////////////////////////////////
+/// Gaussian
+//////////////////////////////////////////////////////////////////////////
+class TEXTUREGRAPHENGINE_API FSH_EdgeDetect : public FSH_Base
+{
+public:
+	DECLARE_GLOBAL_SHADER(FSH_EdgeDetect);
+	SHADER_USE_PARAMETER_STRUCT(FSH_EdgeDetect, FSH_Base);
+	BEGIN_SHADER_PARAMETER_STRUCT(FParameters, )
+		SHADER_PARAMETER_STRUCT(FTileInfo, TileInfo)
+		SHADER_PARAMETER_STRUCT(FStandardSamplerStates, SamplerStates)
+		SHADER_PARAMETER_TEXTURE(Texture2D, SourceTexture)
+		SHADER_PARAMETER(float, Thickness)
+	END_SHADER_PARAMETER_STRUCT()
+
+public:
+	TEXTURE_ENGINE_DEFAULT_PERMUTATION;
+	TEXTUREGRAPH_ENGINE_DEFAULT_COMPILATION_ENV;
+};
+
+class TEXTUREGRAPHENGINE_API T_EdgeDetect
+{
+public:
+	//////////////////////////////////////////////////////////////////////////
+	/// Static functions
+	//////////////////////////////////////////////////////////////////////////
+	static TiledBlobPtr				Create(MixUpdateCyclePtr InCycle, BufferDescriptor DesiredDesc, TiledBlobPtr SourceTexture, float Thickness, int32 InTargetId);
+};
