Viewport setting updates

#virtualized

diff --git a/Engine/Plugins/Experimental/TextureGraph/Source/TextureGraphEditor/Private/Customizations/TG_ViewportSettingsCustomization.h b/Engine/Plugins/Experimental/TextureGraph/Source/TextureGraphEditor/Private/Customizations/TG_ViewportSettingsCustomization.h
new file mode 100644
--- /dev/null
+++ b/Engine/Plugins/Experimental/TextureGraph/Source/TextureGraphEditor/Private/Customizations/TG_ViewportSettingsCustomization.h
@@ -0,0 +1,76 @@
+// Copyright Epic Games, Inc. All Rights Reserved.
+
+#pragma once
+
+#include "CoreMinimal.h"
+
+#include "TG_Parameter.h"
+#include "Modules/ModuleManager.h"
+#include "PropertyEditorModule.h"
+#include "IPropertyTypeCustomization.h"
+//#include "IPropertyUtilities.h"
+//#include "PropertyHandle.h"
+#include "DetailWidgetRow.h"
+//#include "DetailLayoutBuilder.h"
+//#include "ISinglePropertyView.h"
+#include "IDetailChildrenBuilder.h"
+//#include "IDetailCustomization.h"
+#include "IDetailGroup.h"
+#include "PropertyCustomizationHelpers.h"
+#include "TG_Pin.h"
+#include "Expressions/Output/TG_Expression_Output.h"
+#include "TG_Node.h"
+#include "TG_Graph.h"
+#include "TextureGraph.h"
+#include "Model/Mix/ViewportSettings.h"
+#include "Widgets/Layout/SBox.h"
+#include "Widgets/Text/STextBlock.h"
+
+#define LOCTEXT_NAMESPACE "FTextureGraphEditorModule"
+
+class FTG_ViewportSettingsCustomization : public IPropertyTypeCustomization
+{
+
+private:
+	UMixSettings* MixSettings;
+	
+	TSharedPtr<IPropertyHandle> MaterialPropertyHandle;
+	TSharedPtr<IPropertyHandle> MaterialMappingInfosPropertyHandle;
+	
+public:
+	static TSharedRef<IPropertyTypeCustomization> Create()
+	{
+		return MakeShareable(new FTG_ViewportSettingsCustomization);
+	}
+
+	virtual void CustomizeHeader(TSharedRef<IPropertyHandle> PropertyHandle, FDetailWidgetRow& HeaderRow, IPropertyTypeCustomizationUtils& CustomizationUtils) override
+	{
+		MaterialPropertyHandle = PropertyHandle->GetChildHandle(GET_MEMBER_NAME_CHECKED(FViewportSettings, Material));
+		MaterialMappingInfosPropertyHandle = PropertyHandle->GetChildHandle(GET_MEMBER_NAME_CHECKED(FViewportSettings, MaterialMappingInfos));
+	}
+
+	virtual void CustomizeChildren(TSharedRef<IPropertyHandle> PropertyHandle, IDetailChildrenBuilder& ChildBuilder, IPropertyTypeCustomizationUtils& CustomizationUtils) override
+	{
+		FPropertyEditorModule& PropertyModule = FModuleManager::LoadModuleChecked<FPropertyEditorModule>("PropertyEditor");
+		
+		if (PropertyHandle->IsValidHandle())
+		{
+			ChildBuilder.AddProperty(MaterialPropertyHandle.ToSharedRef());
+
+			// In this case, we'll add the array elements directly without the header for the array name.
+			uint32 NumChildren;
+			MaterialMappingInfosPropertyHandle->GetNumChildren(NumChildren);
+
+			for (uint32 ChildIndex = 0; ChildIndex < NumChildren; ++ChildIndex)
+			{
+				TSharedRef<IPropertyHandle> ChildHandle = MaterialMappingInfosPropertyHandle->GetChildHandle(ChildIndex).ToSharedRef();
+				if (ChildHandle->IsValidHandle())
+				{
+					ChildBuilder.AddProperty(ChildHandle);
+				}
+			}
+		}
+	}
+};
+
+#undef LOCTEXT_NAMESPACE
diff --git a/Engine/Plugins/Experimental/TextureGraph/Source/TextureGraphEditor/Private/STG_EditorViewport.cpp b/Engine/Plugins/Experimental/TextureGraph/Source/TextureGraphEditor/Private/STG_EditorViewport.cpp
--- a/Engine/Plugins/Experimental/TextureGraph/Source/TextureGraphEditor/Private/STG_EditorViewport.cpp
+++ b/Engine/Plugins/Experimental/TextureGraph/Source/TextureGraphEditor/Private/STG_EditorViewport.cpp
@@ -469,7 +469,10 @@ void STG_EditorViewport::UpdateRenderMode()
 
 void STG_EditorViewport::InitPreviewMesh()
 {
-	SetPreviewAsset(GUnrealEd->GetThumbnailManager()->EditorCube);
+	UStaticMesh* Primitive = GUnrealEd->GetThumbnailManager()->EditorCube;
+	SetPreviewAsset(Primitive);
+	TG_EditorPtr.Pin().Get()->GetTextureGraphInterface()->GetSettings()->SetPreviewMesh(Primitive);	
+			
 }
 
 bool STG_EditorViewport::IsPreviewPrimitiveChecked(EThumbnailPrimType PrimType) const
@@ -493,7 +496,7 @@ void STG_EditorViewport::OnSetPreviewPrimitive(EThumbnailPrimType PrimType, bool
 		if (Primitive != nullptr)
 		{
 			SetPreviewAsset(Primitive);
-			TG_EditorPtr.Pin().Get()->GetTextureGraphInterface()->PreviewMesh = Primitive->GetPathName();	
+			TG_EditorPtr.Pin().Get()->GetTextureGraphInterface()->GetSettings()->SetPreviewMesh(Primitive);	
 			
 			RefreshViewport();
 		}
@@ -502,7 +505,6 @@ void STG_EditorViewport::OnSetPreviewPrimitive(EThumbnailPrimType PrimType, bool
 
 void STG_EditorViewport::OnSetPreviewMeshFromSelection()
 {
-	//TODO: Need to enable the code. 
 	bool bFoundPreviewMesh = false;
 	FEditorDelegates::LoadSelectedAssetsIfNeeded.Broadcast();
 
@@ -519,7 +521,7 @@ void STG_EditorViewport::OnSetPreviewMeshFromSelection()
 				if (ComponentClass->IsChildOf(UMeshComponent::StaticClass()))
 				{
 					SetPreviewAsset(TestAsset);
-					TG_Interface->PreviewMesh = TestAsset->GetPathName();
+					TG_Interface->GetSettings()->SetPreviewMesh(Cast<UStaticMesh>(TestAsset));
 					bFoundPreviewMesh = true;
 				}
 			}
@@ -528,9 +530,6 @@ void STG_EditorViewport::OnSetPreviewMeshFromSelection()
 
 	if (bFoundPreviewMesh)
 	{
-		//FTG_Editor::UpdateThumbnailInfoPreviewMesh(TG_Interface);
-
-		//TG_Interface->MarkPackageDirty();
 		RefreshViewport();
 	}
 	else
diff --git a/Engine/Plugins/Experimental/TextureGraph/Source/TextureGraphEditor/Private/TG_Editor.cpp b/Engine/Plugins/Experimental/TextureGraph/Source/TextureGraphEditor/Private/TG_Editor.cpp
--- a/Engine/Plugins/Experimental/TextureGraph/Source/TextureGraphEditor/Private/TG_Editor.cpp
+++ b/Engine/Plugins/Experimental/TextureGraph/Source/TextureGraphEditor/Private/TG_Editor.cpp
@@ -206,8 +206,9 @@ void FTG_Editor::InitEditor(const EToolkitMode::Type Mode, const TSharedPtr< cla
 
 	FViewportSettings& ViewportSettings = EditedTextureGraph->GetSettings()->GetViewportSettings();
 
-	ViewportSettings.OnViewportMaterialChangeEvent.AddSP(this, &FTG_Editor::OnViewportSettingsChanged);
+	ViewportSettings.OnViewportMaterialChangedEvent.AddSP(this, &FTG_Editor::OnViewportSettingsChanged);
 	ViewportSettings.OnMaterialMappingChangedEvent.AddSP(this, &FTG_Editor::OnMaterialMappingChanged);
+	EditedTextureGraph->GetSettings()->OnPreviewMeshChangedEvent.AddSP(this, &FTG_Editor::OnPreviewMeshChangedEvent);
 
 	TextureGraphEngine::RegisterErrorReporter(EditedTextureGraph, std::make_shared<FTG_EditorErrorReporter>(this));
 
@@ -369,7 +370,7 @@ void FTG_Editor::InitEditor(const EToolkitMode::Type Mode, const TSharedPtr< cla
 	}
 	
 	// Set the preview mesh for the material.  This call must occur after the toolbar is initialized.
-	if (!SetPreviewAssetByName(*EditedTextureGraph->PreviewMesh.ToString()))
+	if (!SetPreviewAssetByName(*EditedTextureGraph->GetSettings()->GetPreviewMesh().GetPath()))
 	{
 		// The material preview mesh couldn't be found or isn't loaded.  Default to the one of the primitive types.
 		GetEditorViewport()->InitPreviewMesh();
@@ -672,6 +673,11 @@ void FTG_Editor::OnRenderingDone(UMixInterface* TextureGraph, const FInvalidatio
 	}
 }
 
+void FTG_Editor::OnPreviewMeshChangedEvent()
+{
+	SetPreviewAsset(EditedTextureGraph->GetSettings()->GetPreviewMesh());
+}
+
 void FTG_Editor::OnViewportSettingsChanged()
 {
 	const UTG_Node* FirstTargetNode = nullptr;
@@ -1001,7 +1007,7 @@ TSharedRef<SDockTab> FTG_Editor::SpawnTab_Viewport(const FSpawnTabArgs& Args)
 	ViewportTabContent->Initialize(MakeViewportFunc, DockableTab, LayoutId);
 	
 	// Set the preview mesh for the material.  This call must occur after the toolbar is initialized.
-	if (!SetPreviewAssetByName(*EditedTextureGraph->PreviewMesh.ToString()))
+	if (!SetPreviewAssetByName(*EditedTextureGraph->GetSettings()->GetPreviewMesh().GetPath()))
 	{
 		// The material preview mesh couldn't be found or isn't loaded.  Default to the one of the primitive types.
 		GetEditorViewport()->InitPreviewMesh();
diff --git a/Engine/Plugins/Experimental/TextureGraph/Source/TextureGraphEditor/Private/TG_Editor.h b/Engine/Plugins/Experimental/TextureGraph/Source/TextureGraphEditor/Private/TG_Editor.h
--- a/Engine/Plugins/Experimental/TextureGraph/Source/TextureGraphEditor/Private/TG_Editor.h
+++ b/Engine/Plugins/Experimental/TextureGraph/Source/TextureGraphEditor/Private/TG_Editor.h
@@ -169,6 +169,7 @@ private:
 	bool											GetTabSelected(const FName TabID);
 
 	void											OnRenderingDone(UMixInterface* TextureGraph, const FInvalidationDetails* Details);
+	void											OnPreviewMeshChangedEvent();
 	void											OnViewportSettingsChanged();
 	void											OnMaterialMappingChanged();
 	/** Log Graph in console */
diff --git a/Engine/Plugins/Experimental/TextureGraph/Source/TextureGraphEditor/Private/TextureGraphEditorModule.cpp b/Engine/Plugins/Experimental/TextureGraph/Source/TextureGraphEditor/Private/TextureGraphEditorModule.cpp
--- a/Engine/Plugins/Experimental/TextureGraph/Source/TextureGraphEditor/Private/TextureGraphEditorModule.cpp
+++ b/Engine/Plugins/Experimental/TextureGraph/Source/TextureGraphEditor/Private/TextureGraphEditorModule.cpp
@@ -16,6 +16,7 @@
 #include "Customizations/TG_VariantCustomization.h"
 #include "Customizations/TG_ScalarCustomization.h"
 #include "Customizations/TG_MaterialMappingInfoCustomization.h"
+#include "Customizations/TG_ViewportSettingsCustomization.h"
 #include "Customizations/TG_OutputSettingsCustomization.h"
 #include "Customizations/TG_LevelsSettingsCustomization.h"
 
@@ -37,6 +38,7 @@ void FTextureGraphEditorModule::StartupModule()
 	auto ScalarIdentifier = MakeShared<FTG_ScalarTypeIdentifier>();
 	PropertyEditorModule.RegisterCustomPropertyTypeLayout("FloatProperty", FOnGetPropertyTypeCustomizationInstance::CreateStatic(&FTG_ScalarCustomization::Create), ScalarIdentifier);
 	PropertyEditorModule.RegisterCustomPropertyTypeLayout("MaterialMappingInfo", FOnGetPropertyTypeCustomizationInstance::CreateStatic(&FTG_MaterialMappingInfoCustomization::Create));
+	PropertyEditorModule.RegisterCustomPropertyTypeLayout("ViewportSettings", FOnGetPropertyTypeCustomizationInstance::CreateStatic(&FTG_ViewportSettingsCustomization::Create));
 	PropertyEditorModule.RegisterCustomPropertyTypeLayout("OutputSettings", FOnGetPropertyTypeCustomizationInstance::CreateStatic(&FTG_OutputSettingsCustomization::Create));
 	PropertyEditorModule.RegisterCustomPropertyTypeLayout("TG_LevelsSettings", FOnGetPropertyTypeCustomizationInstance::CreateStatic(&FTG_LevelsSettingsCustomization::Create));
 
diff --git a/Engine/Plugins/Experimental/TextureGraph/Source/TextureGraphEngine/Model/Mix/MixInterface.h b/Engine/Plugins/Experimental/TextureGraph/Source/TextureGraphEngine/Model/Mix/MixInterface.h
--- a/Engine/Plugins/Experimental/TextureGraph/Source/TextureGraphEngine/Model/Mix/MixInterface.h
+++ b/Engine/Plugins/Experimental/TextureGraph/Source/TextureGraphEngine/Model/Mix/MixInterface.h
@@ -60,9 +60,6 @@ public:
 	virtual								~UMixInterface() override;
 
 
-	/** The mesh used by the TS editor to preview.*/
-	UPROPERTY(EditAnywhere, Category = Previewing, meta = (AllowedClasses = "/Script/Engine.StaticMesh,/Script/Engine.SkeletalMesh", ExactClass = "true"))
-	FSoftObjectPath						PreviewMesh;  // Adding it here because this will be the base class for TS_Script and TS_ScriptInstance.
 
 	DECLARE_DELEGATE_TwoParams(FOnRenderDone, UMixInterface*, const FInvalidationDetails*);
 	FOnRenderDone						OnRenderDone;
diff --git a/Engine/Plugins/Experimental/TextureGraph/Source/TextureGraphEngine/Model/Mix/MixSettings.cpp b/Engine/Plugins/Experimental/TextureGraph/Source/TextureGraphEngine/Model/Mix/MixSettings.cpp
--- a/Engine/Plugins/Experimental/TextureGraph/Source/TextureGraphEngine/Model/Mix/MixSettings.cpp
+++ b/Engine/Plugins/Experimental/TextureGraph/Source/TextureGraphEngine/Model/Mix/MixSettings.cpp
@@ -127,11 +127,16 @@ void UMixSettings::PostEditChangeProperty(FPropertyChangedEvent& PropertyChanged
 
 	if (MemberPropertyName == GET_MEMBER_NAME_CHECKED(UMixSettings, ViewportSettings))
 	{
-		if(Name == "Material")
+		if(Name == GET_MEMBER_NAME_CHECKED(FViewportSettings, Material))
 		{
 			ViewportSettings.OnMaterialUpdate();	
 		}
 	}
+	if (MemberPropertyName == GET_MEMBER_NAME_CHECKED(UMixSettings, PreviewMesh))
+	{
+		OnPreviewMeshChangedEvent.Broadcast();
+	}
+	
 }
 
 void UMixSettings::PostEditChangeChainProperty(FPropertyChangedChainEvent& PropertyChangedEvent)
diff --git a/Engine/Plugins/Experimental/TextureGraph/Source/TextureGraphEngine/Model/Mix/MixSettings.h b/Engine/Plugins/Experimental/TextureGraph/Source/TextureGraphEngine/Model/Mix/MixSettings.h
--- a/Engine/Plugins/Experimental/TextureGraph/Source/TextureGraphEngine/Model/Mix/MixSettings.h
+++ b/Engine/Plugins/Experimental/TextureGraph/Source/TextureGraphEngine/Model/Mix/MixSettings.h
@@ -94,7 +94,7 @@ class RenderMaterial_BP;
 typedef std::shared_ptr<RenderMaterial_BP>	RenderMaterial_BPPtr;
 
 class UMixInterface;
-
+class UStaticMesh;
 //////////////////////////////////////////////////////////////////////////
 /// MixSettings: Collectively represents the settings of a particular
 /// mix. These can be specified on a per mix basis. However, they can 
@@ -152,9 +152,13 @@ PRAGMA_ENABLE_DEPRECATION_WARNINGS
 	UPROPERTY()
 	bool							bAllowRectangularResolution = false;		// Allow rectangular outputs
 
-	UPROPERTY(EditAnywhere, Category = NoCategory, meta = (NoResetToDefault))
+	UPROPERTY(EditAnywhere, Category = "Viewport Material", meta = (NoResetToDefault, ShowOnlyInnerProperties))
 	FViewportSettings				ViewportSettings;							// Viewport settings
 	
+	UPROPERTY(EditAnywhere, Category = "Viewport Mesh", DisplayName = "Static Mesh", NoClear)
+	TObjectPtr<UStaticMesh>			PreviewMesh;
+
+private:
 	/// Mesh related data 
 	RenderMeshPtr					_mesh;										/// The full mesh (along with all its parts) that was loaded 
 	TArray<RenderMeshPtr>			_sceneMeshes;								/// All the different submeshes that were added to the scene
@@ -166,7 +170,7 @@ PRAGMA_ENABLE_DEPRECATION_WARNINGS
 	template<typename RenderMeshClass>
 	void							SetMeshInternal(RenderMeshPtr mesh, int meshType, FVector scale, FVector2D dimension);
 	void							SetMesh(RenderMeshPtr mesh);
-
+	
 protected:
 
 #if WITH_EDITOR
@@ -192,9 +196,14 @@ public:
 	
 	FViewportSettings&				GetViewportSettings();
 
+	DECLARE_MULTICAST_DELEGATE(FPreviewMeshUpdateEvent)
+	FPreviewMeshUpdateEvent OnPreviewMeshChangedEvent;
+	
 	//////////////////////////////////////////////////////////////////////////
 	/// Inline functions
 	//////////////////////////////////////////////////////////////////////////
+	[[nodiscard]] TObjectPtr<UStaticMesh> GetPreviewMesh() const{	return PreviewMesh; }
+	FORCEINLINE void				SetPreviewMesh(const TObjectPtr<UStaticMesh>& InPreviewMesh)	{	this->PreviewMesh = InPreviewMesh; }
 	FORCEINLINE bool				IsPlane() const { return bIsPlane; }
 	FORCEINLINE FVector2D			GetPlaneDimensions() const { return PlaneDimensions; }
 	FORCEINLINE const TargetTextureSetPtr& Target(size_t index) const { verify(index < _targets->size()); return (*_targets)[index]; }
diff --git a/Engine/Plugins/Experimental/TextureGraph/Source/TextureGraphEngine/Model/Mix/ViewportSettings.cpp b/Engine/Plugins/Experimental/TextureGraph/Source/TextureGraphEngine/Model/Mix/ViewportSettings.cpp
--- a/Engine/Plugins/Experimental/TextureGraph/Source/TextureGraphEngine/Model/Mix/ViewportSettings.cpp
+++ b/Engine/Plugins/Experimental/TextureGraph/Source/TextureGraphEngine/Model/Mix/ViewportSettings.cpp
@@ -38,7 +38,7 @@ void FViewportSettings::SetDefaultTarget(FName DefaultTargetName)
 
 UMaterial* FViewportSettings::GetDefaultMaterial()
 {
-	UMaterial* LoadedMaterial = TextureGraphEngine::GetMaterialManager()->LoadMaterial(TEXT("Scene/DefaultMaterial"));
+	UMaterial* LoadedMaterial = TextureGraphEngine::GetMaterialManager()->LoadMaterial(TEXT("Scene/M_BasicMaterial"));
 
 	check(LoadedMaterial)
 
@@ -66,6 +66,7 @@ void FViewportSettings::InitDefaultSettings(FName InitialTargetName)
 		OnMaterialUpdate();
 		SetDefaultTarget(InitialTargetName);
 	}
+
 }
 
 bool FViewportSettings::RemoveMaterialMappingForTarget(const FName OutputNode)
@@ -111,7 +112,7 @@ void FViewportSettings::OnMaterialUpdate()
 			}
 		}
 
-		OnViewportMaterialChangeEvent.Broadcast();
+		OnViewportMaterialChangedEvent.Broadcast();
 	}
 }
 
diff --git a/Engine/Plugins/Experimental/TextureGraph/Source/TextureGraphEngine/Model/Mix/ViewportSettings.h b/Engine/Plugins/Experimental/TextureGraph/Source/TextureGraphEngine/Model/Mix/ViewportSettings.h
--- a/Engine/Plugins/Experimental/TextureGraph/Source/TextureGraphEngine/Model/Mix/ViewportSettings.h
+++ b/Engine/Plugins/Experimental/TextureGraph/Source/TextureGraphEngine/Model/Mix/ViewportSettings.h
@@ -25,15 +25,15 @@ struct TEXTUREGRAPHENGINE_API FViewportSettings
 {
 	GENERATED_BODY()
 	
-	UPROPERTY(EditAnywhere, Category = "Viewport Settings", NoClear)
+	UPROPERTY(EditAnywhere, Category = "Viewport Material", NoClear)
 	TObjectPtr<UMaterial> Material;
 	
-	UPROPERTY(EditAnywhere, EditFixedSize, Category = "Viewport Settings")
+	UPROPERTY(EditAnywhere, EditFixedSize, Category = "Viewport Material")
 	TArray<FMaterialMappingInfo> MaterialMappingInfos;
-
+	
 	void InitDefaultSettings(FName InitialTargetName);
 	void SetDefaultTarget(FName DefaultTargetName);
-
+	
 	UMaterial* GetDefaultMaterial();
 	
 	FName GetMaterialName() const;
@@ -47,7 +47,7 @@ struct TEXTUREGRAPHENGINE_API FViewportSettings
 	int  NumAssignedTargets();
 	
 	DECLARE_MULTICAST_DELEGATE(FViewportSettingsUpdateEvent)
-	FViewportSettingsUpdateEvent OnViewportMaterialChangeEvent;
+	FViewportSettingsUpdateEvent OnViewportMaterialChangedEvent;
 
 	DECLARE_MULTICAST_DELEGATE(FMaterialMappingChangedEvent)
 	FMaterialMappingChangedEvent OnMaterialMappingChangedEvent;
