Fixed some unused and potentially uninitialized FVector2Ds.

diff --git a/Engine/Plugins/Experimental/TextureGraph/Source/TextureGraphEditor/Private/TG_Editor.cpp b/Engine/Plugins/Experimental/TextureGraph/Source/TextureGraphEditor/Private/TG_Editor.cpp
--- a/Engine/Plugins/Experimental/TextureGraph/Source/TextureGraphEditor/Private/TG_Editor.cpp
+++ b/Engine/Plugins/Experimental/TextureGraph/Source/TextureGraphEditor/Private/TG_Editor.cpp
@@ -783,7 +783,7 @@ void FTG_Editor::OnCreateComment()
 		FTG_EdGraphSchemaAction_NewComment CommentAction;
 
 		TSharedPtr<SGraphEditor> GraphEditorPtr = SGraphEditor::FindGraphEditorForGraph(TG_EdGraph);
-		FVector2D Location;
+		FVector2D Location = FVector2D::ZeroVector;
 		if (GraphEditorPtr)
 		{
 			Location = GraphEditorPtr->GetPasteLocation();
