Fixed potentially uninitialized variables.

diff --git a/Engine/Plugins/Experimental/TextureGraph/Source/TextureGraph/Private/Blueprint/TG_BlueprintFunctionLibrary.cpp b/Engine/Plugins/Experimental/TextureGraph/Source/TextureGraph/Private/Blueprint/TG_BlueprintFunctionLibrary.cpp
--- a/Engine/Plugins/Experimental/TextureGraph/Source/TextureGraph/Private/Blueprint/TG_BlueprintFunctionLibrary.cpp
+++ b/Engine/Plugins/Experimental/TextureGraph/Source/TextureGraph/Private/Blueprint/TG_BlueprintFunctionLibrary.cpp
@@ -158,7 +158,7 @@ void UTG_BlueprintFunctionLibrary::SetVectorParameterValue(UObject* WorldContext
 
 FLinearColor UTG_BlueprintFunctionLibrary::GetVectorParameterValue(UObject* WorldContextObject, UTextureGraph* InTextureGraph, FName ParameterName)
 {
-	FLinearColor ParameterValue;
+	FLinearColor ParameterValue = FLinearColor::Black;
 
 	if (UWorld* World = GEngine->GetWorldFromContextObject(WorldContextObject, EGetWorldErrorMode::LogAndReturnNull))
 	{
@@ -210,7 +210,7 @@ void UTG_BlueprintFunctionLibrary::SetColorParameterValue(UObject* WorldContextO
 
 FLinearColor UTG_BlueprintFunctionLibrary::GetColorParameterValue(UObject* WorldContextObject, UTextureGraph* InTextureGraph, FName ParameterName)
 {
-	FLinearColor ParameterValue;
+	FLinearColor ParameterValue = FLinearColor::Black;
 
 	if (UWorld* World = GEngine->GetWorldFromContextObject(WorldContextObject, EGetWorldErrorMode::LogAndReturnNull))
 	{
diff --git a/Engine/Source/Editor/PropertyEditor/Private/UserInterface/PropertyEditor/SPropertyEditorColor.cpp b/Engine/Source/Editor/PropertyEditor/Private/UserInterface/PropertyEditor/SPropertyEditorColor.cpp
--- a/Engine/Source/Editor/PropertyEditor/Private/UserInterface/PropertyEditor/SPropertyEditorColor.cpp
+++ b/Engine/Source/Editor/PropertyEditor/Private/UserInterface/PropertyEditor/SPropertyEditorColor.cpp
@@ -127,7 +127,7 @@ FLinearColor SPropertyEditorColor::OnGetColor() const
 	const FProperty* Property = PropertyEditor->GetProperty();
 	check(Property);
 
-	FLinearColor Color;
+	FLinearColor Color = FLinearColor::Black;
 
 	FReadAddressList ReadAddresses;
 	PropertyNode->GetReadAddress( false, ReadAddresses, false );
diff --git a/Engine/Source/Editor/UnrealEd/Private/ThumbnailRendering/WorldThumbnailRenderer.cpp b/Engine/Source/Editor/UnrealEd/Private/ThumbnailRendering/WorldThumbnailRenderer.cpp
--- a/Engine/Source/Editor/UnrealEd/Private/ThumbnailRendering/WorldThumbnailRenderer.cpp
+++ b/Engine/Source/Editor/UnrealEd/Private/ThumbnailRendering/WorldThumbnailRenderer.cpp
@@ -143,8 +143,8 @@ FSceneView* UWorldThumbnailRenderer::CreateView(UWorld* World, FSceneViewFamily*
 	}
 
 	const FVector Origin = WorldBox.GetCenter();
-	FMatrix ViewRotationMatrix;
-	FMatrix ProjectionMatrix;
+	FMatrix ViewRotationMatrix = FMatrix::Identity;
+	FMatrix ProjectionMatrix = FMatrix::Identity;
 	float FOVScreenSize = 0; // Screen size taking FOV into account
 	if (ThumbnailInfo->CameraMode == ECameraProjectionMode::Perspective)
 	{
diff --git a/Engine/Source/Editor/UnrealEd/Private/Tools/LegacyEdModeWidgetHelpers.cpp b/Engine/Source/Editor/UnrealEd/Private/Tools/LegacyEdModeWidgetHelpers.cpp
--- a/Engine/Source/Editor/UnrealEd/Private/Tools/LegacyEdModeWidgetHelpers.cpp
+++ b/Engine/Source/Editor/UnrealEd/Private/Tools/LegacyEdModeWidgetHelpers.cpp
@@ -138,7 +138,7 @@ namespace
 	template<typename T>
 	T GetPropertyValueByName(UObject* Object, FString PropertyName, int32 PropertyIndex)
 	{
-		T Value;
+		T Value = {};
 		FProperty* DummyProperty = NULL;
 		if (T* ValuePtr = GetPropertyValuePtrByName<T>(Object->GetClass(), Object, PropertyName, PropertyIndex, DummyProperty))
 		{
