[TextureGraphEngine]
* Changed absolute path in define MODULE_DIR to be relative to Engine dir and then expand it at runtime. This is needed to make Definitions.h deterministic between machines

diff --git a/Engine/Plugins/Experimental/TextureGraph/Source/TextureGraphEngine/TextureGraphEngine.Build.cs b/Engine/Plugins/Experimental/TextureGraph/Source/TextureGraphEngine/TextureGraphEngine.Build.cs
--- a/Engine/Plugins/Experimental/TextureGraph/Source/TextureGraphEngine/TextureGraphEngine.Build.cs
+++ b/Engine/Plugins/Experimental/TextureGraph/Source/TextureGraphEngine/TextureGraphEngine.Build.cs
@@ -82,9 +82,14 @@ public class TextureGraphEngine : ModuleRules
 
 		AddDefaultIncludePaths();
 
-		string ModDirLiteral = ModuleDirectory.Replace('\\', '/');
+		if (!ModuleDirectory.StartsWith(EngineDirectory))
+			throw new BuildException("TextureGraphEngine module directory must be under engine");
+
+		string ModuleRelativeToEngineDir = ModuleDirectory.Substring(EngineDirectory.Length + 1);
+
+		string ModDirLiteral = ModuleRelativeToEngineDir.Replace('\\', '/');
 		string defModuleName = "MODULE_DIR \"" + ModDirLiteral + "\"=";
-		PublicDefinitions.Add(defModuleName);
+		PrivateDefinitions.Add(defModuleName);
 	}
 
 }
diff --git a/Engine/Plugins/Experimental/TextureGraph/Source/TextureGraphEngine/TextureGraphEngine.cpp b/Engine/Plugins/Experimental/TextureGraph/Source/TextureGraphEngine/TextureGraphEngine.cpp
--- a/Engine/Plugins/Experimental/TextureGraph/Source/TextureGraphEngine/TextureGraphEngine.cpp
+++ b/Engine/Plugins/Experimental/TextureGraph/Source/TextureGraphEngine/TextureGraphEngine.cpp
@@ -35,7 +35,7 @@ void FTextureGraphEngineModule::StartupModule()
 bool FTextureGraphEngineModule::MapShaders()
 {
 	// This code will execute after your module is loaded into memory; the exact timing is specified in the .uplugin file per-module
-	FString fullModuleDir = MODULE_DIR;
+	FString fullModuleDir = FPaths::Combine(FPaths::EngineDir(), MODULE_DIR);
 	bool added = false;
 
 	FString rightSanitized;
