3D Preview Settings fixes
- 3D Preview switches Output preview after adjusting settings in the node
- Editor crashes when setting custom preview mesh in the 3D Preview Settings
- Tests fail with Failed to find object 'Object None.None'
- Selecting any material in 3D Preview settings does not update the list of inputs based on specific Material's parameters

diff --git a/Engine/Plugins/Experimental/TextureGraph/Source/TextureGraphEditor/Private/Customizations/TG_ViewportSettingsCustomization.h b/Engine/Plugins/Experimental/TextureGraph/Source/TextureGraphEditor/Private/Customizations/TG_ViewportSettingsCustomization.h
--- a/Engine/Plugins/Experimental/TextureGraph/Source/TextureGraphEditor/Private/Customizations/TG_ViewportSettingsCustomization.h
+++ b/Engine/Plugins/Experimental/TextureGraph/Source/TextureGraphEditor/Private/Customizations/TG_ViewportSettingsCustomization.h
@@ -4,27 +4,15 @@
 
 #include "CoreMinimal.h"
 
-#include "TG_Parameter.h"
 #include "Modules/ModuleManager.h"
 #include "PropertyEditorModule.h"
+#include "PropertyHandle.h"
 #include "IPropertyTypeCustomization.h"
-//#include "IPropertyUtilities.h"
-//#include "PropertyHandle.h"
-#include "DetailWidgetRow.h"
-//#include "DetailLayoutBuilder.h"
-//#include "ISinglePropertyView.h"
+#include "IPropertyUtilities.h"
 #include "IDetailChildrenBuilder.h"
-//#include "IDetailCustomization.h"
-#include "IDetailGroup.h"
-#include "PropertyCustomizationHelpers.h"
-#include "TG_Pin.h"
-#include "Expressions/Output/TG_Expression_Output.h"
-#include "TG_Node.h"
-#include "TG_Graph.h"
-#include "TextureGraph.h"
+
 #include "Model/Mix/ViewportSettings.h"
-#include "Widgets/Layout/SBox.h"
-#include "Widgets/Text/STextBlock.h"
+#include "Model/Mix/MixSettings.h"
 
 #define LOCTEXT_NAMESPACE "FTextureGraphEditorModule"
 
@@ -52,11 +40,15 @@ public:
 	virtual void CustomizeChildren(TSharedRef<IPropertyHandle> PropertyHandle, IDetailChildrenBuilder& ChildBuilder, IPropertyTypeCustomizationUtils& CustomizationUtils) override
 	{
 		FPropertyEditorModule& PropertyModule = FModuleManager::LoadModuleChecked<FPropertyEditorModule>("PropertyEditor");
-		
+		const TSharedPtr<IPropertyUtilities> PropertyUtils = CustomizationUtils.GetPropertyUtilities();
 		if (PropertyHandle->IsValidHandle())
 		{
 			ChildBuilder.AddProperty(MaterialPropertyHandle.ToSharedRef());
-
+			MaterialPropertyHandle->SetOnPropertyValueChanged(FSimpleDelegate::CreateLambda([=, this]()
+			{
+				// refresh details
+				PropertyUtils->ForceRefresh();
+			}));
 			// In this case, we'll add the array elements directly without the header for the array name.
 			uint32 NumChildren;
 			MaterialMappingInfosPropertyHandle->GetNumChildren(NumChildren);
diff --git a/Engine/Plugins/Experimental/TextureGraph/Source/TextureGraphEditor/Private/TG_Editor.cpp b/Engine/Plugins/Experimental/TextureGraph/Source/TextureGraphEditor/Private/TG_Editor.cpp
--- a/Engine/Plugins/Experimental/TextureGraph/Source/TextureGraphEditor/Private/TG_Editor.cpp
+++ b/Engine/Plugins/Experimental/TextureGraph/Source/TextureGraphEditor/Private/TG_Editor.cpp
@@ -206,7 +206,7 @@ void FTG_Editor::InitEditor(const EToolkitMode::Type Mode, const TSharedPtr< cla
 
 	FViewportSettings& ViewportSettings = EditedTextureGraph->GetSettings()->GetViewportSettings();
 
-	ViewportSettings.OnViewportMaterialChangedEvent.AddSP(this, &FTG_Editor::OnViewportSettingsChanged);
+	ViewportSettings.OnViewportMaterialChangedEvent.AddSP(this, &FTG_Editor::OnViewportMaterialChanged);
 	ViewportSettings.OnMaterialMappingChangedEvent.AddSP(this, &FTG_Editor::OnMaterialMappingChanged);
 	EditedTextureGraph->GetSettings()->OnPreviewMeshChangedEvent.AddSP(this, &FTG_Editor::OnPreviewMeshChangedEvent);
 
@@ -370,11 +370,7 @@ void FTG_Editor::InitEditor(const EToolkitMode::Type Mode, const TSharedPtr< cla
 	}
 	
 	// Set the preview mesh for the material.  This call must occur after the toolbar is initialized.
-	if (!SetPreviewAssetByName(*EditedTextureGraph->GetSettings()->GetPreviewMesh().GetPath()))
-	{
-		// The material preview mesh couldn't be found or isn't loaded.  Default to the one of the primitive types.
-		GetEditorViewport()->InitPreviewMesh();
-	}
+	SetViewportPreviewMesh();
 
 	if (bViewportIsOff)
 	{
@@ -675,10 +671,10 @@ void FTG_Editor::OnRenderingDone(UMixInterface* TextureGraph, const FInvalidatio
 
 void FTG_Editor::OnPreviewMeshChangedEvent()
 {
-	SetPreviewAsset(EditedTextureGraph->GetSettings()->GetPreviewMesh());
+	SetViewportPreviewMesh();
 }
 
-void FTG_Editor::OnViewportSettingsChanged()
+void FTG_Editor::OnViewportMaterialChanged()
 {
 	const UTG_Node* FirstTargetNode = nullptr;
 	EditedTextureGraph->Graph()->ForEachNodes([&](const UTG_Node* node, uint32 index)
@@ -985,6 +981,17 @@ TSharedRef<SDockTab> FTG_Editor::SpawnTab_GraphEditor(const FSpawnTabArgs& Args)
 		];
 }
 
+void FTG_Editor::SetViewportPreviewMesh()
+{
+	TObjectPtr<UStaticMesh> PreviewMesh = EditedTextureGraph->GetSettings()->GetPreviewMesh();
+	// Set the preview mesh for the material.  
+	if (!PreviewMesh || !SetPreviewAsset(PreviewMesh))
+	{
+		// The material preview mesh couldn't be found or isn't loaded. Fallback to the one of the primitive types.
+		GetEditorViewport()->InitPreviewMesh();
+	}
+}
+
 TSharedRef<SDockTab> FTG_Editor::SpawnTab_Viewport(const FSpawnTabArgs& Args)
 {
 	check(Args.GetTabId() == FTG_EditorTabs::ViewportTabId);
@@ -1005,14 +1012,10 @@ TSharedRef<SDockTab> FTG_Editor::SpawnTab_Viewport(const FSpawnTabArgs& Args)
 
 	const FString LayoutId = FString("TG_EditorViewport");
 	ViewportTabContent->Initialize(MakeViewportFunc, DockableTab, LayoutId);
+
+	// This call must occur after the toolbar is initialized.
+	SetViewportPreviewMesh();
 	
-	// Set the preview mesh for the material.  This call must occur after the toolbar is initialized.
-	if (!SetPreviewAssetByName(*EditedTextureGraph->GetSettings()->GetPreviewMesh().GetPath()))
-	{
-		// The material preview mesh couldn't be found or isn't loaded.  Default to the one of the primitive types.
-		GetEditorViewport()->InitPreviewMesh();
-	
-	}
 	return DockableTab;
 }
 
@@ -2017,15 +2020,6 @@ bool FTG_Editor::SetPreviewAsset(UObject* InAsset)
 	return false;
 }
 
-bool FTG_Editor::SetPreviewAssetByName(const TCHAR* InAssetName)
-{
-	if (GetEditorViewport().IsValid())
-	{
-		return GetEditorViewport()->SetPreviewAssetByName(InAssetName);
-	}
-	return false;
-}
-
 FString FTG_Editor::GetCurrentFolderInContentBrowser()
 {
 	IContentBrowserSingleton& ContentBrowser = FModuleManager::LoadModuleChecked<FContentBrowserModule>("ContentBrowser").Get();
diff --git a/Engine/Plugins/Experimental/TextureGraph/Source/TextureGraphEditor/Private/TG_Editor.h b/Engine/Plugins/Experimental/TextureGraph/Source/TextureGraphEditor/Private/TG_Editor.h
--- a/Engine/Plugins/Experimental/TextureGraph/Source/TextureGraphEditor/Private/TG_Editor.h
+++ b/Engine/Plugins/Experimental/TextureGraph/Source/TextureGraphEditor/Private/TG_Editor.h
@@ -107,8 +107,7 @@ public:
 
 	void											SetMesh(class UMeshComponent* InPreviewMesh, class UWorld* InWorld) override;
 	bool 											SetPreviewAsset(UObject* InAsset);
-	bool 											SetPreviewAssetByName(const TCHAR* InAssetName);
-
+	
 	/** Force Refresh Details View **/
 	void											RefreshDetailsView() const;
 
@@ -170,7 +169,7 @@ private:
 
 	void											OnRenderingDone(UMixInterface* TextureGraph, const FInvalidationDetails* Details);
 	void											OnPreviewMeshChangedEvent();
-	void											OnViewportSettingsChanged();
+	void											OnViewportMaterialChanged();
 	void											OnMaterialMappingChanged();
 	/** Log Graph in console */
 	void											OnLogGraph_Clicked();
@@ -274,6 +273,7 @@ private:
 	FReply											OnExportClick();
 
 	TSharedPtr<class STG_EditorViewport>			GetEditorViewport() const;	
+	void 											SetViewportPreviewMesh();
 	
 	bool											CanEnableOnRun();
 
diff --git a/Engine/Plugins/Experimental/TextureGraph/Source/TextureGraphEngine/2D/TargetTextureSet.cpp b/Engine/Plugins/Experimental/TextureGraph/Source/TextureGraphEngine/2D/TargetTextureSet.cpp
--- a/Engine/Plugins/Experimental/TextureGraph/Source/TextureGraphEngine/2D/TargetTextureSet.cpp
+++ b/Engine/Plugins/Experimental/TextureGraph/Source/TextureGraphEngine/2D/TargetTextureSet.cpp
@@ -36,8 +36,12 @@ AsyncBufferResultPtr TargetTextureSet::BindTo(RenderMaterial_BPPtr Material, TAr
 {
 	RenderCount++;
 
+	// clear previous callbacks
+	Callbacks.Empty();
+	
 	for(const FMaterialMappingInfo MappingInfo : MaterialMappingInfo)
 	{
+		
 		BindOnTextureUpdate(Material, MappingInfo);
 	}
 			
