Enabling render doc batch captures in texture graph

diff --git a/Engine/Plugins/Experimental/TextureGraph/Source/TextureGraphEngine/Profiling/RenderDoc/RenderDocManager.cpp b/Engine/Plugins/Experimental/TextureGraph/Source/TextureGraphEngine/Profiling/RenderDoc/RenderDocManager.cpp
--- a/Engine/Plugins/Experimental/TextureGraph/Source/TextureGraphEngine/Profiling/RenderDoc/RenderDocManager.cpp
+++ b/Engine/Plugins/Experimental/TextureGraph/Source/TextureGraphEngine/Profiling/RenderDoc/RenderDocManager.cpp
@@ -7,10 +7,15 @@
 
 #include "TextureGraphEngine.h"
 #include "RenderingThread.h"
-//#include "RenderDocPluginModule.h"
 
 DEFINE_LOG_CATEGORY(LogRenderDocTextureGraph);
 
+/**
+ * Enable (or disable) concretely renderdoc capture in texture graph module without changing the api
+ * DO NOT FORGET to run the project with the command line -AttachRenderDoc OR enable auto attach in the renderdoc plugin
+ */
+#define TEXTUREGRAPH_RENDERDOC_ENABLED 1
+
 namespace TextureGraphEditor
 {
 
@@ -20,8 +25,7 @@ namespace TextureGraphEditor
 	///////////////////////////////////////////////////////////////////////////////////////////////////////////////////////////////////////////////
 	RenderDocManager::RenderDocManager()
 	{
-		// Disabling for the time being
-#if 0
+#if TEXTUREGRAPH_RENDERDOC_ENABLED
 		static FAutoConsoleCommand CCmdRenderDocCaptureFrame = FAutoConsoleCommand(
 			TEXT("renderdoc.TextureGraph_CaptureNextBatch"),
 			TEXT("Captures the next Job Batch and launches RenderDoc"),
@@ -43,37 +47,32 @@ namespace TextureGraphEditor
 
 	void RenderDocManager::CaptureNextBatch()
 	{
-		// Disabling for the time being
-#if 0
+#if TEXTUREGRAPH_RENDERDOC_ENABLED
 		TextureGraphEngine::GetScheduler()->SetCaptureRenderDocNextBatch(true);
 #endif 
 	}
 
 	void RenderDocManager::CapturePreviousBatch()
 	{
-		// Disabling for the time being
-#if 0
+#if TEXTUREGRAPH_RENDERDOC_ENABLED
 		TextureGraphEngine::GetScheduler()->CaptureRenderDocLastRunBatch();
 #endif 
 	}
 
 	void RenderDocManager::BeginCapture()
 	{
-		// Disabling for the time being
-#if 0
+#if TEXTUREGRAPH_RENDERDOC_ENABLED
 		ENQUEUE_RENDER_COMMAND(BeginCaptureCommand)([this](FRHICommandListImmediate& RHICommandList)
 			{
 				IRenderDocPlugin& PluginModule = FModuleManager::GetModuleChecked<IRenderDocPlugin>("RenderDocPlugin");
 				PluginModule.BeginCapture(&RHICommandList, IRenderCaptureProvider::ECaptureFlags_Launch, FString("TextureGraph"));
 			});
-		/*PluginModule.CaptureFrame(nullptr, IRenderCaptureProvider::ECaptureFlags_Launch, FString());*/
 #endif 
 	}
 
 	void RenderDocManager::EndCapture()
 	{
-		// Disabling for the time being
-#if 0
+#if TEXTUREGRAPH_RENDERDOC_ENABLED
 		ENQUEUE_RENDER_COMMAND(EndnCaptureCommand)([this](FRHICommandListImmediate& RHICommandList)
 			{
 				IRenderDocPlugin& PluginModule = FModuleManager::GetModuleChecked<IRenderDocPlugin>("RenderDocPlugin");
diff --git a/Engine/Plugins/Experimental/TextureGraph/Source/TextureGraphEngine/Profiling/RenderDoc/RenderDocManager.h b/Engine/Plugins/Experimental/TextureGraph/Source/TextureGraphEngine/Profiling/RenderDoc/RenderDocManager.h
--- a/Engine/Plugins/Experimental/TextureGraph/Source/TextureGraphEngine/Profiling/RenderDoc/RenderDocManager.h
+++ b/Engine/Plugins/Experimental/TextureGraph/Source/TextureGraphEngine/Profiling/RenderDoc/RenderDocManager.h
@@ -3,9 +3,6 @@
 
 #include "CoreMinimal.h"
 #include <memory>
-/**
- * 
- */
 
 DECLARE_LOG_CATEGORY_EXTERN(LogRenderDocTextureGraph, Log, All);
 
diff --git a/Engine/Plugins/Experimental/TextureGraph/Source/TextureGraphEngine/TextureGraphEngine.cpp b/Engine/Plugins/Experimental/TextureGraph/Source/TextureGraphEngine/TextureGraphEngine.cpp
--- a/Engine/Plugins/Experimental/TextureGraph/Source/TextureGraphEngine/TextureGraphEngine.cpp
+++ b/Engine/Plugins/Experimental/TextureGraph/Source/TextureGraphEngine/TextureGraphEngine.cpp
@@ -139,9 +139,8 @@ void TextureGraphEngine::InitEngineInternal()
 	/// Init some of the stock textures that we keep
 	//TextureHelper::InitStockTextures();
 
-#if 0 // Disabling for the time being
+	// Allocate the render doc manager (wether or not render doc define is enabled)
 	RenderDocMgrObj = std::make_unique<TextureGraphEditor::RenderDocManager>();
-#endif 
 
 	MixMgrObj = std::make_unique<::MixManager>();
 
@@ -187,9 +186,8 @@ TextureGraphEngine::~TextureGraphEngine()
 	/// This is to control the order of destruction
 	TextureHelper::FreeStockTextures();
 
-#if 0 // Disabling for the time being
+	// No more render doc manager
 	RenderDocMgrObj = nullptr;
-#endif 
 
 	MaterialMgrObj = nullptr;
 	MixMgrObj = nullptr;
