Minor bug fixes around Enum display values in nodes and debug visualizations in exports
#tests 6630d847f62d8f7360ef137b,

diff --git a/Engine/Plugins/Experimental/TextureGraph/Source/TextureGraph/Private/Expressions/Procedural/TG_Expression_Transform.cpp b/Engine/Plugins/Experimental/TextureGraph/Source/TextureGraph/Private/Expressions/Procedural/TG_Expression_Transform.cpp
--- a/Engine/Plugins/Experimental/TextureGraph/Source/TextureGraph/Private/Expressions/Procedural/TG_Expression_Transform.cpp
+++ b/Engine/Plugins/Experimental/TextureGraph/Source/TextureGraph/Private/Expressions/Procedural/TG_Expression_Transform.cpp
@@ -30,6 +30,16 @@ void UTG_Expression_Transform::Evaluate(FTG_EvaluationContext* InContext)
 	FVector2f OutCoverage = Coverage;
 	FVector2f OutOffset = Offset;
 
+	// hide debug grid when exporting
+	if (InContext->Cycle->GetDetails().bExporting)
+	{
+		// show warning to user that debug grid is being turned off
+		auto ErrorType = static_cast<int32>(ETextureGraphErrorType::NODE_WARNING);
+		const FString ErrorMsg = FString::Format(TEXT("Export command issued, turning 'Debug Grid' off in {0} node."), {GetTitleName().ToString()} );
+		TextureGraphEngine::GetErrorReporter(InContext->Cycle->GetMix())->ReportWarning(ErrorType, ErrorMsg, GetParentNode());
+		ShowDebugGrid = 0.0f;
+	}
+	
 	T_Transform::TransformParameter XformParam{
 		.Coverage = OutCoverage,
 		.Translation = OutOffset,
@@ -53,7 +63,7 @@ void UTG_Expression_Transform::Evaluate(FTG_EvaluationContext* InContext)
 		.MirrorY = MirrorY,
 		.ShowDebugGrid = ShowDebugGrid
 	};
-
+	
 	Output = T_Transform::Create(InContext->Cycle, DesiredDescriptor, Input,
 		XformParam, CellParam, ColorParam, InContext->TargetId);
 }
diff --git a/Engine/Plugins/Experimental/TextureGraph/Source/TextureGraph/Private/TG_Pin.cpp b/Engine/Plugins/Experimental/TextureGraph/Source/TextureGraph/Private/TG_Pin.cpp
--- a/Engine/Plugins/Experimental/TextureGraph/Source/TextureGraph/Private/TG_Pin.cpp
+++ b/Engine/Plugins/Experimental/TextureGraph/Source/TextureGraph/Private/TG_Pin.cpp
@@ -253,7 +253,7 @@ FString UTG_Pin::GetEvaluatedVarValue() const
 		UEnum* Enum = ByteProperty->GetIntPropertyEnum();
 		if (Enum)
 		{
-			CurrentValue = Enum->GetNameByValue(FCString::Atoi(*CurrentValue)).ToString();
+			CurrentValue = Enum->GetNameStringByIndex(FCString::Atoi(*CurrentValue));
 		}
 	}
 	FEnumProperty* EnumProperty = CastField<FEnumProperty>(Property);
@@ -262,7 +262,7 @@ FString UTG_Pin::GetEvaluatedVarValue() const
 		UEnum* Enum = EnumProperty->GetEnum();
 		if (Enum)
 		{
-			CurrentValue = Enum->GetNameByValue(FCString::Atoi(*CurrentValue)).ToString();
+			CurrentValue = Enum->GetNameStringByIndex(FCString::Atoi(*CurrentValue));
 		}
 	}
 
