Warp filter

diff --git a/Engine/Plugins/Experimental/TextureGraph/Shaders/Expressions/Expression_Warp.usf b/Engine/Plugins/Experimental/TextureGraph/Shaders/Expressions/Expression_Warp.usf
new file mode 100644
--- /dev/null
+++ b/Engine/Plugins/Experimental/TextureGraph/Shaders/Expressions/Expression_Warp.usf
@@ -0,0 +1,79 @@
+// Copyright Epic Games, Inc. All Rights Reserved.
+#include "/Engine/Public/Platform.ush"
+#include "/Plugin/TextureGraph/SamplerStates.ush"
+#include "/Plugin/TextureGraph/TileInfo.ush"
+
+Texture2D SourceTexture;
+Texture2D Mask;
+float Intensity;
+float AngleRad;
+float PhaseU;
+float PhaseV;
+
+#define HALF_PI 1.5707963267948966
+#define QUARTER_PI HALF_PI * 0.5
+
+float2 SineWave(float2 p, float tx, float ty)
+{
+	//float tx = 0.3477;
+	//float ty = 0.7812;
+    // wave distortion
+	float x = sin(25.0 * p.y + 30.0 * p.x + 6.28 * tx) * 0.05;
+	float y = sin(25.0 * p.y + 30.0 * p.x + 6.28 * ty) * 0.05;
+	return float2(p.x + x, p.y + y);
+}
+
+float4 FSH_DirectionalWarp(float2 InUV_ : TEXCOORD0) : SV_Target0
+{
+	float2 InUV = TileInfo_fromCurrentTileToLayer(InUV_);
+	float MaskValue = Mask.Sample(SamplerStates_Linear_Clamp, InUV).r;
+	
+	/// We premultiply the intensity with mask. Whites will get maximum intensity
+	/// and blacks will get none essentially
+	float IntensityValue = Intensity * MaskValue * 0.1;
+	
+	/// We need to calculate the UV that will get warped to the current UV. So we need
+	/// to move in the opposite direction of the angle to find the UV that will end
+	/// replacing the current pixel. Hence we add 180 degrees to it (or half PI in radians)
+	float AngleValue = AngleRad + HALF_PI;
+	float C = cos(AngleValue);
+	float S = sin(AngleValue);
+	float2 Dir = float2(C, S);
+	
+	float2 UVOffset = Dir * float2(IntensityValue, IntensityValue) + InUV;
+	return SourceTexture.Sample(SamplerStates_Linear_Clamp, UVOffset);
+}
+
+float4 FSH_NormalWarp(float2 InUV_ : TEXCOORD0) : SV_Target0
+{
+	float2 InUV = TileInfo_fromCurrentTileToLayer(InUV_);
+	float MaskValue = Mask.Sample(SamplerStates_Linear_Clamp, InUV).r;
+	
+	/// We premultiply the intensity with mask. Whites will get maximum intensity
+	/// and blacks will get none essentially
+	float IntensityValue = Intensity * MaskValue * 0.1;
+	
+	/// We need to calculate the UV that will get warped to the current UV. So we need
+	/// to move in the opposite direction of the angle to find the UV that will end
+	/// replacing the current pixel. Hence we add 180 degrees to it (or half PI in radians)
+	float AngleValue = MaskValue * QUARTER_PI + HALF_PI;
+	float C = cos(AngleValue);
+	float S = sin(AngleValue);
+	float2 Dir = float2(C, S);
+	
+	float2 UVOffset = Dir * float2(IntensityValue, IntensityValue) + InUV;
+	return SourceTexture.Sample(SamplerStates_Linear_Clamp, UVOffset);
+}
+
+float4 FSH_SineWarp(float2 InUV_ : TEXCOORD0) : SV_Target0
+{
+	float2 InUV = TileInfo_fromCurrentTileToLayer(InUV_);
+	float MaskValue = Mask.Sample(SamplerStates_Linear_Clamp, InUV).r;
+	
+	/// We premultiply the intensity with mask. Whites will get maximum intensity
+	/// and blacks will get none essentially
+	float IntensityValue = Intensity * MaskValue * 0.1;
+	float2 SineUV = SineWave(InUV, PhaseU, PhaseV);
+	float2 UVOffset = float2(IntensityValue, IntensityValue) * SineUV;
+	return SourceTexture.Sample(SamplerStates_Linear_Clamp, SineUV);
+}
diff --git a/Engine/Plugins/Experimental/TextureGraph/Source/TextureGraph/Private/Expressions/FIlter/TG_Expression_EdgeDetect.cpp b/Engine/Plugins/Experimental/TextureGraph/Source/TextureGraph/Private/Expressions/FIlter/TG_Expression_EdgeDetect.cpp
--- a/Engine/Plugins/Experimental/TextureGraph/Source/TextureGraph/Private/Expressions/FIlter/TG_Expression_EdgeDetect.cpp
+++ b/Engine/Plugins/Experimental/TextureGraph/Source/TextureGraph/Private/Expressions/FIlter/TG_Expression_EdgeDetect.cpp
@@ -1,8 +1,7 @@
 // Copyright Epic Games, Inc. All Rights Reserved.
 
 #include "Expressions/Filter/TG_Expression_EdgeDetect.h"
-
-#include "Transform/Expressions/T_EdgeDetect.h"
+#include "Transform/Expressions/T_Filter.h"
 
 void UTG_Expression_EdgeDetect::Evaluate(FTG_EvaluationContext* InContext)
 {
@@ -14,5 +13,5 @@ void UTG_Expression_EdgeDetect::Evaluate(FTG_EvaluationContext* InContext)
 		return;
 	}
 
-	Output = T_EdgeDetect::Create(InContext->Cycle, Output.GetBufferDescriptor(), Input, Thickness, InContext->TargetId);
+	Output = T_Filter::CreateEdgeDetect(InContext->Cycle, Output.GetBufferDescriptor(), Input, Thickness, InContext->TargetId);
 }
diff --git a/Engine/Plugins/Experimental/TextureGraph/Source/TextureGraph/Private/Expressions/FIlter/TG_Expression_Warp.cpp b/Engine/Plugins/Experimental/TextureGraph/Source/TextureGraph/Private/Expressions/FIlter/TG_Expression_Warp.cpp
new file mode 100644
--- /dev/null
+++ b/Engine/Plugins/Experimental/TextureGraph/Source/TextureGraph/Private/Expressions/FIlter/TG_Expression_Warp.cpp
@@ -0,0 +1,60 @@
+// Copyright Epic Games, Inc. All Rights Reserved.
+
+#include "Expressions/Filter/TG_Expression_Warp.h"
+
+void UTG_Expression_Warp::Evaluate(FTG_EvaluationContext* InContext)
+{
+	Super::Evaluate(InContext);
+
+	if (!Input)
+	{
+		Output = FTG_Texture::GetBlack();
+		return;
+	}
+
+	/// By default we just a flat white mask
+	if (!Mask)
+		Mask = FTG_Texture::GetWhite();
+
+	if (Type == EWarp::Directional)
+	{
+		Output = T_Filter::CreateDirectionalWarp(InContext->Cycle, Output.GetBufferDescriptor(), Input, Mask, Intensity, (Angle * PI) / 180.0f, InContext->TargetId);
+	}
+	else if (Type == EWarp::Normal)
+	{
+		Output = T_Filter::CreateNormalWarp(InContext->Cycle, Output.GetBufferDescriptor(), Input, Mask, Intensity, InContext->TargetId);
+	}
+	else if (Type == EWarp::Sine)
+	{
+		Output = T_Filter::CreateSineWarp(InContext->Cycle, Output.GetBufferDescriptor(), Input, Mask, Intensity, PhaseU, PhaseV, InContext->TargetId);
+	}
+	else
+	{
+		Output = FTG_Texture::GetBlack();
+	}
+}
+
+#if WITH_EDITOR
+bool UTG_Expression_Warp::CanEditChange(const FProperty* InProperty) const
+{
+	bool bEditCondition = Super::CanEditChange(InProperty);
+
+	const FName PropertyName = InProperty->GetFName();
+
+	if (PropertyName == GET_MEMBER_NAME_CHECKED(UTG_Expression_Warp, Angle))
+	{
+		bEditCondition = Type == EWarp::Directional;
+	}
+	if (PropertyName == GET_MEMBER_NAME_CHECKED(UTG_Expression_Warp, PhaseU))
+	{
+		bEditCondition = Type == EWarp::Sine;
+	}
+	if (PropertyName == GET_MEMBER_NAME_CHECKED(UTG_Expression_Warp, PhaseV))
+	{
+		bEditCondition = Type == EWarp::Sine;
+	}
+
+	// Default behaviour
+	return bEditCondition;
+}
+#endif
diff --git a/Engine/Plugins/Experimental/TextureGraph/Source/TextureGraph/Public/Expressions/Filter/TG_Expression_Warp.h b/Engine/Plugins/Experimental/TextureGraph/Source/TextureGraph/Public/Expressions/Filter/TG_Expression_Warp.h
new file mode 100644
--- /dev/null
+++ b/Engine/Plugins/Experimental/TextureGraph/Source/TextureGraph/Public/Expressions/Filter/TG_Expression_Warp.h
@@ -0,0 +1,63 @@
+// Copyright Epic Games, Inc. All Rights Reserved.
+
+#pragma once
+
+#include "CoreMinimal.h"
+#include "Expressions/TG_Expression_MaterialBase.h"
+#include "Transform/Expressions/T_Filter.h"
+
+#include "TG_Expression_Warp.generated.h"
+
+//////////////////////////////////////////////////////////////////////////
+/// Base warp class with common options
+//////////////////////////////////////////////////////////////////////////
+UCLASS()
+class TEXTUREGRAPH_API UTG_Expression_Warp : public UTG_Expression
+{
+	GENERATED_BODY()
+
+public:
+	TG_DECLARE_EXPRESSION(TG_Category::Filter);
+
+	/// The type of warp that we want on the input image	
+	UPROPERTY(EditAnywhere, BlueprintReadWrite, Category = NoCategory, meta = (TGType = "TG_Setting", RegenPinsOnChange))
+	TEnumAsByte<EWarp::Type> Type = EWarp::Directional;
+
+	// What is the intensity of the warp. Warp of 1 is 10% (1/10th) of the width of the input image
+	UPROPERTY(EditAnywhere, BlueprintReadWrite, Category = NoCategory, meta = (TGType = "TG_Setting", ClampMin = 0.0, ClampMax = 20, UIMin = 0.0, UIMax = 1, EditConditionHides))
+	float Intensity = 1;
+
+	// Angle of the directional warp 0 - 360 degrees
+	UPROPERTY(EditAnywhere, BlueprintReadWrite, Category = NoCategory, meta = (TGType = "TG_Setting", ClampMin = 0.0, ClampMax = 360, UIMin = 0.0, UIMax = 360, EditConditionHides))
+	float Angle = 0;
+
+	// Phase X for the Sine wave warp (for U coordinate in UV space)
+	UPROPERTY(EditAnywhere, BlueprintReadWrite, Category = NoCategory, meta = (TGType = "TG_Setting", ClampMin = -1, ClampMax = 1, UIMin = -1, UIMax = 1, EditConditionHides))
+	float PhaseU = 0.0f;
+
+	// Phase Y for the Sine wave warp (for V coordinate in UV space)
+	UPROPERTY(EditAnywhere, BlueprintReadWrite, Category = NoCategory, meta = (TGType = "TG_Setting", ClampMin = -1, ClampMax = 1, UIMin = -1, UIMax = 1, EditConditionHides))
+	float PhaseV = 0.0f;
+
+	// The input texture to apply the warp
+	UPROPERTY(meta = (TGType = "TG_Input"))
+	FTG_Texture Input;
+
+	// The mask used for the warp effect
+	UPROPERTY(meta = (TGType = "TG_Input"))
+	FTG_Texture Mask;
+
+	// The output texture having blurred effect
+	UPROPERTY(EditAnywhere, Category = NoCategory, meta = (TGType = "TG_Output", PinDisplayName = ""))
+	FTG_Texture Output;
+
+public:
+#if WITH_EDITOR
+	// Used to implement EditCondition logic for both Node UI and Details View
+	virtual bool						CanEditChange(const FProperty* InProperty) const override;
+#endif
+
+	virtual FText GetTooltipText() const override { return FText::FromString(TEXT("Evalute different kinds of warp depending on the Mask and Intensity settings.")); }
+	virtual void Evaluate(FTG_EvaluationContext* InContext) override;
+};
+
diff --git a/Engine/Plugins/Experimental/TextureGraph/Source/TextureGraphEngine/Device/FX/Device_FX.cpp b/Engine/Plugins/Experimental/TextureGraph/Source/TextureGraphEngine/Device/FX/Device_FX.cpp
--- a/Engine/Plugins/Experimental/TextureGraph/Source/TextureGraphEngine/Device/FX/Device_FX.cpp
+++ b/Engine/Plugins/Experimental/TextureGraph/Source/TextureGraphEngine/Device/FX/Device_FX.cpp
@@ -578,7 +578,7 @@ TexPtr Device_FX::AllocateRenderTarget(const BufferDescriptor& Desc)
 
 	SCOPE_CYCLE_COUNTER(STAT_Device_FX_AllocateRendertarget)
 
-		RTList* List = nullptr;
+	RTList* List = nullptr;
 	HashType Hash = Desc.FormatHashValue();
 
 	{
@@ -686,10 +686,10 @@ int32 Device_FX::InitRTResource(TexPtr TextureObj, UTextureRenderTarget* RT)
 	UE_LOG(LogDevice, VeryVerbose, TEXT("New RT Array allocation: %s %llu [Ptr: 0x%x, Size: %dx%d]"), *TextureObj->GetDescriptor().Name,
 		TextureObj->GetDescriptor().Format_HashValue(), RT, RTRes->GetSizeX(), RTRes->GetSizeY());
 
-	FTextureRHIRef rhiTexture = RTRes->GetTextureRHI();
-	check(rhiTexture);
+	FTextureRHIRef RHITexture = RTRes->GetTextureRHI();
+	check(RHITexture);
 
-	RHIBindDebugLabelName(rhiTexture, *TextureObj->GetDescriptor().Name);
+	RHIBindDebugLabelName(RHITexture, *TextureObj->GetDescriptor().Name);
 
 	return 0;
 }
diff --git a/Engine/Plugins/Experimental/TextureGraph/Source/TextureGraphEngine/Transform/Expressions/T_EdgeDetect.cpp b/Engine/Plugins/Experimental/TextureGraph/Source/TextureGraphEngine/Transform/Expressions/T_EdgeDetect.cpp
deleted file mode 100644
diff --git a/Engine/Plugins/Experimental/TextureGraph/Source/TextureGraphEngine/Transform/Expressions/T_EdgeDetect.h b/Engine/Plugins/Experimental/TextureGraph/Source/TextureGraphEngine/Transform/Expressions/T_EdgeDetect.h
deleted file mode 100644
diff --git a/Engine/Plugins/Experimental/TextureGraph/Source/TextureGraphEngine/Transform/Expressions/T_Filter.cpp b/Engine/Plugins/Experimental/TextureGraph/Source/TextureGraphEngine/Transform/Expressions/T_Filter.cpp
new file mode 100644
--- /dev/null
+++ b/Engine/Plugins/Experimental/TextureGraph/Source/TextureGraphEngine/Transform/Expressions/T_Filter.cpp
@@ -0,0 +1,126 @@
+// Copyright Epic Games, Inc. All Rights Reserved.
+
+#include "T_Filter.h"
+#include "FxMat/MaterialManager.h"
+#include "Job/JobBatch.h"
+#include "Transform/Utility/T_CombineTiledBlob.h"
+
+IMPLEMENT_GLOBAL_SHADER(FSH_EdgeDetect, "/Plugin/TextureGraph/Expressions/Expression_EdgeDetect.usf", "FSH_EdgeDetect", SF_Pixel);
+IMPLEMENT_GLOBAL_SHADER(FSH_DirectionalWarp, "/Plugin/TextureGraph/Expressions/Expression_Warp.usf", "FSH_DirectionalWarp", SF_Pixel);
+IMPLEMENT_GLOBAL_SHADER(FSH_NormalWarp, "/Plugin/TextureGraph/Expressions/Expression_Warp.usf", "FSH_NormalWarp", SF_Pixel);
+IMPLEMENT_GLOBAL_SHADER(FSH_SineWarp, "/Plugin/TextureGraph/Expressions/Expression_Warp.usf", "FSH_SineWarp", SF_Pixel);
+
+TiledBlobPtr T_Filter::CreateEdgeDetect(MixUpdateCyclePtr InCycle, BufferDescriptor DesiredDesc, TiledBlobPtr SourceTexture, float Thickness, int32 InTargetId)
+{
+	if (!SourceTexture)
+		return TextureHelper::GetBlack();
+
+	RenderMaterial_FXPtr RenderMaterial = TextureGraphEngine::GetMaterialManager()->CreateMaterial_FX<VSH_Simple, FSH_EdgeDetect>(TEXT("T_EdgeDetect"));
+	check(RenderMaterial);
+
+	// Combine the tiled texture
+	TiledBlobPtr CombinedBlob = T_CombineTiledBlob::Create(InCycle, SourceTexture->GetDescriptor(), InTargetId, SourceTexture);
+	float StepX = Thickness / (float)SourceTexture->GetWidth();
+	float StepY = Thickness / (float)SourceTexture->GetHeight();
+
+	BufferDescriptor Desc = BufferDescriptor::Combine(DesiredDesc, SourceTexture->GetDescriptor());
+
+	FTileInfo TileInfo;
+	JobUPtr RenderJob = std::make_unique<Job>(InCycle->GetMix(), InTargetId, std::static_pointer_cast<BlobTransform>(RenderMaterial));
+	RenderJob
+		->AddArg(ARG_TILEINFO(TileInfo, "TileInfo"))
+		->AddArg(ARG_BLOB(CombinedBlob, "SourceTexture"))
+		->AddArg(ARG_FLOAT(Thickness, "Thickness"))
+		;
+
+	const FString Name = FString::Printf(TEXT("T_EdgeDetect.[%llu]"), InCycle->GetBatch()->GetBatchId());
+
+	TiledBlobPtr Result = RenderJob->InitResult(Name, &Desc);
+	InCycle->AddJob(InTargetId, std::move(RenderJob));
+
+	return Result;
+}
+
+template <typename FSH_Type>
+static JobUPtr CreateWarp(FString Name, MixUpdateCyclePtr InCycle, BufferDescriptor DesiredDesc, TiledBlobPtr SourceTexture, TiledBlobPtr Mask, 
+	float Intensity, int32 InTargetId, BufferDescriptor& Desc)
+{
+	/// If no mask is given then we just use a white mask
+	if (!Mask)
+		Mask = TextureHelper::GetWhite();
+
+	RenderMaterial_FXPtr RenderMaterial = TextureGraphEngine::GetMaterialManager()->CreateMaterial_FX<VSH_Simple, FSH_Type>(Name);
+	check(RenderMaterial);
+
+	// Combine the tiled texture
+	TiledBlobPtr CombinedSourceTexture = T_CombineTiledBlob::Create(InCycle, SourceTexture->GetDescriptor(), InTargetId, SourceTexture);
+	TiledBlobPtr CombinedMask = T_CombineTiledBlob::Create(InCycle, Mask->GetDescriptor(), InTargetId, Mask);
+	BufferDescriptor CombinedDesc = BufferDescriptor::Combine(SourceTexture->GetDescriptor(), Mask->GetDescriptor());
+	Desc = BufferDescriptor::Combine(DesiredDesc, CombinedDesc);
+
+	FTileInfo TileInfo;
+	JobUPtr RenderJob = std::make_unique<Job>(InCycle->GetMix(), InTargetId, std::static_pointer_cast<BlobTransform>(RenderMaterial));
+	RenderJob
+		->AddArg(ARG_TILEINFO(TileInfo, "TileInfo"))
+		->AddArg(ARG_BLOB(CombinedSourceTexture, "SourceTexture"))
+		->AddArg(ARG_BLOB(CombinedMask, "Mask"))
+		->AddArg(ARG_FLOAT(Intensity, "Intensity"))
+		;
+
+	return std::move(RenderJob);
+}
+
+TiledBlobPtr T_Filter::CreateDirectionalWarp(MixUpdateCyclePtr InCycle, BufferDescriptor DesiredDesc, TiledBlobPtr SourceTexture, TiledBlobPtr Mask, float Intensity, float AngleRad, int32 InTargetId)
+{
+	if (!SourceTexture)
+		return TextureHelper::GetBlack();
+
+	BufferDescriptor Desc;
+	JobUPtr RenderJob = CreateWarp<FSH_DirectionalWarp>(TEXT("T_DirectionalWarp"), InCycle, DesiredDesc, SourceTexture, Mask, Intensity, InTargetId, Desc);
+	RenderJob
+		->AddArg(ARG_FLOAT(AngleRad, "AngleRad"))
+		;
+
+	const FString Name = FString::Printf(TEXT("T_DirectionalWarp.[%llu]"), InCycle->GetBatch()->GetBatchId());
+
+	TiledBlobPtr Result = RenderJob->InitResult(Name, &Desc);
+	InCycle->AddJob(InTargetId, std::move(RenderJob));
+
+	return Result;
+}
+
+TiledBlobPtr T_Filter::CreateNormalWarp(MixUpdateCyclePtr InCycle, BufferDescriptor DesiredDesc, TiledBlobPtr SourceTexture, TiledBlobPtr Mask, float Intensity, int32 InTargetId)
+{
+	if (!SourceTexture)
+		return TextureHelper::GetBlack();
+
+	BufferDescriptor Desc;
+	JobUPtr RenderJob = CreateWarp<FSH_NormalWarp>(TEXT("T_NormalWarp"), InCycle, DesiredDesc, SourceTexture, Mask, Intensity, InTargetId, Desc);
+	const FString Name = FString::Printf(TEXT("T_NormalWarp.[%llu]"), InCycle->GetBatch()->GetBatchId());
+
+	TiledBlobPtr Result = RenderJob->InitResult(Name, &Desc);
+	InCycle->AddJob(InTargetId, std::move(RenderJob));
+
+	return Result;
+}
+
+TiledBlobPtr T_Filter::CreateSineWarp(MixUpdateCyclePtr InCycle, BufferDescriptor DesiredDesc, TiledBlobPtr SourceTexture, TiledBlobPtr Mask, float Intensity, float PhaseU, float PhaseV, int32 InTargetId)
+{
+	if (!SourceTexture)
+		return TextureHelper::GetBlack();
+
+	BufferDescriptor Desc;
+	JobUPtr RenderJob = CreateWarp<FSH_SineWarp>(TEXT("T_SineWarp"), InCycle, DesiredDesc, SourceTexture, Mask, Intensity, InTargetId, Desc);
+	RenderJob
+		->AddArg(ARG_FLOAT(PhaseU, "PhaseU"))
+		->AddArg(ARG_FLOAT(PhaseV, "PhaseV"))
+		;
+
+	const FString Name = FString::Printf(TEXT("T_SineWarp.[%llu]"), InCycle->GetBatch()->GetBatchId());
+
+	TiledBlobPtr Result = RenderJob->InitResult(Name, &Desc);
+	InCycle->AddJob(InTargetId, std::move(RenderJob));
+
+	return Result;
+}
+
diff --git a/Engine/Plugins/Experimental/TextureGraph/Source/TextureGraphEngine/Transform/Expressions/T_Filter.h b/Engine/Plugins/Experimental/TextureGraph/Source/TextureGraphEngine/Transform/Expressions/T_Filter.h
new file mode 100644
--- /dev/null
+++ b/Engine/Plugins/Experimental/TextureGraph/Source/TextureGraphEngine/Transform/Expressions/T_Filter.h
@@ -0,0 +1,120 @@
+// Copyright Epic Games, Inc. All Rights Reserved.
+#pragma once
+
+#include <2D/BlendModes.h>
+
+#include "CoreMinimal.h"
+#include "Job/Job.h"
+#include "FxMat/FxMaterial.h"
+#include "Model/Mix/MixUpdateCycle.h"
+#include <DataDrivenShaderPlatformInfo.h>
+
+//////////////////////////////////////////////////////////////////////////
+/// Edge detection
+//////////////////////////////////////////////////////////////////////////
+class TEXTUREGRAPHENGINE_API FSH_EdgeDetect : public FSH_Base
+{
+public:
+	DECLARE_GLOBAL_SHADER(FSH_EdgeDetect);
+	SHADER_USE_PARAMETER_STRUCT(FSH_EdgeDetect, FSH_Base);
+	BEGIN_SHADER_PARAMETER_STRUCT(FParameters, )
+		SHADER_PARAMETER_STRUCT(FTileInfo, TileInfo)
+		SHADER_PARAMETER_STRUCT(FStandardSamplerStates, SamplerStates)
+		SHADER_PARAMETER_TEXTURE(Texture2D, SourceTexture)
+		SHADER_PARAMETER(float, Thickness)
+	END_SHADER_PARAMETER_STRUCT()
+
+public:
+	TEXTURE_ENGINE_DEFAULT_PERMUTATION;
+	TEXTUREGRAPH_ENGINE_DEFAULT_COMPILATION_ENV;
+};
+
+UENUM(BlueprintType)
+namespace EWarp
+{
+	enum Type : int
+	{
+		Directional = 0				UMETA(DisplayName = "Directional"),
+		Normal = 1					UMETA(DisplayName = "Normal"),
+		Sine = 2					UMETA(DisplayName = "Sine"),
+	};
+}
+
+//////////////////////////////////////////////////////////////////////////
+/// Directional warp
+//////////////////////////////////////////////////////////////////////////
+class TEXTUREGRAPHENGINE_API FSH_DirectionalWarp : public FSH_Base
+{
+public:
+	DECLARE_GLOBAL_SHADER(FSH_DirectionalWarp);
+	SHADER_USE_PARAMETER_STRUCT(FSH_DirectionalWarp, FSH_Base);
+	BEGIN_SHADER_PARAMETER_STRUCT(FParameters, )
+		SHADER_PARAMETER_STRUCT(FTileInfo, TileInfo)
+		SHADER_PARAMETER_STRUCT(FStandardSamplerStates, SamplerStates)
+		SHADER_PARAMETER_TEXTURE(Texture2D, SourceTexture)
+		SHADER_PARAMETER_TEXTURE(Texture2D, Mask)
+		SHADER_PARAMETER(float, AngleRad)
+		SHADER_PARAMETER(float, Intensity)
+	END_SHADER_PARAMETER_STRUCT()
+
+public:
+	TEXTURE_ENGINE_DEFAULT_PERMUTATION;
+	TEXTUREGRAPH_ENGINE_DEFAULT_COMPILATION_ENV;
+};
+
+//////////////////////////////////////////////////////////////////////////
+/// Directional warp
+//////////////////////////////////////////////////////////////////////////
+class TEXTUREGRAPHENGINE_API FSH_NormalWarp : public FSH_Base
+{
+public:
+	DECLARE_GLOBAL_SHADER(FSH_NormalWarp);
+	SHADER_USE_PARAMETER_STRUCT(FSH_NormalWarp, FSH_Base);
+	BEGIN_SHADER_PARAMETER_STRUCT(FParameters, )
+		SHADER_PARAMETER_STRUCT(FTileInfo, TileInfo)
+		SHADER_PARAMETER_STRUCT(FStandardSamplerStates, SamplerStates)
+		SHADER_PARAMETER_TEXTURE(Texture2D, SourceTexture)
+		SHADER_PARAMETER_TEXTURE(Texture2D, Mask)
+		SHADER_PARAMETER(float, Intensity)
+	END_SHADER_PARAMETER_STRUCT()
+
+public:
+	TEXTURE_ENGINE_DEFAULT_PERMUTATION;
+	TEXTUREGRAPH_ENGINE_DEFAULT_COMPILATION_ENV;
+};
+
+//////////////////////////////////////////////////////////////////////////
+/// Directional warp
+//////////////////////////////////////////////////////////////////////////
+class TEXTUREGRAPHENGINE_API FSH_SineWarp : public FSH_Base
+{
+public:
+	DECLARE_GLOBAL_SHADER(FSH_SineWarp);
+	SHADER_USE_PARAMETER_STRUCT(FSH_SineWarp, FSH_Base);
+	BEGIN_SHADER_PARAMETER_STRUCT(FParameters, )
+		SHADER_PARAMETER_STRUCT(FTileInfo, TileInfo)
+		SHADER_PARAMETER_STRUCT(FStandardSamplerStates, SamplerStates)
+		SHADER_PARAMETER_TEXTURE(Texture2D, SourceTexture)
+		SHADER_PARAMETER_TEXTURE(Texture2D, Mask)
+		SHADER_PARAMETER(float, Intensity)
+		SHADER_PARAMETER(float, PhaseU)
+		SHADER_PARAMETER(float, PhaseV)
+	END_SHADER_PARAMETER_STRUCT()
+
+public:
+	TEXTURE_ENGINE_DEFAULT_PERMUTATION;
+	TEXTUREGRAPH_ENGINE_DEFAULT_COMPILATION_ENV;
+};
+
+//////////////////////////////////////////////////////////////////////////
+class TEXTUREGRAPHENGINE_API T_Filter
+{
+public:
+	//////////////////////////////////////////////////////////////////////////
+	/// Static functions
+	//////////////////////////////////////////////////////////////////////////
+	static TiledBlobPtr				CreateEdgeDetect(MixUpdateCyclePtr InCycle, BufferDescriptor DesiredDesc, TiledBlobPtr SourceTexture, float Thickness, int32 InTargetId);
+	static TiledBlobPtr				CreateDirectionalWarp(MixUpdateCyclePtr InCycle, BufferDescriptor DesiredDesc, TiledBlobPtr SourceTexture, TiledBlobPtr Mask, float Intensity, float AngleRad, int32 InTargetId);
+	static TiledBlobPtr				CreateNormalWarp(MixUpdateCyclePtr InCycle, BufferDescriptor DesiredDesc, TiledBlobPtr SourceTexture, TiledBlobPtr Mask, float Intensity, int32 InTargetId);
+	static TiledBlobPtr				CreateSineWarp(MixUpdateCyclePtr InCycle, BufferDescriptor DesiredDesc, TiledBlobPtr SourceTexture, TiledBlobPtr Mask, float Intensity, float PhaseU, float PhaseV, int32 InTargetId);
+};
diff --git a/Engine/Plugins/Experimental/TextureGraph/Source/TextureGraphEngine/Transform/Utility/T_CombineTiledBlob.cpp b/Engine/Plugins/Experimental/TextureGraph/Source/TextureGraphEngine/Transform/Utility/T_CombineTiledBlob.cpp
--- a/Engine/Plugins/Experimental/TextureGraph/Source/TextureGraphEngine/Transform/Utility/T_CombineTiledBlob.cpp
+++ b/Engine/Plugins/Experimental/TextureGraph/Source/TextureGraphEngine/Transform/Utility/T_CombineTiledBlob.cpp
@@ -11,7 +11,7 @@ TiledBlobPtr T_CombineTiledBlob::Create(MixUpdateCyclePtr Cycle, BufferDescripto
 {
 	if (SourceTex && SourceTex->IsTiled())
 	{
-		FString Name = FString::Printf(TEXT("[%s].[%d].[%llu] CombineTiledBlob"), *SourceTex->Name(), TargetId, Cycle->GetBatch()->GetFrameId());
+		FString Name = TEXT("T_CombineTiledBlob"); //FString::Printf(TEXT("[%s].[%d].[%llu] CombineTiledBlob"), *SourceTex->Name(), TargetId, Cycle->GetBatch()->GetFrameId());
 
 		std::shared_ptr<CombineTiledBlob_Transform> CombinedTiledBlobMat = std::make_shared< CombineTiledBlob_Transform>(Name, SourceTex);
 
