Text conflict for AssetDefinition_Texture. Needed to be AssetDefinition_TextureGraph

diff --git a/Engine/Plugins/Experimental/TextureGraph/Source/TextureGraphEditor/Private/AssetDefinition_TextureGraph.h b/Engine/Plugins/Experimental/TextureGraph/Source/TextureGraphEditor/Private/AssetDefinition_TextureGraph.h
--- a/Engine/Plugins/Experimental/TextureGraph/Source/TextureGraphEditor/Private/AssetDefinition_TextureGraph.h
+++ b/Engine/Plugins/Experimental/TextureGraph/Source/TextureGraphEditor/Private/AssetDefinition_TextureGraph.h
@@ -14,7 +14,7 @@ class UAssetDefinition_TextureGraph : public UAssetDefinitionDefault
 
 public:
 	// UAssetDefinition Begin
-	virtual FText GetAssetDisplayName() const override { return NSLOCTEXT("AssetTypeActions", "AssetTypeActions_Texture", "Texture Graph"); }
+	virtual FText GetAssetDisplayName() const override { return NSLOCTEXT("AssetTypeActions", "AssetTypeActions_TextureGraph", "Texture Graph"); }
 	virtual FLinearColor GetAssetColor() const override { return FLinearColor(FColor::Emerald); }
 	virtual TSoftClassPtr<UObject> GetAssetClass() const override { return UTextureGraph::StaticClass(); }
 	virtual TConstArrayView<FAssetCategoryPath> GetAssetCategories() const override
