Rerender after Export if DebugGrid is enabled

diff --git a/Engine/Plugins/Experimental/TextureGraph/Source/TextureGraph/Private/Expressions/Procedural/TG_Expression_Transform.cpp b/Engine/Plugins/Experimental/TextureGraph/Source/TextureGraph/Private/Expressions/Procedural/TG_Expression_Transform.cpp
--- a/Engine/Plugins/Experimental/TextureGraph/Source/TextureGraph/Private/Expressions/Procedural/TG_Expression_Transform.cpp
+++ b/Engine/Plugins/Experimental/TextureGraph/Source/TextureGraph/Private/Expressions/Procedural/TG_Expression_Transform.cpp
@@ -31,13 +31,10 @@ void UTG_Expression_Transform::Evaluate(FTG_EvaluationContext* InContext)
 	FVector2f OutOffset = Offset;
 
 	// hide debug grid when exporting
+	float ShowDebugGridValue = ShowDebugGrid;
 	if (InContext->Cycle->GetDetails().bExporting)
 	{
-		// show warning to user that debug grid is being turned off
-		auto ErrorType = static_cast<int32>(ETextureGraphErrorType::NODE_WARNING);
-		const FString ErrorMsg = FString::Format(TEXT("Export command issued, turning 'Debug Grid' off in {0} node."), {GetTitleName().ToString()} );
-		TextureGraphEngine::GetErrorReporter(InContext->Cycle->GetMix())->ReportWarning(ErrorType, ErrorMsg, GetParentNode());
-		ShowDebugGrid = 0.0f;
+		ShowDebugGridValue = 0.0f;
 	}
 	
 	T_Transform::TransformParameter XformParam{
@@ -61,7 +58,7 @@ void UTG_Expression_Transform::Evaluate(FTG_EvaluationContext* InContext)
 		.WrapFilterMode = WrapMode,
 		.MirrorX = MirrorX,
 		.MirrorY = MirrorY,
-		.ShowDebugGrid = ShowDebugGrid
+		.ShowDebugGrid = ShowDebugGridValue
 	};
 	
 	Output = T_Transform::Create(InContext->Cycle, DesiredDescriptor, Input,
diff --git a/Engine/Plugins/Experimental/TextureGraph/Source/TextureGraph/Private/TG_HelperFunctions.cpp b/Engine/Plugins/Experimental/TextureGraph/Source/TextureGraph/Private/TG_HelperFunctions.cpp
--- a/Engine/Plugins/Experimental/TextureGraph/Source/TextureGraph/Private/TG_HelperFunctions.cpp
+++ b/Engine/Plugins/Experimental/TextureGraph/Source/TextureGraph/Private/TG_HelperFunctions.cpp
@@ -167,6 +167,10 @@ AsyncBool FTG_HelperFunctions::ExportAsync(UTextureGraph* InTextureGraph, FStrin
 		return TextureExporter::ExportAsUAsset(InTextureGraph, TargetExportSettings, ExportPath);
 	})
 	.then([InTextureGraph, &TargetExportSettings]()
+	{
+		InTextureGraph->InvalidateAll();
+	})
+	.then([InTextureGraph, &TargetExportSettings]()
 	{
 		// Add analytics tag
 		if (FEngineAnalytics::IsAvailable())
