Fix Histogram

- Modify the REnderDoc capture begin/end injection specifically on any Batch to be able to capture any type of batch from the Texture Graph engine.
- Introduce a new REnderDoc Next Batch Histogram command which will catprue the next batch run by the histogram service
- Introduced a new JobArg modifier for tiled blob which allow to collect the full araray of tile textures in a tiled blob and pass it to the RenderMaterial. This is needed for the histogram job to run ONE instance for tiled blob and collect the full array of tiles in the bound argument.

- Introduced the Histogram FXMAterial in order to control the execution, The histogram with a multipass algorithm:
   For each tiles, we run one pass producing one histogram as part of a shared resource
   Do one more pass producing the final histogram for the full source data
- Modified the Grayscale function to use the standard RGB to Lum code.

#tests 66394393242b321d4ca72401

diff --git a/Engine/Plugins/Experimental/TextureGraph/Shaders/ShaderUtil.ush b/Engine/Plugins/Experimental/TextureGraph/Shaders/ShaderUtil.ush
--- a/Engine/Plugins/Experimental/TextureGraph/Shaders/ShaderUtil.ush
+++ b/Engine/Plugins/Experimental/TextureGraph/Shaders/ShaderUtil.ush
@@ -45,9 +45,18 @@ float3x3 GetTransform_TS_to_WS(float3x3 _ObjectToWorld, float3 worldNormal, floa
 	return tangentToWorld;
 }
 
-float Grayscale(float3 SourceClr)
+
+float LumaFromRGB(float3 Color)
+{
+	// ITU-R BT.709 colour luminance gradients
+	return 0.2126 * Color.r + 0.7152 * Color.g + 0.0722* Color.b;
+}
+
+
+float Grayscale(float3 Color)
 {
-	return 0.299 * SourceClr.r + 0.587 * SourceClr.g + 0.114 * SourceClr.b;
+	// ITU-R BT.709 colour luminance gradients
+	return LumaFromRGB(Color);
 }
 
 float3 RGBToHSV(float3 c)
diff --git a/Engine/Plugins/Experimental/TextureGraph/Shaders/Utils/Histogram_comp.usf b/Engine/Plugins/Experimental/TextureGraph/Shaders/Utils/Histogram_comp.usf
--- a/Engine/Plugins/Experimental/TextureGraph/Shaders/Utils/Histogram_comp.usf
+++ b/Engine/Plugins/Experimental/TextureGraph/Shaders/Utils/Histogram_comp.usf
@@ -1,5 +1,6 @@
 // Copyright Epic Games, Inc. All Rights Reserved.
 #include "/Engine/Public/Platform.ush"
+#include "/Plugin/TextureGraph/ShaderUtil.ush"
 
 #ifndef MAX_flt
 #define MAX_flt			(3.402823466e+38F)
@@ -7,8 +8,16 @@
 
 Texture2D SourceTiles;
 
+RWStructuredBuffer<uint4> TilesHistogramUAV;
+StructuredBuffer<uint4> TilesHistogramSRV;
+
 RWTexture2D<float4> Result;
-RWStructuredBuffer<uint4> Histogram;
+
+int4 InvocationDim;
+
+// Shared mem Histogram is 4x256, one element per bin.channel
+#define HistogramSize 256
+groupshared uint4 _SharedHistogram[HistogramSize];
 
 #ifndef THREADGROUPSIZE_X
 #define THREADGROUPSIZE_X 16
@@ -17,68 +26,108 @@ RWStructuredBuffer<uint4> Histogram;
 #endif
 
 [numthreads(THREADGROUPSIZE_X, THREADGROUPSIZE_Y, THREADGROUPSIZE_Z)]
-void CSH_Histogram(
+void CSH_HistogramPerTile(
 	uint3 GroupId : SV_GroupID,
 	uint3 ThreadId : SV_DispatchThreadID,
 	uint3 GroupThreadId : SV_GroupThreadID,
 	uint GroupIndex: SV_GroupIndex
 	)
 {
-	float4 Col = SourceTiles.Load(uint3(ThreadId.xy, 0));
-	uint RBin = uint(Col.r * float(255));
-	uint GBin = Col.g * 255;
-	uint BBin = Col.b * 255;
-	int Luma = (int)((0.2126 * RBin) + (0.7152 * GBin) + (0.0722 * BBin));
+	// Grab one pixel of the source
+	float4 Fetched = SourceTiles.Load(uint3(ThreadId.xy, 0));
+	
+	// Evaluate Luma for 4th component and find the bin index for each channel
+	float Luma = Grayscale(Fetched);
+	// The considered value is in the range [0, 1];
+	float4 Value = saturate(float4(Fetched.rgb, Luma));
+	// From the float value to the matching histogram bin [0,255]
+	float4 BinValue = floor(Value * 256.0);
+	uint4 Bin = clamp(uint4(BinValue), 0, 255);
+	
+	// Offset where to write the histogram for this tile in the destination buffer,
+	// using the region allocated for this tile
+	int TileOffset = InvocationDim.y * 256; 
+	uint4 offsetbins = Bin + TileOffset;
+	
+	InterlockedAdd(TilesHistogramUAV[offsetbins.x].x, 1);
+	InterlockedAdd(TilesHistogramUAV[offsetbins.y].y, 1);
+	InterlockedAdd(TilesHistogramUAV[offsetbins.z].z, 1);
+	InterlockedAdd(TilesHistogramUAV[offsetbins.w].w, 1);
+}
 
-	InterlockedAdd(Histogram[RBin].r, 1);
-	InterlockedAdd(Histogram[GBin].g, 1);
-	InterlockedAdd(Histogram[BBin].b, 1);
-	InterlockedAdd(Histogram[Luma].a, 1);
 
+[numthreads(256, 1, 1)]
+void CSH_Histogram(
+	uint3 GroupId : SV_GroupID,
+	uint3 ThreadId : SV_DispatchThreadID,
+	uint3 GroupThreadId : SV_GroupThreadID,
+	uint GroupIndex: SV_GroupIndex
+	)
+{
+	uint ThreadIndex = ThreadId.x;
+	
+	// Accumulate all the histogram bin values for all the tiles in THIS thread bin
+	uint4 HistogramAccum = 0;
+	for (uint i = 0; i < InvocationDim.x; ++i)
+	{
+		uint TileHistogramBinOffset = (i * 256 + ThreadIndex);
+		HistogramAccum += TilesHistogramSRV.Load(TileHistogramBinOffset);
+	}
+	
+	// And store in the histogram texture and also in shared meme histo
+	Result[uint2(ThreadIndex,0)] = float4(HistogramAccum);
+	_SharedHistogram[ThreadIndex] = HistogramAccum;
+	
+	
 	GroupMemoryBarrierWithGroupSync();
 
-	if (GroupIndex == 0)
+	// Now only in first group, let's go through the shared mem histogram bins SEQUENTIALLY and extract the min and max count and min and max values
+	if (ThreadId.x == 0)
 	{
-		// Now only in first group, let's go through the histogram bins SEQUENTIALLY and extract the min and max count and min and max values
 		float4 MinCount = float4(MAX_flt,MAX_flt,MAX_flt,MAX_flt);
 		float4 MaxCount = float4(0,0,0,0);
 		
-		float4 MinValue = float4(MAX_flt,MAX_flt,MAX_flt,MAX_flt);
-		float4 MaxValue = float4(0, 0, 0, 0);
+		float4 MinBin = float4(MAX_flt,MAX_flt,MAX_flt,MAX_flt);
+		float4 MaxBin = float4(0, 0, 0, 0);
 			
+		float4 TotalCount = float4(0,0,0,0);
+		
 		LOOP for(uint y = 0; y < 256 ; y++)
 		{		
 			
-			float4 HistogramVal = float4(Histogram[y]);
+			float4 HistogramBin = float4(_SharedHistogram[y]);
 	
-			MinCount = min(HistogramVal, MinCount);
-			MaxCount = max(HistogramVal, MaxCount);
-			
-			Result[uint2(y,0)] = HistogramVal;
+			MinCount = min(HistogramBin, MinCount);
+			MaxCount = max(HistogramBin, MaxCount);
+			TotalCount += HistogramBin;
 			
-			if (any(HistogramVal))
+			if (any(HistogramBin))
 			{
-				// Assigned min value for the first HistogramVal found, never after
-				MinValue.x = (MinValue.x != MAX_flt ? MinValue.x : (HistogramVal.x == 0 ? MAX_flt : y) );
-				MinValue.y = (MinValue.y != MAX_flt ? MinValue.y : (HistogramVal.y == 0 ? MAX_flt : y) );
-				MinValue.z = (MinValue.z != MAX_flt ? MinValue.z : (HistogramVal.z == 0 ? MAX_flt : y) );
-				MinValue.w = (MinValue.w != MAX_flt ? MinValue.w : (HistogramVal.w == 0 ? MAX_flt : y) );
+				// Assigned min bin for the first HistogramBin found, never after
+				MinBin.x = (MinBin.x != MAX_flt ? MinBin.x : (HistogramBin.x == 0 ? MAX_flt : y) );
+				MinBin.y = (MinBin.y != MAX_flt ? MinBin.y : (HistogramBin.y == 0 ? MAX_flt : y) );
+				MinBin.z = (MinBin.z != MAX_flt ? MinBin.z : (HistogramBin.z == 0 ? MAX_flt : y) );
+				MinBin.w = (MinBin.w != MAX_flt ? MinBin.w : (HistogramBin.w == 0 ? MAX_flt : y) );
 				
-				// Assigned max value, always the current Y of the loop if the Histogram is not 0.
-				MaxValue.x = (HistogramVal.x == 0 ? MaxValue.x : y + 1); 
-				MaxValue.y = (HistogramVal.y == 0 ? MaxValue.y : y + 1); 
-				MaxValue.z = (HistogramVal.z == 0 ? MaxValue.z : y + 1); 
-				MaxValue.w = (HistogramVal.w == 0 ? MaxValue.w : y + 1); 
+				// Assigned max bin, always the current Y of the loop if the Histogram is not 0.
+				MaxBin.x = (HistogramBin.x == 0 ? MaxBin.x : y + 1); 
+				MaxBin.y = (HistogramBin.y == 0 ? MaxBin.y : y + 1); 
+				MaxBin.z = (HistogramBin.z == 0 ? MaxBin.z : y + 1); 
+				MaxBin.w = (HistogramBin.w == 0 ? MaxBin.w : y + 1); 
 			}			
 		}
 		
-		MaxValue = max(MaxValue, MinValue); // Make sure Max was assigned and coherent with min
-		MinValue /= 256.0;
-		MaxValue /= 256.0;
+		MaxBin = max(MaxBin, MinBin); // Make sure Max was assigned and coherent with min
 		
-		Result[uint2(0,1)] = MinCount;//2nd row first pixel is the min count
-		Result[uint2(1,1)] = MaxCount;//2nd row second pixel is the max count
-		Result[uint2(2,1)] = MinValue;//2nd row third pixel is the min value
-		Result[uint2(3,1)] = MaxValue;//2nd row fourth pixel is the max value
+		// Normalize the min / max bin index
+		MinBin /= 256.0;
+		MaxBin /= 256.0;
+			
+		Result[uint2(0,1)] = MinCount;//2nd row first pixel is the min count in a bin (usually 0)
+		Result[uint2(1,1)] = MaxCount;//2nd row second pixel is the max count in a bin
+		Result[uint2(2,1)] = MinBin;//2nd row third pixel is the min bin index normalized 
+		Result[uint2(3,1)] = MaxBin;//2nd row fourth pixel is the max bin index noramlized
+		Result[uint2(4,1)] = TotalCount;//2nd row fifth pixel is the total count
 	}
-}
+
+}
\ No newline at end of file
diff --git a/Engine/Plugins/Experimental/TextureGraph/Source/TextureGraphEngine/Device/TaskPriorityQueue.h b/Engine/Plugins/Experimental/TextureGraph/Source/TextureGraphEngine/Device/TaskPriorityQueue.h
--- a/Engine/Plugins/Experimental/TextureGraph/Source/TextureGraphEngine/Device/TaskPriorityQueue.h
+++ b/Engine/Plugins/Experimental/TextureGraph/Source/TextureGraphEngine/Device/TaskPriorityQueue.h
@@ -136,6 +136,14 @@ public:
 		return vec;
 	}
 
+	void addBeginAndEnd(PtrType begin, PtrType end)
+	{
+		std::unique_lock<std::mutex> lock(Mutex);
+
+		PriorityQueue::c.insert(PriorityQueue::c.begin(), begin);
+		PriorityQueue::c.push_back(end);
+	}
+
 	void clear()
 	{
 		std::unique_lock<std::mutex> lock(Mutex);
diff --git a/Engine/Plugins/Experimental/TextureGraph/Source/TextureGraphEngine/Job/HistogramService.cpp b/Engine/Plugins/Experimental/TextureGraph/Source/TextureGraphEngine/Job/HistogramService.cpp
--- a/Engine/Plugins/Experimental/TextureGraph/Source/TextureGraphEngine/Job/HistogramService.cpp
+++ b/Engine/Plugins/Experimental/TextureGraph/Source/TextureGraphEngine/Job/HistogramService.cpp
@@ -50,6 +50,8 @@ AsyncJobResultPtr HistogramService::Tick()
 		Batch = nullptr;
 
 		//TextureGraphEngine::GetScheduler()->AddBatch(LastBatch);
+		LastBatch->SetCaptureRenderDoc(bCaptureNextBatch);
+		bCaptureNextBatch = false;
 
 		TextureGraphEngine::GetInstance()->GetScheduler()->GetObserverSource()->BatchAdded(LastBatch); // notify observer
 
@@ -71,4 +73,9 @@ void HistogramService::Stop()
 	check(IsInGameThread());
 	Batch = nullptr;
 }
- 
\ No newline at end of file
+ 
+
+void HistogramService::CaptureNextBatch()
+{
+	bCaptureNextBatch = true;
+}
diff --git a/Engine/Plugins/Experimental/TextureGraph/Source/TextureGraphEngine/Job/HistogramService.h b/Engine/Plugins/Experimental/TextureGraph/Source/TextureGraphEngine/Job/HistogramService.h
--- a/Engine/Plugins/Experimental/TextureGraph/Source/TextureGraphEngine/Job/HistogramService.h
+++ b/Engine/Plugins/Experimental/TextureGraph/Source/TextureGraphEngine/Job/HistogramService.h
@@ -16,6 +16,7 @@ class TEXTUREGRAPHENGINE_API HistogramService : public IdleService
 {
 private:
 	JobBatchPtr Batch;
+	bool bCaptureNextBatch = false;
 
 public:
 	HistogramService();
@@ -26,6 +27,8 @@ public:
 
 	void AddHistogramJob(MixUpdateCyclePtr Cycle,JobUPtr JobToAdd, int32 TargetID, UMixInterface* Mix);
 	JobBatchPtr GetOrCreateNewBatch(UMixInterface* mix);
+
+	void CaptureNextBatch();
 };
 
 typedef std::shared_ptr<HistogramService>	HistogramServicePtr;
diff --git a/Engine/Plugins/Experimental/TextureGraph/Source/TextureGraphEngine/Job/JobArgs.cpp b/Engine/Plugins/Experimental/TextureGraph/Source/TextureGraphEngine/Job/JobArgs.cpp
--- a/Engine/Plugins/Experimental/TextureGraph/Source/TextureGraphEngine/Job/JobArgs.cpp
+++ b/Engine/Plugins/Experimental/TextureGraph/Source/TextureGraphEngine/Job/JobArgs.cpp
@@ -79,6 +79,12 @@ bool JobArg_Blob::CanHandleTiles() const
 	return bCanHandleTiles;
 }
 
+JobArg_Blob& JobArg_Blob::WithNotHandleTiles()
+{
+	bCanHandleTiles = false;
+	return (*this);
+}
+
 void JobArg_Blob::SetForceNonTiledTransform(bool bInForceNonTiledTransform)
 {
 	bForceNonTiledTransform = bInForceNonTiledTransform;
@@ -111,6 +117,17 @@ bool JobArg_Blob::IsNeighborTiles() const
 	return bBindNeighborTiles;
 }
 
+JobArg_Blob& JobArg_Blob::WithArrayOfTiles()
+{
+	bBindArrayOfTiles = true;
+	return (*this);
+}
+
+bool JobArg_Blob::IsArrayOfTiles() const
+{
+	return bBindArrayOfTiles;
+}
+
 //////////////////////////////////////////////////////////////////////////
 TiledBlobPtr JobArg_Blob::GetRootBlob(JobArgBindInfo JobBindInfo) const
 {
@@ -229,6 +246,43 @@ AsyncJobArgResultPtr JobArg_Blob::Bind(JobArgBindInfo JobBindInfo)
 				return cti::make_ready_continuable(std::make_shared<JobArgResult>());
 			});
 	}
+	else if (IsArrayOfTiles())
+	{
+		return RootBlob->TransferTo(BindInfo.Dev)
+			.then([RootBlob, BindInfo, JobBindInfo](auto)
+				{
+					// Collect the individual tile textures from the root tiled BlobObj in an array
+					std::vector<TexPtr> Textures;
+					int NumX = RootBlob->Cols();
+					int NumY = RootBlob->Rows();
+
+					for (int Y = 0; Y < NumY; ++Y)
+					{
+						for (int X = 0; X < NumX; ++X)
+						{
+							if (RootBlob->IsValidTileIndex(X, Y))
+							{
+								auto Tile = RootBlob->GetTile(X, Y);
+								auto TileBuffer = std::static_pointer_cast<DeviceBuffer_FX>(Tile->GetBufferRef().GetPtr());
+								Textures.push_back(TileBuffer->GetTexture());
+							}
+							else
+							{
+								auto Black = TextureHelper::GetBlack()->GetTile(0, 0);
+								auto BlackBuffer = std::static_pointer_cast<DeviceBuffer_FX> (Black->GetBufferRef().GetPtr());
+								Textures.push_back(BlackBuffer->GetTexture());
+							}
+						}
+					}
+
+					// True binding of the array of textures to the Material
+					const RenderMaterial* Material = static_cast<const RenderMaterial*>(JobBindInfo.Transform.get());
+					check(Material);
+					Material->SetArrayTexture(*BindInfo.Target, Textures);
+
+					return cti::make_ready_continuable(std::make_shared<JobArgResult>());
+				});
+	}
 	else
 	{
 		BlobPtr BlobToBind = RootBlob;
diff --git a/Engine/Plugins/Experimental/TextureGraph/Source/TextureGraphEngine/Job/JobArgs.h b/Engine/Plugins/Experimental/TextureGraph/Source/TextureGraphEngine/Job/JobArgs.h
--- a/Engine/Plugins/Experimental/TextureGraph/Source/TextureGraphEngine/Job/JobArgs.h
+++ b/Engine/Plugins/Experimental/TextureGraph/Source/TextureGraphEngine/Job/JobArgs.h
@@ -102,6 +102,9 @@ protected:
 
 	bool							bBindNeighborTiles = false;			/// Bind the tile and the rign of neighbors in order to be able to go fetch border information OUT of the current tile bound
 
+	bool							bBindArrayOfTiles = false;			/// Bind ALL the tiles contained in this arg, filling the array of UTextures before execution of the job.
+
+
 	TiledBlobPtr					GetRootBlob(JobArgBindInfo JobBindInfo) const;
 
 public:
@@ -111,6 +114,8 @@ public:
 	virtual							~JobArg_Blob() override;
 	virtual void					SetHandleTiles(bool bInCanHandleTiles);
 	virtual bool					CanHandleTiles() const override;
+	JobArg_Blob&					WithNotHandleTiles(); // equivalent to call  SetHandleTiles(false)
+
 
 	virtual void					SetForceNonTiledTransform(bool bInForceNonTiledTransform);
 	virtual bool					ForceNonTiledTransform() const override;
@@ -121,6 +126,9 @@ public:
 	JobArg_Blob&					WithNeighborTiles();
 	bool							IsNeighborTiles() const;
 
+	JobArg_Blob&					WithArrayOfTiles();
+	bool							IsArrayOfTiles() const;
+
 	virtual AsyncJobArgResultPtr	Bind(JobArgBindInfo JobBindInfo) override;
 	virtual AsyncJobArgResultPtr	Unbind(JobArgBindInfo JobBindInfo) override;
 
@@ -387,6 +395,7 @@ public:
 //Blob can be fetched in HLSL using GetFullBlob(inout float 4) method. See AdjustUVGeneric.usf and TiledFetch_Combined.ush for help
 #define ARG_COMBINEDBLOB(v, n)				std::make_shared<JobArg_Blob_Combined>(v, n)	//With Custom SRV
 
+
 FORCEINLINE JobArgPtr				WithIgnoreHash(JobArgPtr Arg, bool bIgnoreHash = true)
 {
 	Arg->WithIgnoreHash(bIgnoreHash);
diff --git a/Engine/Plugins/Experimental/TextureGraph/Source/TextureGraphEngine/Job/JobBatch.cpp b/Engine/Plugins/Experimental/TextureGraph/Source/TextureGraphEngine/Job/JobBatch.cpp
--- a/Engine/Plugins/Experimental/TextureGraph/Source/TextureGraphEngine/Job/JobBatch.cpp
+++ b/Engine/Plugins/Experimental/TextureGraph/Source/TextureGraphEngine/Job/JobBatch.cpp
@@ -4,6 +4,7 @@
 #include "Model/Mix/MixUpdateCycle.h"
 #include "Device/FX/Device_FX.h"
 #include "Job.h"
+#include "Transform/Utility/T_RenderCapture.h"
 
 DEFINE_LOG_CATEGORY(LogBatch);
 DECLARE_CYCLE_STAT(TEXT("JobBatch_Exec"), STAT_JobBatch_Exec, STATGROUP_TextureGraphEngine);
@@ -236,8 +237,19 @@ AsyncJobResultPtr JobBatch::Exec(OnAllJobsDoneCallback callback)
 	if (callback)
 		OnAllJobsDoneCallbacks.push_back(std::move(callback));
 
+
+	if (IsCaptureRenderDoc())
+	{
+		JobUPtr BeginJob = T_BeginRenderCapture::CreateJob(Cycle, 0);
+		BeginJob->SetPriority((int32)E_Priority::kHighest, false);
+		JobUPtr EndJob = T_EndRenderCapture::CreateJob(Cycle, 0);
+
+		Queue.addBeginAndEnd(std::move(BeginJob), std::move(EndJob));
+	}
+
 	AllJobs = Queue.to_vector();
 
+
 	UE_LOG(LogBatch, VeryVerbose, TEXT("Running Batch: %llu [Count: %llu]"), BatchId, Queue.size());
 
 	/// We really need these to be in order. The reason being that the Job::Prepare phase (which happens inside
diff --git a/Engine/Plugins/Experimental/TextureGraph/Source/TextureGraphEngine/Profiling/RenderDoc/RenderDocManager.cpp b/Engine/Plugins/Experimental/TextureGraph/Source/TextureGraphEngine/Profiling/RenderDoc/RenderDocManager.cpp
--- a/Engine/Plugins/Experimental/TextureGraph/Source/TextureGraphEngine/Profiling/RenderDoc/RenderDocManager.cpp
+++ b/Engine/Plugins/Experimental/TextureGraph/Source/TextureGraphEngine/Profiling/RenderDoc/RenderDocManager.cpp
@@ -37,6 +37,12 @@ namespace TextureGraphEditor
 			TEXT("Captures the previous Job Batch and launches RenderDoc"),
 			FConsoleCommandDelegate::CreateRaw(this, &RenderDocManager::CapturePreviousBatch)
 		);
+
+		static FAutoConsoleCommand CCmdRenderDocCaptureHistogram = FAutoConsoleCommand(
+			TEXT("renderdoc.TextureGraph_CaptureNextBatchHistogram"),
+			TEXT("Captures the next Job Batch producing histogram and launches RenderDoc"),
+			FConsoleCommandDelegate::CreateRaw(this, &RenderDocManager::CaptureNextBatchHistogram)
+		);
 #endif 
 	}
 
@@ -59,6 +65,17 @@ namespace TextureGraphEditor
 #endif 
 	}
 
+	void RenderDocManager::CaptureNextBatchHistogram()
+	{
+#if TEXTUREGRAPH_RENDERDOC_ENABLED
+		HistogramServicePtr HistoService = TextureGraphEngine::GetScheduler()->GetHistogramService().lock();
+		if (HistoService)
+		{
+			HistoService->CaptureNextBatch();
+		}
+#endif 
+	}
+
 	void RenderDocManager::BeginCapture()
 	{
 #if TEXTUREGRAPH_RENDERDOC_ENABLED
diff --git a/Engine/Plugins/Experimental/TextureGraph/Source/TextureGraphEngine/Profiling/RenderDoc/RenderDocManager.h b/Engine/Plugins/Experimental/TextureGraph/Source/TextureGraphEngine/Profiling/RenderDoc/RenderDocManager.h
--- a/Engine/Plugins/Experimental/TextureGraph/Source/TextureGraphEngine/Profiling/RenderDoc/RenderDocManager.h
+++ b/Engine/Plugins/Experimental/TextureGraph/Source/TextureGraphEngine/Profiling/RenderDoc/RenderDocManager.h
@@ -20,6 +20,8 @@ namespace TextureGraphEditor
 		void												CapturePreviousBatch();
 		void												BeginCapture();
 		void												EndCapture();
+		void												CaptureNextBatchHistogram();
+
 };
 	typedef std::unique_ptr<RenderDocManager>				RenderDocManagerPtr;
 }
\ No newline at end of file
diff --git a/Engine/Plugins/Experimental/TextureGraph/Source/TextureGraphEngine/Transform/Mix/T_InvalidateTiles.cpp b/Engine/Plugins/Experimental/TextureGraph/Source/TextureGraphEngine/Transform/Mix/T_InvalidateTiles.cpp
--- a/Engine/Plugins/Experimental/TextureGraph/Source/TextureGraphEngine/Transform/Mix/T_InvalidateTiles.cpp
+++ b/Engine/Plugins/Experimental/TextureGraph/Source/TextureGraphEngine/Transform/Mix/T_InvalidateTiles.cpp
@@ -3,7 +3,6 @@
 #include "Device/FX/Device_FX.h"
 #include "TextureGraphEngine.h"
 #include "Job/JobBatch.h"
-#include "Profiling/RenderDoc/RenderDocManager.h"
 
 T_InvalidateTiles::T_InvalidateTiles() : BlobTransform(TEXT("T_InvalidateTiles"))
 {
@@ -24,9 +23,6 @@ AsyncTransformResultPtr T_InvalidateTiles::Exec(const TransformArgs& args)
 	/// by the time we get to this point!
 	check(IsInRenderingThread());
 
-	if (args.Cycle->GetBatch()->IsCaptureRenderDoc())
-		TextureGraphEngine::GetRenderDocManager()->BeginCapture();
-
 	/// We don't need a target, so it shouldn't be there. The job needs to 
 	/// handle this correctly
 	check(args.Target.expired());
diff --git a/Engine/Plugins/Experimental/TextureGraph/Source/TextureGraphEngine/Transform/Mix/T_UpdateTargets.cpp b/Engine/Plugins/Experimental/TextureGraph/Source/TextureGraphEngine/Transform/Mix/T_UpdateTargets.cpp
--- a/Engine/Plugins/Experimental/TextureGraph/Source/TextureGraphEngine/Transform/Mix/T_UpdateTargets.cpp
+++ b/Engine/Plugins/Experimental/TextureGraph/Source/TextureGraphEngine/Transform/Mix/T_UpdateTargets.cpp
@@ -5,7 +5,6 @@
 #include "2D/TargetTextureSet.h"
 #include "Profiling/StatGroup.h"
 #include "Job/JobBatch.h"
-#include "Profiling/RenderDoc/RenderDocManager.h"
 #include "Model/Mix/MixSettings.h"
 #include "Model/Mix/MixInterface.h"
 
@@ -56,9 +55,6 @@ AsyncTransformResultPtr T_UpdateTargets::Exec(const TransformArgs& Args)
 		}
 	}
 
-	if (Args.Cycle->GetBatch() && Args.Cycle->GetBatch()->IsCaptureRenderDoc())
-		TextureGraphEngine::GetRenderDocManager()->EndCapture();
-
 	return cti::make_ready_continuable(std::make_shared<TransformResult>());
 }
 
@@ -78,3 +74,4 @@ void T_UpdateTargets::Create(MixUpdateCyclePtr Cycle, int32 TargetId, bool Shoul
 		JobUPtr job = CreateJob(Cycle, TargetId, ShouldUpdate);
 	Cycle->AddJob(TargetId, std::move(job));
 }
+
diff --git a/Engine/Plugins/Experimental/TextureGraph/Source/TextureGraphEngine/Transform/Utility/T_RenderCapture.cpp b/Engine/Plugins/Experimental/TextureGraph/Source/TextureGraphEngine/Transform/Utility/T_RenderCapture.cpp
new file mode 100644
--- /dev/null
+++ b/Engine/Plugins/Experimental/TextureGraph/Source/TextureGraphEngine/Transform/Utility/T_RenderCapture.cpp
@@ -0,0 +1,83 @@
+// Copyright Epic Games, Inc. All Rights Reserved.
+#include "T_RenderCapture.h"
+#include "Device/FX/Device_FX.h"
+#include "TextureGraphEngine.h"
+#include "Job/JobBatch.h"
+#include "Profiling/RenderDoc/RenderDocManager.h"
+
+
+T_BeginRenderCapture::T_BeginRenderCapture() : BlobTransform(TEXT("T_BeginRenderCapture"))
+{
+}
+
+T_BeginRenderCapture::~T_BeginRenderCapture()
+{
+}
+
+Device* T_BeginRenderCapture::TargetDevice(size_t index) const
+{
+	return Device_FX::Get();
+}
+
+AsyncTransformResultPtr T_BeginRenderCapture::Exec(const TransformArgs& args)
+{
+	/// Device::Use should have ensured that we we're in the rendering thread
+	/// by the time we get to this point!
+	check(IsInRenderingThread());
+
+	if (args.Cycle->GetBatch()->IsCaptureRenderDoc())
+		TextureGraphEngine::GetRenderDocManager()->BeginCapture();
+	/// We don't need a target, so it shouldn't be there. The job needs to 
+/// handle this correctly
+	check(args.Target.expired());
+	check(args.Cycle);
+
+	MixUpdateCyclePtr cycle = args.Cycle;
+	cycle->GetTarget(args.TargetId)->InvalidateAllTiles();
+
+	return cti::make_ready_continuable(std::make_shared<TransformResult>());
+}
+
+JobUPtr T_BeginRenderCapture::CreateJob(MixUpdateCyclePtr cycle, int32 targetId)
+{
+	BlobTransformPtr transform = std::static_pointer_cast<BlobTransform>(std::make_shared<T_BeginRenderCapture>());
+	JobUPtr job = std::make_unique<Job>(cycle->GetMix(), targetId, transform);
+	job->SetPriority((int32)E_Priority::kHighest, false);
+	return job;
+}
+
+
+
+T_EndRenderCapture::T_EndRenderCapture() : BlobTransform(TEXT("T_EndRenderCapture"))
+{
+}
+
+T_EndRenderCapture::~T_EndRenderCapture()
+{
+}
+
+Device* T_EndRenderCapture::TargetDevice(size_t index) const
+{
+	return Device_FX::Get();
+}
+
+AsyncTransformResultPtr T_EndRenderCapture::Exec(const TransformArgs& args)
+{
+	/// Device::Use should have ensured that we we're in the rendering thread
+	/// by the time we get to this point!
+	check(IsInRenderingThread());
+
+	if (args.Cycle->GetBatch() && args.Cycle->GetBatch()->IsCaptureRenderDoc())
+		TextureGraphEngine::GetRenderDocManager()->EndCapture();
+
+	return cti::make_ready_continuable(std::make_shared<TransformResult>());
+}
+
+JobUPtr T_EndRenderCapture::CreateJob(MixUpdateCyclePtr cycle, int32 targetId)
+{
+	BlobTransformPtr transform = std::static_pointer_cast<BlobTransform>(std::make_shared<T_EndRenderCapture>());
+	JobUPtr job = std::make_unique<Job>(cycle->GetMix(), targetId, transform);
+//	job->SetPriority((int32)E_Priority::kHighest, false);
+	return job;
+}
+
diff --git a/Engine/Plugins/Experimental/TextureGraph/Source/TextureGraphEngine/Transform/Utility/T_RenderCapture.h b/Engine/Plugins/Experimental/TextureGraph/Source/TextureGraphEngine/Transform/Utility/T_RenderCapture.h
new file mode 100644
--- /dev/null
+++ b/Engine/Plugins/Experimental/TextureGraph/Source/TextureGraphEngine/Transform/Utility/T_RenderCapture.h
@@ -0,0 +1,42 @@
+// Copyright Epic Games, Inc. All Rights Reserved.
+#pragma once
+
+#include "CoreMinimal.h"
+#include "Transform/BlobTransform.h"
+#include "Job/Job.h"
+#include "Model/Mix/MixUpdateCycle.h"
+
+class TEXTUREGRAPHENGINE_API T_BeginRenderCapture : public BlobTransform
+{
+public:
+	T_BeginRenderCapture();
+	virtual							~T_BeginRenderCapture() override;
+
+	virtual Device* TargetDevice(size_t index) const override;
+	virtual AsyncTransformResultPtr	Exec(const TransformArgs& args) override;
+
+	virtual bool					GeneratesData() const override { return false; }
+
+	////////////////////////////////////////////////////////////////////////// 
+	/// Static functions
+	//////////////////////////////////////////////////////////////////////////
+	static JobUPtr					CreateJob(MixUpdateCyclePtr cycle, int32 targetId);
+};
+
+
+class TEXTUREGRAPHENGINE_API T_EndRenderCapture : public BlobTransform
+{
+public:
+	T_EndRenderCapture();
+	virtual							~T_EndRenderCapture() override;
+
+	virtual Device* TargetDevice(size_t index) const override;
+	virtual AsyncTransformResultPtr	Exec(const TransformArgs& args) override;
+
+	virtual bool					GeneratesData() const override { return false; }
+
+	////////////////////////////////////////////////////////////////////////// 
+	/// Static functions
+	//////////////////////////////////////////////////////////////////////////
+	static JobUPtr					CreateJob(MixUpdateCyclePtr cycle, int32 targetId);
+};
\ No newline at end of file
diff --git a/Engine/Plugins/Experimental/TextureGraph/Source/TextureGraphEngine/Transform/Utility/T_TextureHistogram.cpp b/Engine/Plugins/Experimental/TextureGraph/Source/TextureGraphEngine/Transform/Utility/T_TextureHistogram.cpp
--- a/Engine/Plugins/Experimental/TextureGraph/Source/TextureGraphEngine/Transform/Utility/T_TextureHistogram.cpp
+++ b/Engine/Plugins/Experimental/TextureGraph/Source/TextureGraphEngine/Transform/Utility/T_TextureHistogram.cpp
@@ -1,5 +1,6 @@
 // Copyright Epic Games, Inc. All Rights Reserved.
 #include "T_TextureHistogram.h"
+
 #include "TextureGraphEngine.h"
 #include "Job/JobArgs.h"
 #include "TextureGraphEngine.h"
@@ -9,67 +10,237 @@
 #include "Helper/MathUtils.h"
 #include "Helper/GraphicsUtil.h"
 #include "Device/FX/DeviceBuffer_FX.h"
+#include "Device/FX/Device_FX.h"
+
 #include "Engine/TextureRenderTarget2D.h"
 #include "Engine/Texture2D.h"
 #include "Device/DeviceManager.h"
+#include "Device/FX/Device_FX.h"
 #include "Device/Mem/Device_Mem.h"
 #include "Data/Blobber.h"
 #include "2D/Tex.h"
+#include "TextureResource.h"
 #include "FxMat/MaterialManager.h"
 #include "T_MinMax.h"
 
-IMPLEMENT_GLOBAL_SHADER(CSH_Histogram, "/Plugin/TextureGraph/Utils/Histogram_comp.usf", "CSH_Histogram", SF_Compute);
 
 static constexpr int NumBins = 256;
 
-T_TextureHistogram::T_TextureHistogram()
+class CSH_HistogramPerTile : public CmpSH_Base<16, 16, 1>
 {
-}
+public:
+	DECLARE_GLOBAL_SHADER(CSH_HistogramPerTile);
+	SHADER_USE_PARAMETER_STRUCT(CSH_HistogramPerTile, CmpSH_Base);
+	BEGIN_SHADER_PARAMETER_STRUCT(FParameters, )
+		SHADER_PARAMETER_UAV(RWTexture2D<float4>, Result)   // Unused but still there for compatibility
+		SHADER_PARAMETER_UAV(RWStructuredBuffer<uint4>, TilesHistogramUAV)
+		SHADER_PARAMETER_TEXTURE(Texture2D, SourceTiles)
+		SHADER_PARAMETER(FIntVector4, InvocationDim)
+		END_SHADER_PARAMETER_STRUCT()
 
-T_TextureHistogram::~T_TextureHistogram()
+public:
+	static bool ShouldCompilePermutation(const FGlobalShaderPermutationParameters& Parameters)
+	{
+		return IsPCPlatform(Parameters.Platform) && EnumHasAllFlags(Parameters.Flags, EShaderPermutationFlags::HasEditorOnlyData)
+			&& IsFeatureLevelSupported(Parameters.Platform, ERHIFeatureLevel::SM5);
+	}
+};
+
+class CSH_Histogram : public CmpSH_Base<256, 1, 1>
 {
-}
+public:
+	DECLARE_GLOBAL_SHADER(CSH_Histogram);
+	SHADER_USE_PARAMETER_STRUCT(CSH_Histogram, CmpSH_Base);
+	BEGIN_SHADER_PARAMETER_STRUCT(FParameters, )
+		SHADER_PARAMETER_UAV(RWTexture2D<float4>, Result)
+		SHADER_PARAMETER_SRV(StructuredBuffer<uint4>, TilesHistogramSRV)
+		SHADER_PARAMETER(FIntVector4, InvocationDim)
+		END_SHADER_PARAMETER_STRUCT()
+
+public:
+	static bool ShouldCompilePermutation(const FGlobalShaderPermutationParameters& Parameters)
+	{
+		return IsPCPlatform(Parameters.Platform) && EnumHasAllFlags(Parameters.Flags, EShaderPermutationFlags::HasEditorOnlyData)
+			&& IsFeatureLevelSupported(Parameters.Platform, ERHIFeatureLevel::SM5);
+	}
+};
+
 
-RenderMaterial_FXPtr T_TextureHistogram::CreateMaterial_Histogram(FString Name, FString OutputId, const CSH_Histogram::FPermutationDomain& cmpshPermutationDomain,
-	int NumThreadsX, int NumThreadsY, int NumThreadsZ)
+IMPLEMENT_GLOBAL_SHADER(CSH_HistogramPerTile, "/Plugin/TextureGraph/Utils/Histogram_comp.usf", "CSH_HistogramPerTile", SF_Compute);
+IMPLEMENT_GLOBAL_SHADER(CSH_Histogram, "/Plugin/TextureGraph/Utils/Histogram_comp.usf", "CSH_Histogram", SF_Compute);
+
+
+
+/////////////////////////////////////////////////////////////////////////
+
+//template <typename CmpSH_Type>
+class FxMaterial_Histogram : public FxMaterial_Compute<CSH_HistogramPerTile>
+{
+	FRWBufferStructured TilesHistogramBuffer;
+	FTextureRWBuffer TilesHistogramTexture;
+public:
+	// type Name for the permutation domain
+	using CmpSHPermutationDomain = typename CSH_HistogramPerTile::FPermutationDomain;
+
+	FxMaterial_Histogram(FString outputId, const CmpSHPermutationDomain* permDomain = nullptr,
+		int numThreadsX = FxMaterial_Compute<CSH_HistogramPerTile>::GDefaultNumThreadsXY, int numThreadsY = FxMaterial_Compute<CSH_HistogramPerTile>::GDefaultNumThreadsXY, int numThreadsZ = 1,
+		FUnorderedAccessViewRHIRef unorderedAccessView = nullptr)
+		: FxMaterial_Compute<CSH_HistogramPerTile>(outputId, permDomain, numThreadsX, numThreadsY, numThreadsZ, unorderedAccessView)
+	{
+
+	}
+
+	virtual void Blit(FRHICommandListImmediate& RHI, FRHITexture* Target, const RenderMesh* MeshObj, int32 TargetId, FGraphicsPipelineStateInitializer* PSO = nullptr) override
+	{
+		QUICK_SCOPE_CYCLE_COUNTER(STAT_ShaderPlugin_ComputeShader); // Used to gather CPU profiling data for the UE4 session frontend
+		SCOPED_DRAW_EVENT(RHI, ShaderPlugin_Compute); // Used to profile GPU activity and add metadata to be consumed by for example RenderDoc
+
+		CSH_HistogramPerTile::FParameters& PerTileParams = Params;
+		CSH_Histogram::FParameters FinalParams;
+
+		// Collect the Tile textures contributing to the whole Source texture
+		TArray<FRHITexture*> TextureTiles;
+		int32 TilesCount = 0;
+		if (Textures.size() == 1 && Textures[0].tiles.size())
+		{
+			TilesCount = Textures[0].tiles.size();
+			TextureTiles.Reserve(TilesCount);
+
+			int i = 0;
+			for (const auto& ti : Textures[0].tiles)
+			{
+				FRHITexture* TextureRHI = ti->GetResource()->TextureRHI;
+				TextureTiles.Add(TextureRHI);
+				i++;
+			}
+		}
+
+		if (TilesCount == 0)
+		{
+			TilesCount = 1;
+		}
+
+		// Allocate and bind the buffer used by the shader to store intermediate results of each pass / tile 
+		{
+			TilesHistogramBuffer.Initialize(RHI, TEXT("HistogramBuffer"), sizeof(uint32) * 4, NumBins * TilesCount);
+			RHI.ClearUAVUint(TilesHistogramBuffer.UAV, FUintVector4(0, 0, 0, 0));
+
+			PerTileParams.TilesHistogramUAV = TilesHistogramBuffer.UAV;
+			FinalParams.TilesHistogramSRV = TilesHistogramBuffer.SRV;
+		}
+
+		// Allocate and set the result UAV only used in the final pass / param 
+		if (!OutputId.IsEmpty())
+		{
+			FUnorderedAccessViewRHIRef RenderTargetUAV = UnorderedAccessView;
+
+			if (!RenderTargetUAV)
+			{
+				RenderTargetUAV = RHI.CreateUnorderedAccessView(Target);
+				RHI.Transition(FRHITransitionInfo(Target, ERHIAccess::Unknown, ERHIAccess::UAVMask));
+			}
+
+			FinalParams.Result = RenderTargetUAV;
+		}
+
+		// Same as the standard FxMaterial_Compute<> blit but adjusted so we do as many  passes as tiles
+		{
+			TShaderMapRef<CSH_HistogramPerTile> computeShaderPerTile(GetGlobalShaderMap(GMaxRHIFeatureLevel), PermutationDomain);
+			FIntVector groupSize = computeShaderPerTile->ThreadGroupSize();
+			
+			PerTileParams.InvocationDim.X = TextureTiles.Num();
+			PerTileParams.InvocationDim.Y = 0;
+			PerTileParams.InvocationDim.Z = 0;
+			PerTileParams.InvocationDim.W = 0;
+
+			int NumTiles = TextureTiles.Num();
+			for (int i = 0; i < NumTiles; ++i)
+			{
+				// bind the pass's tile and increment the tile index for each pass
+				PerTileParams.SourceTiles = TextureTiles[i];
+				PerTileParams.InvocationDim.Y = i; 
+				PerTileParams.InvocationDim.W = (i + 1 == TextureTiles.Num());
+
+				FIntPoint TileDim = TextureTiles[i]->GetSizeXY();
+
+				FComputeShaderUtils::Dispatch(RHI, computeShaderPerTile, PerTileParams,
+					FIntVector(
+						FMath::DivideAndRoundUp(TileDim.X, groupSize.X),
+						FMath::DivideAndRoundUp(TileDim.Y, groupSize.Y),
+						FMath::DivideAndRoundUp(1, groupSize.Z)
+					)
+				);
+
+				RHI.Transition(FRHITransitionInfo(TilesHistogramBuffer.UAV, ERHIAccess::UAVCompute, ERHIAccess::UAVCompute));
+			}
+
+			// All the per tiles histograms have been collected
+			// transition to srv read
+			RHI.Transition(FRHITransitionInfo(TilesHistogramBuffer.UAV, ERHIAccess::UAVCompute, ERHIAccess::SRVCompute));
+
+			// Last pass
+			TShaderMapRef<CSH_Histogram> computeShader(GetGlobalShaderMap(GMaxRHIFeatureLevel), PermutationDomain);
+			groupSize = computeShader->ThreadGroupSize();
+
+			FinalParams.InvocationDim.X = TextureTiles.Num();
+			FinalParams.InvocationDim.Y = 0;
+			FinalParams.InvocationDim.Z = 0;
+			FinalParams.InvocationDim.W = 0;
+
+			// Dispatch ONE group 256 threads
+			FComputeShaderUtils::Dispatch(RHI, computeShader, FinalParams, FIntVector(1, 1, 1));
+
+
+		}
+		// UAV target has been rendered, transition to the default SRV state for read
+		RHI.Transition(FRHITransitionInfo(Target, ERHIAccess::UAVMask, ERHIAccess::SRVMask));
+	}
+};
+
+T_TextureHistogram::T_TextureHistogram()
 {
-	std::shared_ptr<FxMaterial_Histogram<CSH_Histogram>> Mat = std::make_shared<FxMaterial_Histogram<CSH_Histogram>>(OutputId, &cmpshPermutationDomain, NumThreadsX, NumThreadsY, NumThreadsZ);
-	return std::make_shared<RenderMaterial_FX>(Name, std::static_pointer_cast<FxMaterial>(Mat));
 }
 
+T_TextureHistogram::~T_TextureHistogram()
+{
+}
 
-TiledBlobPtr T_TextureHistogram::CreateJobAndResult(JobUPtr& OutJob, MixUpdateCyclePtr Cycle, TiledBlobPtr SourceTex, int32 TargetId)
+TiledBlobPtr T_TextureHistogram::CreateJobAndResult(JobUPtr& OutJob, MixUpdateCyclePtr InCycle, TiledBlobPtr SourceTexture, int32 InTargetId)
 {
-	check(!SourceTex->HasHistogram())
+	check(!SourceTexture->HasHistogram())
 
 	CSH_Histogram::FPermutationDomain PermutationVector;
 
-	FIntVector4 SrcDimensions(SourceTex->GetWidth(), SourceTex->GetHeight(), 1, 1);
+	FIntVector4 SrcDimensions(SourceTexture->GetWidth(), SourceTexture->GetHeight(), 1, 1);
 
-	FString Name = FString::Printf(TEXT("[%s].[%d] Histogram"), *SourceTex->DisplayName(), TargetId);
-	RenderMaterial_FXPtr Transform = T_TextureHistogram::CreateMaterial_Histogram(TEXT("T_Histogram"), TEXT("Result"), PermutationVector, SrcDimensions.X, SrcDimensions.Y, 1);
+	FString Name = FString::Printf(TEXT("[%s].[%d] Histogram"), *SourceTexture->DisplayName(), InTargetId);
 
-	OutJob = std::make_unique<Job>(Cycle->GetMix(), TargetId, std::static_pointer_cast<BlobTransform>(Transform));
-	OutJob->AddArg(ARG_BLOB(SourceTex, "SourceTiles"));
+	// Regular RendermaterialFX using our custom FXMaterial for histogram
+	std::shared_ptr<FxMaterial_Histogram> FxMat = std::make_shared<FxMaterial_Histogram>(TEXT("Result"), &PermutationVector, SrcDimensions.X, SrcDimensions.Y, 1);
+	RenderMaterial_FXPtr Transform = std::make_shared<RenderMaterial_FX>(TEXT("T_Histogram"), std::static_pointer_cast<FxMaterial>(FxMat));
+	
+	OutJob = std::make_unique<Job>(InCycle->GetMix(), InTargetId, std::static_pointer_cast<BlobTransform>(Transform));
+	auto SourceArg = ARG_BLOB(SourceTexture, "SourceTiles");
+	SourceArg->WithArrayOfTiles(); // THe source tiles will be bound as an array of tiles and the FXMat will explicitely bind the correct tile for the various passes
+	OutJob->AddArg(SourceArg);
 
 	BufferDescriptor Desc;
 	Desc.Width = NumBins;
 	Desc.Height = 2;
 	Desc.Format = BufferFormat::Float;// BufferFormat::Int;
 	Desc.ItemsPerPoint = 4;
-	Desc.Name = FString::Printf(TEXT("Histogram - %s"), *SourceTex->Name());
+	Desc.Name = FString::Printf(TEXT("Histogram - %s"), *SourceTexture->Name());
 	Desc.AllowUAV();
 
-	OutJob->SetTiled(false);
+	OutJob->SetTiled(false); // Run single tile for the texture and in the transform::exec, unroll the loop for each tiles of the source
 
 	TiledBlobPtr Result = OutJob->InitResult(Name, &Desc, 1, 1);
-
 	Result->MakeSingleBlob();
 
-	if (!SourceTex->HasHistogram())
+	if (!SourceTexture->HasHistogram())
 	{
 		//setting it as the histogram of source so it is retained untill the life cycle of source blob.
-		SourceTex->SetHistogram(Result);
+		SourceTexture->SetHistogram(Result);
 	}
 	return Result;
 }
diff --git a/Engine/Plugins/Experimental/TextureGraph/Source/TextureGraphEngine/Transform/Utility/T_TextureHistogram.h b/Engine/Plugins/Experimental/TextureGraph/Source/TextureGraphEngine/Transform/Utility/T_TextureHistogram.h
--- a/Engine/Plugins/Experimental/TextureGraph/Source/TextureGraphEngine/Transform/Utility/T_TextureHistogram.h
+++ b/Engine/Plugins/Experimental/TextureGraph/Source/TextureGraphEngine/Transform/Utility/T_TextureHistogram.h
@@ -8,55 +8,6 @@
 #include "Model/Mix/MixUpdateCycle.h"
 #include <DataDrivenShaderPlatformInfo.h>
 
-class CSH_Histogram : public CmpSH_Base<16, 16, 1>
-{
-public:
-	DECLARE_GLOBAL_SHADER(CSH_Histogram);
-	SHADER_USE_PARAMETER_STRUCT(CSH_Histogram, CmpSH_Base);
-	BEGIN_SHADER_PARAMETER_STRUCT(FParameters, )
-		SHADER_PARAMETER_TEXTURE(Texture2D, SourceTiles)
-		SHADER_PARAMETER_UAV(RWTexture2D<float4>, Result)
-		SHADER_PARAMETER_UAV(RWStructuredBuffer<uint4>, Histogram)
-	END_SHADER_PARAMETER_STRUCT()
-
-public:
-	static bool ShouldCompilePermutation(const FGlobalShaderPermutationParameters& Parameters)
-	{
-		return IsPCPlatform(Parameters.Platform) && EnumHasAllFlags(Parameters.Flags, EShaderPermutationFlags::HasEditorOnlyData)
-			&& IsFeatureLevelSupported(Parameters.Platform, ERHIFeatureLevel::SM5);
-	}
-};
-
-template <typename CmpSH_Type>
-class FxMaterial_Histogram : public FxMaterial_Compute<CmpSH_Type>
-{
-	FRWBufferStructured Buffer;
-public:
-	// type Name for the permutation domain
-	using CmpSHPermutationDomain = typename CmpSH_Type::FPermutationDomain;
-
-	 FxMaterial_Histogram(FString outputId, const CmpSHPermutationDomain* permDomain = nullptr,
-		int numThreadsX = FxMaterial_Compute<CmpSH_Type>::GDefaultNumThreadsXY, int numThreadsY = FxMaterial_Compute<CmpSH_Type>::GDefaultNumThreadsXY, int numThreadsZ = 1,
-		FUnorderedAccessViewRHIRef unorderedAccessView = nullptr)
-		: FxMaterial_Compute<CmpSH_Type>(outputId, permDomain, numThreadsX, numThreadsY, numThreadsZ, unorderedAccessView)
-	{
-
-	}
-
-	virtual void Blit(FRHICommandListImmediate& RHI, FRHITexture* Target, const RenderMesh* MeshObj, int32 TargetId, FGraphicsPipelineStateInitializer* PSO = nullptr) override
-	{
-		SetBuffer(RHI);
-
-		FxMaterial_Compute<CmpSH_Type>::Blit(RHI, Target, MeshObj, TargetId, PSO);
-	}
-
-	void SetBuffer(FRHICommandListImmediate& RHI)
-	{
-		Buffer.Initialize(RHI, TEXT("HistogramRBuffer"), sizeof(uint32) * 4, 256, BUF_UnorderedAccess | BUF_ShaderResource);
-		RHI.ClearUAVUint(Buffer.UAV, FUintVector4(0, 0, 0, 0));
-		FxMaterial_Compute<CmpSH_Type>::Params.Histogram = Buffer.UAV;
-	}
-};
 
 /**
  * T_TextureHistogram Transform
@@ -72,6 +23,5 @@ public:
 	static TiledBlobPtr	CreateOnService(UMixInterface* InMix, TiledBlobPtr SourceTex, int32 TargetId);
 
 private:
-	static RenderMaterial_FXPtr CreateMaterial_Histogram(FString Name, FString OutputId, const CSH_Histogram::FPermutationDomain& cmpshPermutationDomain,int NumThreadsX, int NumThreadsY, int NumThreadsZ = 1);
 	static TiledBlobPtr			CreateJobAndResult(JobUPtr& OutJob, MixUpdateCyclePtr Cycle, TiledBlobPtr Histogram, int32 TargetId);
 };
