Crash at exit and some minor naming inconsistencies.

diff --git a/Engine/Plugins/Experimental/TextureGraph/Source/TextureGraph/Private/Expressions/Input/TG_Expression_MaterialFunction.cpp b/Engine/Plugins/Experimental/TextureGraph/Source/TextureGraph/Private/Expressions/Input/TG_Expression_MaterialFunction.cpp
--- a/Engine/Plugins/Experimental/TextureGraph/Source/TextureGraph/Private/Expressions/Input/TG_Expression_MaterialFunction.cpp
+++ b/Engine/Plugins/Experimental/TextureGraph/Source/TextureGraph/Private/Expressions/Input/TG_Expression_MaterialFunction.cpp
@@ -272,7 +272,7 @@ UMaterialInterface* UTG_Expression_MaterialFunction::CreateMaterialReference()
 	MaterialFunction->GetInputsAndOutputs(FunctionInputs, FunctionOutputs);
 	UE_LOG(LogTextureGraph, Log, TEXT("Function <%s>\n"), *MaterialFunction->GetName());
 
-	// For all the inputs of the function, create a matching material parameter that can be exposed as an input pin of the texture script 
+	// For all the inputs of the function, create a matching material parameter that can be exposed as an input pin of the texture graph 
 	for (int i = 0; i < FunctionInputs.Num(); ++i)
 	{
 		bool SupportedInput = true;
diff --git a/Engine/Plugins/Experimental/TextureGraph/Source/TextureGraph/Public/Blueprint/TG_BlueprintFunctionLibrary.h b/Engine/Plugins/Experimental/TextureGraph/Source/TextureGraph/Public/Blueprint/TG_BlueprintFunctionLibrary.h
--- a/Engine/Plugins/Experimental/TextureGraph/Source/TextureGraph/Public/Blueprint/TG_BlueprintFunctionLibrary.h
+++ b/Engine/Plugins/Experimental/TextureGraph/Source/TextureGraph/Public/Blueprint/TG_BlueprintFunctionLibrary.h
@@ -19,45 +19,45 @@ public:
 
 	/** Sets a Texture parameter value on the TextureGraph instance. Logs if ParameterName is invalid. */
 	UFUNCTION(BlueprintCallable, Category="TextureGraph", meta=(Keywords="SetTextureParameterValue", WorldContext="WorldContextObject"))
-	static void SetTextureParameterValue(UObject* WorldContextObject, UTextureGraph* TextureScript, FName ParameterName, UTexture* ParameterValue);
+	static void SetTextureParameterValue(UObject* WorldContextObject, UTextureGraph* TextureGraph, FName ParameterName, UTexture* ParameterValue);
 
 	/** Gets a texture parameter value from the TextureGraph instance. Logs if ParameterName is invalid. */
 	UFUNCTION(BlueprintCallable, Category="TextureGraph", meta=(Keywords="GetTextureParameterValue", WorldContext="WorldContextObject"))
-	static UTexture* GetTextureParameterValue(UObject* WorldContextObject, UTextureGraph* TextureScript, FName ParameterName);
+	static UTexture* GetTextureParameterValue(UObject* WorldContextObject, UTextureGraph* TextureGraph, FName ParameterName);
 
 	/** Sets a scalar parameter value on the TextureGraph instance. Logs if ParameterName is invalid. */
 	UFUNCTION(BlueprintCallable, Category = "TextureGraph", meta = (Keywords = "SetScalarParameterValue", WorldContext = "WorldContextObject"))
-	static void SetScalarParameterValue(UObject* WorldContextObject, UTextureGraph* TextureScript, FName ParameterName, float ParameterValue);
+	static void SetScalarParameterValue(UObject* WorldContextObject, UTextureGraph* TextureGraph, FName ParameterName, float ParameterValue);
 
 	/** Gets a scalar parameter value from the TextureGraph instance. Logs if ParameterName is invalid. */
 	UFUNCTION(BlueprintCallable, Category = "TextureGraph", meta = (Keywords = "GetScalarParameterValue", WorldContext = "WorldContextObject"))
-	static float GetScalarParameterValue(UObject* WorldContextObject, UTextureGraph* TextureScript, FName ParameterName);
+	static float GetScalarParameterValue(UObject* WorldContextObject, UTextureGraph* TextureGraph, FName ParameterName);
 
 	/** Sets a Vector parameter value on the TextureGraph instance. Logs if ParameterName is invalid. */
 	UFUNCTION(BlueprintCallable, Category = "TextureGraph", meta = (Keywords = "SetVectorParameterValue", WorldContext = "WorldContextObject"))
-	static void SetVectorParameterValue(UObject* WorldContextObject, UTextureGraph* TextureScript, FName ParameterName, FLinearColor ParameterValue);
+	static void SetVectorParameterValue(UObject* WorldContextObject, UTextureGraph* TextureGraph, FName ParameterName, FLinearColor ParameterValue);
 
 	/** Gets a Vector parameter value from the TextureGraph instance. Logs if ParameterName is invalid. */
 	UFUNCTION(BlueprintCallable, Category = "TextureGraph", meta = (Keywords = "GetVectorParameterValue", WorldContext = "WorldContextObject"))
-	static FLinearColor GetVectorParameterValue(UObject* WorldContextObject, UTextureGraph* TextureScript, FName ParameterName);
+	static FLinearColor GetVectorParameterValue(UObject* WorldContextObject, UTextureGraph* TextureGraph, FName ParameterName);
 
 	/** Sets a color parameter value on the TextureGraph instance. Logs if ParameterName is invalid. */
 	UFUNCTION(BlueprintCallable, Category = "TextureGraph", meta = (Keywords = "SetColorParameterValue", WorldContext = "WorldContextObject"))
-	static void SetColorParameterValue(UObject* WorldContextObject, UTextureGraph* TextureScript, FName ParameterName, FLinearColor ParameterValue);
+	static void SetColorParameterValue(UObject* WorldContextObject, UTextureGraph* TextureGraph, FName ParameterName, FLinearColor ParameterValue);
 
 	/** Gets a color parameter value from the TextureGraph instance. Logs if ParameterName is invalid. */
 	UFUNCTION(BlueprintCallable, Category = "TextureGraph", meta = (Keywords = "GetColorParameterValue", WorldContext = "WorldContextObject"))
-	static FLinearColor GetColorParameterValue(UObject* WorldContextObject, UTextureGraph* TextureScript, FName ParameterName);
+	static FLinearColor GetColorParameterValue(UObject* WorldContextObject, UTextureGraph* TextureGraph, FName ParameterName);
 
 	/** Sets a FTG_OutputSettings parameter value on the TextureGraph instance. Logs if ParameterName is invalid. */
 	UFUNCTION(BlueprintCallable, Category = "TextureGraph", meta = (Keywords = "SetSettingsParameterValue", WorldContext = "WorldContextObject"))
-	static void SetSettingsParameterValue(UObject* WorldContextObject, UTextureGraph* TextureScript, FName ParameterName, int Width,int Height, FName FileName = "None",
+	static void SetSettingsParameterValue(UObject* WorldContextObject, UTextureGraph* TextureGraph, FName ParameterName, int Width,int Height, FName FileName = "None",
 		FName Path = "None", ETG_TextureFormat Format = ETG_TextureFormat::BGRA8 , ETG_TexturePresetType TextureType = ETG_TexturePresetType::None,
 		TextureGroup LODTextureGroup = TextureGroup::TEXTUREGROUP_World, TextureCompressionSettings Compression = TextureCompressionSettings::TC_Default, bool SRGB = false);
 
 	/** Gets a FTG_OutputSettings parameter value from the TextureGraph instance. Logs if ParameterName is invalid. */
 	UFUNCTION(BlueprintCallable, Category = "TextureGraph", meta = (Keywords = "GetOutputSettingsParameterValue", WorldContext = "WorldContextObject"))
-	static FTG_OutputSettings GetSettingsParameterValue(UObject* WorldContextObject, UTextureGraph* TextureScript, FName ParameterName , int& Width, int& Height);
+	static FTG_OutputSettings GetSettingsParameterValue(UObject* WorldContextObject, UTextureGraph* TextureGraph, FName ParameterName , int& Width, int& Height);
 
 	static void AddParamWarning(FName ParamName, UObject* ObjectPtr, FString FunctionName);
 	static void AddError(UObject* ObjectPtr, FString FunctionName, FString Error);
diff --git a/Engine/Plugins/Experimental/TextureGraph/Source/TextureGraph/Public/Expressions/Filter/TG_Expression_EdgeDetect.h b/Engine/Plugins/Experimental/TextureGraph/Source/TextureGraph/Public/Expressions/Filter/TG_Expression_EdgeDetect.h
--- a/Engine/Plugins/Experimental/TextureGraph/Source/TextureGraph/Public/Expressions/Filter/TG_Expression_EdgeDetect.h
+++ b/Engine/Plugins/Experimental/TextureGraph/Source/TextureGraph/Public/Expressions/Filter/TG_Expression_EdgeDetect.h
@@ -27,7 +27,7 @@ public:
 	UPROPERTY(meta = (TGType = "TG_Input"))
 	FTG_Texture Input;
 
-	virtual FText GetTooltipText() const override { return FText::FromString(TEXT("Applies a blur filter with a given strength.")); }
+	virtual FText GetTooltipText() const override { return FText::FromString(TEXT("Detects edges (areas of sharp color change) from the input.")); }
 	virtual void Evaluate(FTG_EvaluationContext* InContext) override;
 };
 
diff --git a/Engine/Plugins/Experimental/TextureGraph/Source/TextureGraph/Public/Expressions/TG_Expression.h b/Engine/Plugins/Experimental/TextureGraph/Source/TextureGraph/Public/Expressions/TG_Expression.h
--- a/Engine/Plugins/Experimental/TextureGraph/Source/TextureGraph/Public/Expressions/TG_Expression.h
+++ b/Engine/Plugins/Experimental/TextureGraph/Source/TextureGraph/Public/Expressions/TG_Expression.h
@@ -52,7 +52,7 @@ public:
 	virtual FTG_Name GetDefaultName() const;
 	virtual FName GetCategory() const { return TG_Category::Default; }
 	virtual FTG_SignaturePtr GetSignature() const { return nullptr; }
-	virtual FText GetTooltipText() const { return FText::FromString(TEXT("Texture Scripting Node")); } 
+	virtual FText GetTooltipText() const { return FText::FromString(TEXT("Texture Graph Node")); } 
 
 	virtual void SetTitleName(FName NewName) {}
 	virtual FName GetTitleName() const { return GetDefaultName(); }
diff --git a/Engine/Plugins/Experimental/TextureGraph/Source/TextureGraphEditor/Private/EdGraph/TG_EdGraphSchema.cpp b/Engine/Plugins/Experimental/TextureGraph/Source/TextureGraphEditor/Private/EdGraph/TG_EdGraphSchema.cpp
--- a/Engine/Plugins/Experimental/TextureGraph/Source/TextureGraphEditor/Private/EdGraph/TG_EdGraphSchema.cpp
+++ b/Engine/Plugins/Experimental/TextureGraph/Source/TextureGraphEditor/Private/EdGraph/TG_EdGraphSchema.cpp
@@ -472,7 +472,7 @@ void UTG_EdGraphSchema::GetAssetsGraphHoverMessage(const TArray<FAssetData>& Ass
 		{
 			OutOkIcon = true;
 		}
-		else if(UTextureGraph* TextureScript = Cast<UTextureGraph>(Asset); TextureScript != nullptr)
+		else if(UTextureGraph* TextureGraph = Cast<UTextureGraph>(Asset); TextureGraph != nullptr)
 		{
 			OutOkIcon = true;
 		}
diff --git a/Engine/Plugins/Experimental/TextureGraph/Source/TextureGraphEditor/Private/TG_EditorCommands.cpp b/Engine/Plugins/Experimental/TextureGraph/Source/TextureGraphEditor/Private/TG_EditorCommands.cpp
--- a/Engine/Plugins/Experimental/TextureGraph/Source/TextureGraphEditor/Private/TG_EditorCommands.cpp
+++ b/Engine/Plugins/Experimental/TextureGraph/Source/TextureGraphEditor/Private/TG_EditorCommands.cpp
@@ -11,7 +11,7 @@
 FTG_EditorCommands::FTG_EditorCommands()
 	: TCommands<FTG_EditorCommands>(
 		"TG_Editor",
-		NSLOCTEXT("Contexts", "TG_Editor", "Texture Scripting Editor"),
+		NSLOCTEXT("Contexts", "TG_Editor", "Texture Graph Editor"),
 		NAME_None,
 
 		FTG_Style::Get().GetStyleSetName())
diff --git a/Engine/Plugins/Experimental/TextureGraph/Source/TextureGraphEngine/2D/Tex.cpp b/Engine/Plugins/Experimental/TextureGraph/Source/TextureGraphEngine/2D/Tex.cpp
--- a/Engine/Plugins/Experimental/TextureGraph/Source/TextureGraphEngine/2D/Tex.cpp
+++ b/Engine/Plugins/Experimental/TextureGraph/Source/TextureGraphEngine/2D/Tex.cpp
@@ -261,7 +261,12 @@ void Tex::FreeRT(UTextureRenderTarget2D** RT)
 void Tex::Free()
 {
 	check(IsInGameThread());
-	/// IMPORTANT: Do not free the _image in this function
+
+#if WITH_EDITOR
+	/// Make it insensitive to the order of mixer engine destruction
+	if (TextureGraphEngine::IsDestroying())
+		return;
+#endif 
 	FreeTexture(ToRawPtr(MutableView(Texture)));
 	FreeRT(ToRawPtr(MutableView(RT)));
 }
diff --git a/Engine/Plugins/Experimental/TextureGraph/Source/TextureGraphEngine/Device/FX/Device_FX.cpp b/Engine/Plugins/Experimental/TextureGraph/Source/TextureGraphEngine/Device/FX/Device_FX.cpp
--- a/Engine/Plugins/Experimental/TextureGraph/Source/TextureGraphEngine/Device/FX/Device_FX.cpp
+++ b/Engine/Plugins/Experimental/TextureGraph/Source/TextureGraphEngine/Device/FX/Device_FX.cpp
@@ -113,10 +113,10 @@ void Device_FX::FreeCacheInternal(RenderTargetCache& TargetRTCache)
 
 void Device_FX::FreeRTList(RTList& RTList)
 {
-	for (UTextureRenderTarget2D* RT : RTList)
-	{
-		check(RT);
-	}
+	//for (UTextureRenderTarget2D* RT : RTList)
+	//{
+	//	check(RT);
+	//}
 
 	RTList.clear();
 }
