Fixed a couple of texture graph issues
Set/Get vector parameter value node in blueprint takes as linearcolor
Pallete tile item should not be editable

diff --git a/Engine/Plugins/Experimental/TextureGraph/Source/TextureGraph/Private/Blueprint/TG_BlueprintFunctionLibrary.cpp b/Engine/Plugins/Experimental/TextureGraph/Source/TextureGraph/Private/Blueprint/TG_BlueprintFunctionLibrary.cpp
--- a/Engine/Plugins/Experimental/TextureGraph/Source/TextureGraph/Private/Blueprint/TG_BlueprintFunctionLibrary.cpp
+++ b/Engine/Plugins/Experimental/TextureGraph/Source/TextureGraph/Private/Blueprint/TG_BlueprintFunctionLibrary.cpp
@@ -129,7 +129,7 @@ float UTG_BlueprintFunctionLibrary::GetScalarParameterValue(UObject* WorldContex
 	return ParameterValue;
 }
 
-void UTG_BlueprintFunctionLibrary::SetVectorParameterValue(UObject* WorldContextObject, UTextureGraph* InTextureGraph, FName ParameterName, FLinearColor ParameterValue)
+void UTG_BlueprintFunctionLibrary::SetVectorParameterValue(UObject* WorldContextObject, UTextureGraph* InTextureGraph, FName ParameterName, FVector4f ParameterValue)
 {
 	if (UWorld* World = GEngine->GetWorldFromContextObject(WorldContextObject, EGetWorldErrorMode::LogAndReturnNull))
 	{
@@ -156,9 +156,9 @@ void UTG_BlueprintFunctionLibrary::SetVectorParameterValue(UObject* WorldContext
 	}
 }
 
-FLinearColor UTG_BlueprintFunctionLibrary::GetVectorParameterValue(UObject* WorldContextObject, UTextureGraph* InTextureGraph, FName ParameterName)
+FVector4f UTG_BlueprintFunctionLibrary::GetVectorParameterValue(UObject* WorldContextObject, UTextureGraph* InTextureGraph, FName ParameterName)
 {
-	FLinearColor ParameterValue = FLinearColor::Black;
+	FVector4f ParameterValue = FVector4f::Zero();
 
 	if (UWorld* World = GEngine->GetWorldFromContextObject(WorldContextObject, EGetWorldErrorMode::LogAndReturnNull))
 	{
diff --git a/Engine/Plugins/Experimental/TextureGraph/Source/TextureGraph/Public/Blueprint/TG_BlueprintFunctionLibrary.h b/Engine/Plugins/Experimental/TextureGraph/Source/TextureGraph/Public/Blueprint/TG_BlueprintFunctionLibrary.h
--- a/Engine/Plugins/Experimental/TextureGraph/Source/TextureGraph/Public/Blueprint/TG_BlueprintFunctionLibrary.h
+++ b/Engine/Plugins/Experimental/TextureGraph/Source/TextureGraph/Public/Blueprint/TG_BlueprintFunctionLibrary.h
@@ -35,11 +35,11 @@ public:
 
 	/** Sets a Vector parameter value on the TextureGraph instance. Logs if ParameterName is invalid. */
 	UFUNCTION(BlueprintCallable, Category = "TextureGraph", meta = (Keywords = "SetVectorParameterValue", WorldContext = "WorldContextObject"))
-	static void SetVectorParameterValue(UObject* WorldContextObject, UTextureGraph* TextureGraph, FName ParameterName, FLinearColor ParameterValue);
+	static void SetVectorParameterValue(UObject* WorldContextObject, UTextureGraph* TextureGraph, FName ParameterName, FVector4f ParameterValue);
 
 	/** Gets a Vector parameter value from the TextureGraph instance. Logs if ParameterName is invalid. */
 	UFUNCTION(BlueprintCallable, Category = "TextureGraph", meta = (Keywords = "GetVectorParameterValue", WorldContext = "WorldContextObject"))
-	static FLinearColor GetVectorParameterValue(UObject* WorldContextObject, UTextureGraph* TextureGraph, FName ParameterName);
+	static FVector4f GetVectorParameterValue(UObject* WorldContextObject, UTextureGraph* TextureGraph, FName ParameterName);
 
 	/** Sets a color parameter value on the TextureGraph instance. Logs if ParameterName is invalid. */
 	UFUNCTION(BlueprintCallable, Category = "TextureGraph", meta = (Keywords = "SetColorParameterValue", WorldContext = "WorldContextObject"))
diff --git a/Engine/Plugins/Experimental/TextureGraph/Source/TextureGraphEditor/Private/STG_Palette.cpp b/Engine/Plugins/Experimental/TextureGraph/Source/TextureGraphEditor/Private/STG_Palette.cpp
--- a/Engine/Plugins/Experimental/TextureGraph/Source/TextureGraphEditor/Private/STG_Palette.cpp
+++ b/Engine/Plugins/Experimental/TextureGraph/Source/TextureGraphEditor/Private/STG_Palette.cpp
@@ -140,7 +140,7 @@ void STG_PaletteTileItem::Construct(const FArguments& InArgs, FCreateWidgetForAc
 
 	const FSlateBrush* IconBrush = GetIconBrush();
 	FText IconToolTip = GraphAction->GetTooltipDescription();
-	bool bIsReadOnly = false;
+	bool bIsReadOnly = true;
 
 	IconWidget = SNew(SImage)
 		.ToolTipText(IconToolTip)
