Node preview issue fixes that were causing a crash in GetAsLinearColor. Also deprecated the old Selection preview and viewport

diff --git a/Engine/Plugins/Experimental/TextureGraph/Source/TextureGraphEditor/Private/STG_NodePreview.cpp b/Engine/Plugins/Experimental/TextureGraph/Source/TextureGraphEditor/Private/STG_NodePreview.cpp
--- a/Engine/Plugins/Experimental/TextureGraph/Source/TextureGraphEditor/Private/STG_NodePreview.cpp
+++ b/Engine/Plugins/Experimental/TextureGraph/Source/TextureGraphEditor/Private/STG_NodePreview.cpp
@@ -1,7 +1,5 @@
 ﻿// Copyright Epic Games, Inc. All Rights Reserved.
 
-#if TEXTUREGRAPHEDITOR_ENABLE_NEW_NODE_PREVIEW
-
 #include "STG_NodePreview.h"
 
 #include "CanvasItem.h"
@@ -46,16 +44,29 @@ void FNodeViewer::DrawCurrentImage(FViewport*, FCanvas* Canvas, const FDrawPrope
 TOptional<TVariant<FColor, FLinearColor>> FNodeViewer::GetCurrentImagePixelColor(FIntPoint PixelCoords, int32 MipIndex) const
 {
 	const uint64 PixelIndex = PixelCoords.Y * NodeDescriptor.Width + PixelCoords.X;
-	if (PixelIndex < NodePixels.Num())
+
+	if (NodeTexture && CurrentBlob)
 	{
-		const FLinearColor LinearColor = NodePixels[PixelIndex];
+		if (!CurrentBlob->GetBufferRef()->HasRaw())
+		{
+			/// If it's already fetching raw, then we don't want to do it again
+			if (CurrentBlob->GetBufferRef()->IsFetchingRaw())
+				return {};
 
-		if (NodeDescriptor.Format == BufferFormat::Byte)
+			CurrentBlob->GetBufferRef()->Raw();
+		}
+		else 
 		{
-			return TVariant<FColor, FLinearColor>(TInPlaceType<FColor>(), LinearColor.ToFColor(IsSRGB()));
+			const FLinearColor LinearColor = CurrentBlob->GetBufferRef()->Raw_Now()->GetAsLinearColor(PixelIndex);
+
+			if (NodeDescriptor.Format == BufferFormat::Byte)
+			{
+				return TVariant<FColor, FLinearColor>(TInPlaceType<FColor>(), LinearColor.ToFColor(IsSRGB()));
+			}
+			return TVariant<FColor, FLinearColor>(TInPlaceType<FLinearColor>(), LinearColor);
 		}
-		return TVariant<FColor, FLinearColor>(TInPlaceType<FLinearColor>(), LinearColor);
 	}
+
 	return {};
 }
 
@@ -80,8 +91,8 @@ bool FNodeViewer::IsSingleChannel() const
 void FNodeViewer::SetTexture(const BlobPtr& InBlob)
 {
 	NodeTexture = nullptr;
+	CurrentBlob = InBlob;
 
-	NodePixels.SetNumUnsafeInternal(0);
 	NodeDescriptor = {};
 
 	if (InBlob)
@@ -89,13 +100,7 @@ void FNodeViewer::SetTexture(const BlobPtr& InBlob)
 		if (const DeviceBufferPtr Buffer = InBlob->GetBufferRef().GetPtr())
 		{
 			NodeTexture = GetTextureFromBuffer(Buffer);
-
-			Buffer->Raw()
-			      .then([this](const RawBufferPtr& RawBuffer)
-			      {
-				      RawBuffer->GetAsLinearColor(NodePixels);
-				      NodeDescriptor = RawBuffer->GetDescriptor();
-			      });
+			NodeDescriptor = Buffer->Descriptor();
 		}
 	}
 }
@@ -156,6 +161,7 @@ ESimpleElementBlendMode FNodeViewer::GetBlendMode() const
 UTextureRenderTarget2D* FNodeViewer::GetTextureFromBuffer(const DeviceBufferPtr& Buffer) const
 {
 	UTextureRenderTarget2D* OutTexture = nullptr;
+
 	if (const std::shared_ptr<DeviceBuffer_FX> FXBuffer = std::static_pointer_cast<DeviceBuffer_FX>(Buffer))
 	{
 		if (!FXBuffer->IsNull())
@@ -457,5 +463,3 @@ void STG_NodePreviewWidget::ToggleLock()
 }
 
 #undef LOCTEXT_NAMESPACE
-
-#endif // TEXTUREGRAPHEDITOR_ENABLE_NEW_NODE_PREVIEW
diff --git a/Engine/Plugins/Experimental/TextureGraph/Source/TextureGraphEditor/Private/STG_NodePreview.h b/Engine/Plugins/Experimental/TextureGraph/Source/TextureGraphEditor/Private/STG_NodePreview.h
--- a/Engine/Plugins/Experimental/TextureGraph/Source/TextureGraphEditor/Private/STG_NodePreview.h
+++ b/Engine/Plugins/Experimental/TextureGraph/Source/TextureGraphEditor/Private/STG_NodePreview.h
@@ -2,8 +2,6 @@
 
 #pragma once
 
-#if TEXTUREGRAPHEDITOR_ENABLE_NEW_NODE_PREVIEW
-
 #include "IImageViewer.h"
 #include "Data/Blob.h"
 #include "Widgets/SCompoundWidget.h"
@@ -62,15 +60,15 @@ private:
 	/** Indicates that a node is in sRGB format. */
 	bool IsSRGB() const;
 
+	/** The Blob for this node. */
+	BlobPtr CurrentBlob;
+
 	/** The render target texture for the node. */
 	UTextureRenderTarget2D* NodeTexture = nullptr;
 
 	/** Toggles for enabling RGBA components for drawing. */
 	bool bRGBA[4] = {true, true, true, true};
 
-	/** Duplicate CPU buffer with the node's pixel data for quick lookup of the pixel under the cursor. */
-	TArray<FLinearColor> NodePixels;
-
 	/** Node meta data.  */
 	BufferDescriptor NodeDescriptor;
 };
@@ -138,5 +136,3 @@ private:
 	/** Flags for toggling RGBA channels in the preview, which are hooked up to the RGBA buttons in the viewport toolbar extension. */
 	bool bRGBA[4] = {true, true, true, true};
 };
-
-#endif // TEXTUREGRAPHEDITOR_ENABLE_NEW_NODE_PREVIEW
diff --git a/Engine/Plugins/Experimental/TextureGraph/Source/TextureGraphEditor/Private/STG_SelectionPreview.cpp b/Engine/Plugins/Experimental/TextureGraph/Source/TextureGraphEditor/Private/STG_SelectionPreview.cpp
deleted file mode 100644
diff --git a/Engine/Plugins/Experimental/TextureGraph/Source/TextureGraphEditor/Private/STG_SelectionPreview.h b/Engine/Plugins/Experimental/TextureGraph/Source/TextureGraphEditor/Private/STG_SelectionPreview.h
deleted file mode 100644
diff --git a/Engine/Plugins/Experimental/TextureGraph/Source/TextureGraphEditor/Private/STG_TexturePreviewViewport.cpp b/Engine/Plugins/Experimental/TextureGraph/Source/TextureGraphEditor/Private/STG_TexturePreviewViewport.cpp
deleted file mode 100644
diff --git a/Engine/Plugins/Experimental/TextureGraph/Source/TextureGraphEditor/Private/STG_TexturePreviewViewport.h b/Engine/Plugins/Experimental/TextureGraph/Source/TextureGraphEditor/Private/STG_TexturePreviewViewport.h
deleted file mode 100644
diff --git a/Engine/Plugins/Experimental/TextureGraph/Source/TextureGraphEditor/Private/TG_Editor.cpp b/Engine/Plugins/Experimental/TextureGraph/Source/TextureGraphEditor/Private/TG_Editor.cpp
--- a/Engine/Plugins/Experimental/TextureGraph/Source/TextureGraphEditor/Private/TG_Editor.cpp
+++ b/Engine/Plugins/Experimental/TextureGraph/Source/TextureGraphEditor/Private/TG_Editor.cpp
@@ -44,7 +44,6 @@
 #include "SGraphPanel.h"
 #include "Kismet2/BlueprintEditorUtils.h"
 #include "PropertyEditorModule.h"
-#include "STG_SelectionPreview.h"
 #include "STG_TextureDetails.h"
 #include "ScopedTransaction.h"
 #include "ToolMenus.h"
@@ -69,9 +68,7 @@
 #include "Model/Mix/ViewportSettings.h"
 #include "UObject/MetaData.h"
 #include "UObject/ObjectSaveContext.h"
-#if TEXTUREGRAPHEDITOR_ENABLE_NEW_NODE_PREVIEW
 #include "STG_NodePreview.h"
-#endif // TEXTUREGRAPHEDITOR_ENABLE_NEW_NODE_PREVIEW
 
 #include "Expressions/Input/TG_Expression_InputParam.h"
 #include "STG_OutputSelectionDlg.h"
@@ -121,19 +118,10 @@ void FTG_Editor::RegisterTabSpawners(const TSharedRef<class FTabManager>& InTabM
 		.SetGroup(WorkspaceMenuCategoryRef)
 		.SetIcon(FSlateIcon(FAppStyle::GetAppStyleSetName(), "LevelEditor.Tabs.Details"));
 
-#if TEXTUREGRAPHEDITOR_ENABLE_OLD_SELECTION_PREVIEW
-	InTabManager->RegisterTabSpawner(FTG_EditorTabs::SelectionPreviewTabId, FOnSpawnTab::CreateSP(this, &FTG_Editor::SpawnTab_SelectionPreview))
-		.SetDisplayName(LOCTEXT("SelectionPreviewTab", "Node Preview"))
-		.SetGroup(WorkspaceMenuCategoryRef)
-	.SetIcon(FSlateIcon(FAppStyle::GetAppStyleSetName(), "LevelEditor.Tabs.Viewports"));
-#endif // TEXTUREGRAPHEDITOR_ENABLE_OLD_SELECTION_PREVIEW
-
-#if TEXTUREGRAPHEDITOR_ENABLE_NEW_NODE_PREVIEW
 	InTabManager->RegisterTabSpawner(FTG_EditorTabs::NodePreviewTabId, FOnSpawnTab::CreateSP(this, &FTG_Editor::SpawnTab_NodePreview))
 		.SetDisplayName(LOCTEXT("NodePreviewTab", "Node Preview"))
 		.SetGroup(WorkspaceMenuCategoryRef)
 		.SetIcon(FSlateIcon(FAppStyle::GetAppStyleSetName(), "LevelEditor.Tabs.Viewports"));
-#endif // TEXTUREGRAPHEDITOR_ENABLE_NEW_NODE_PREVIEW
 
 	/*InTabManager->RegisterTabSpawner(FTG_EditorTabs::OutputTabId, FOnSpawnTab::CreateSP(this, &FTG_Editor::SpawnTab_Output))
 		.SetDisplayName(LOCTEXT("OutputTab", "Output"))
@@ -166,12 +154,7 @@ void FTG_Editor::UnregisterTabSpawners(const TSharedRef<class FTabManager>& InTa
 	//InTabManager->UnregisterTabSpawner(FTG_EditorTabs::FindTabId);
 	InTabManager->UnregisterTabSpawner(FTG_EditorTabs::PreviewSceneSettingsTabId);
 	InTabManager->UnregisterTabSpawner(FTG_EditorTabs::ParameterDefaultsTabId);
-#if TEXTUREGRAPHEDITOR_ENABLE_OLD_SELECTION_PREVIEW
-	InTabManager->UnregisterTabSpawner(FTG_EditorTabs::SelectionPreviewTabId);
-#endif // TEXTUREGRAPHEDITOR_ENABLE_OLD_SELECTION_PREVIEW
-#if TEXTUREGRAPHEDITOR_ENABLE_NEW_NODE_PREVIEW
 	InTabManager->UnregisterTabSpawner(FTG_EditorTabs::NodePreviewTabId);
-#endif // TEXTUREGRAPHEDITOR_ENABLE_NEW_NODE_PREVIEW
 	InTabManager->UnregisterTabSpawner(FTG_EditorTabs::OutputTabId);
 	InTabManager->UnregisterTabSpawner(FTG_EditorTabs::PreviewSettingsTabId);
 	InTabManager->UnregisterTabSpawner(FTG_EditorTabs::ErrorsTabId);
@@ -225,12 +208,7 @@ void FTG_Editor::InitEditor(const EToolkitMode::Type Mode, const TSharedPtr< cla
 	EditedTextureGraph->Graph()->OnTGNodeRenamedDelegate.AddSP(this, &FTG_Editor::OnNodeRenamed);
 
 	GraphEditorWidget = CreateGraphEditorWidget();
-#if TEXTUREGRAPHEDITOR_ENABLE_OLD_SELECTION_PREVIEW
-	SelectionPreview = CreateSelectionViewWidget();
-#endif // TEXTUREGRAPHEDITOR_ENABLE_OLD_SELECTION_PREVIEW
-#if TEXTUREGRAPHEDITOR_ENABLE_NEW_NODE_PREVIEW
 	NodePreview = CreateNodePreviewWidget();
-#endif // TEXTUREGRAPHEDITOR_ENABLE_NEW_NODE_PREVIEW
 	TextureDetails = CreateTextureDetailsWidget();
 	Palette = SNew(STG_Palette, SharedThis(this));
 
@@ -331,18 +309,9 @@ void FTG_Editor::InitEditor(const EToolkitMode::Type Mode, const TSharedPtr< cla
 					(
 						FTabManager::NewStack()
 						->SetHideTabWell(true)
-#if TEXTUREGRAPHEDITOR_ENABLE_OLD_SELECTION_PREVIEW
-						->AddTab(FTG_EditorTabs::SelectionPreviewTabId, ETabState::OpenedTab)
-#endif // TEXTUREGRAPHEDITOR_ENABLE_OLD_SELECTION_PREVIEW
-#if TEXTUREGRAPHEDITOR_ENABLE_NEW_NODE_PREVIEW
 						->AddTab(FTG_EditorTabs::NodePreviewTabId, ETabState::OpenedTab)
-#endif // TEXTUREGRAPHEDITOR_ENABLE_NEW_NODE_PREVIEW
 						->AddTab(FTG_EditorTabs::TextureDetailsTabId, ETabState::OpenedTab)
-#if TEXTUREGRAPHEDITOR_ENABLE_NEW_NODE_PREVIEW
 						->SetForegroundTab(FTG_EditorTabs::NodePreviewTabId)
-#else
-						->SetForegroundTab(FTG_EditorTabs::SelectionPreviewTabId)						
-#endif // TEXTUREGRAPHEDITOR_ENABLE_NEW_NODE_PREVIEW
 					)
 					->Split
 					(
@@ -937,19 +906,11 @@ FActionMenuContent FTG_Editor::OnCreateGraphActionMenu(UEdGraph* InGraph, const
 	return FActionMenuContent(ActionMenu, ActionMenu->GetFilterTextBox());
 }
 
-TSharedRef<class STG_SelectionPreview> FTG_Editor::CreateSelectionViewWidget()
-{
-	return SNew(STG_SelectionPreview)
-		.OnBlobSelectionChanged(this, &FTG_Editor::OnSelectedBlobChanged);
-}
-
-#if TEXTUREGRAPHEDITOR_ENABLE_NEW_NODE_PREVIEW
 TSharedRef<STG_NodePreviewWidget> FTG_Editor::CreateNodePreviewWidget()
 {
 	return SNew(STG_NodePreviewWidget)
 		.OnNodeBlobChanged(this, &FTG_Editor::OnSelectedBlobChanged);
 }
-#endif // TEXTUREGRAPHEDITOR_ENABLE_NEW_NODE_PREVIEW
 
 TSharedRef<class STG_TextureDetails> FTG_Editor::CreateTextureDetailsWidget()
 {
@@ -1088,24 +1049,6 @@ TSharedRef<SDockTab> FTG_Editor::SpawnTab_ParameterDefaults(const FSpawnTabArgs&
 		];
 }
 
-#if TEXTUREGRAPHEDITOR_ENABLE_OLD_SELECTION_PREVIEW
-TSharedRef<SDockTab> FTG_Editor::SpawnTab_SelectionPreview(const FSpawnTabArgs& Args)
-{
-	check(Args.GetTabId() == FTG_EditorTabs::SelectionPreviewTabId);
-
-	TSharedPtr<SDockTab> Tab = SNew(SDockTab)
-		.TabColorScale(GetTabColorScale())
-		[
-			SelectionPreview.ToSharedRef()
-		];
-	NodeHistogramTab = Tab;
-
-	return Tab.ToSharedRef();
-
-}
-#endif // TEXTUREGRAPHEDITOR_ENABLE_OLD_SELECTION_PREVIEW
-
-#if TEXTUREGRAPHEDITOR_ENABLE_NEW_NODE_PREVIEW
 TSharedRef<SDockTab> FTG_Editor::SpawnTab_NodePreview(const FSpawnTabArgs& Args)
 {
 	check(Args.GetTabId() == FTG_EditorTabs::NodePreviewTabId);
@@ -1117,7 +1060,6 @@ TSharedRef<SDockTab> FTG_Editor::SpawnTab_NodePreview(const FSpawnTabArgs& Args)
 			NodePreview.ToSharedRef()
 		];
 }
-#endif // TEXTUREGRAPHEDITOR_ENABLE_NEW_NODE_PREVIEW
 
 TSharedRef<SDockTab> FTG_Editor::SpawnTab_Output(const FSpawnTabArgs& Args)
 {
@@ -1329,16 +1271,7 @@ bool FTG_Editor::DeleteNodes(TArray<UEdGraphNode*> NodesToDelete, bool ForceDele
 					UTG_Node* TGNode = TGEdGraphNode->GetNode();
 					check(TGNode);
 
-#if TEXTUREGRAPHEDITOR_ENABLE_OLD_SELECTION_PREVIEW
-					if (SelectionPreview->GetSelectedNode() && SelectionPreview->GetSelectedNode()->GetName() == TGEdGraphNode->GetName())
-					{
-						SelectionPreview->OnSelectedNodeDeleted();
-					}
-#endif // TEXTUREGRAPHEDITOR_ENABLE_OLD_SELECTION_PREVIEW
-#if TEXTUREGRAPHEDITOR_ENABLE_NEW_NODE_PREVIEW
 					NodePreview->NodeDeleted(TGEdGraphNode);
-#endif // TEXTUREGRAPHEDITOR_ENABLE_NEW_NODE_PREVIEW
-
 					TextureGraph->RemoveNode(TGNode);
 				}
 				EdGraphNode->DestroyNode();
@@ -1987,11 +1920,7 @@ void FTG_Editor::RefreshNodePreview(const TSet<class UObject*>& NewSelection, co
 
 	//TODO: Handle multiple selection case
 	UTG_EdGraphNode* SelectedNode = NodesForSelectionPreview.Num() > 0 ? NodesForSelectionPreview[0] : nullptr;
-#if TEXTUREGRAPHEDITOR_ENABLE_OLD_SELECTION_PREVIEW
-	SelectionPreview->OnSelectionChanged(SelectedNode);
-#endif // TEXTUREGRAPHEDITOR_ENABLE_OLD_SELECTION_PREVIEW
 
-#if TEXTUREGRAPHEDITOR_ENABLE_NEW_NODE_PREVIEW
 	if (bUpdateOnly)
 	{
 		NodePreview->Update();
@@ -2000,7 +1929,6 @@ void FTG_Editor::RefreshNodePreview(const TSet<class UObject*>& NewSelection, co
 	{
 		NodePreview->SelectionChanged(SelectedNode);
 	}
-#endif // TEXTUREGRAPHEDITOR_ENABLE_NEW_NODE_PREVIEW
 }
 
 void FTG_Editor::SetMesh(class UMeshComponent* InPreviewMesh, class UWorld* InWorld)
diff --git a/Engine/Plugins/Experimental/TextureGraph/Source/TextureGraphEditor/Private/TG_Editor.h b/Engine/Plugins/Experimental/TextureGraph/Source/TextureGraphEditor/Private/TG_Editor.h
--- a/Engine/Plugins/Experimental/TextureGraph/Source/TextureGraphEditor/Private/TG_Editor.h
+++ b/Engine/Plugins/Experimental/TextureGraph/Source/TextureGraphEditor/Private/TG_Editor.h
@@ -193,12 +193,7 @@ private:
 	FActionMenuContent								OnCreateGraphActionMenu(UEdGraph* InGraph, const FVector2D& InNodePosition, const TArray<UEdGraphPin*>& InDraggedPins, bool bAutoExpand, SGraphEditor::FActionMenuClosed InOnMenuClosed);
 
 	/** Create Node preview widget */
-#if TEXTUREGRAPHEDITOR_ENABLE_OLD_SELECTION_PREVIEW
-	TSharedRef<class STG_SelectionPreview>			CreateSelectionViewWidget();
-#endif // TEXTUREGRAPHEDITOR_ENABLE_OLD_SELECTION_PREVIEW
-#if TEXTUREGRAPHEDITOR_ENABLE_NEW_NODE_PREVIEW
 	TSharedRef<class STG_NodePreviewWidget>			CreateNodePreviewWidget();
-#endif // TEXTUREGRAPHEDITOR_ENABLE_NEW_NODE_PREVIEW
 
 	/** Create Texture details widget */
 	TSharedRef<class STG_TextureDetails>			CreateTextureDetailsWidget();
@@ -258,12 +253,7 @@ private:
 	TSharedRef<SDockTab>							SpawnTab_Find(const FSpawnTabArgs& Args);
 	TSharedRef<SDockTab>							SpawnTab_PreviewSettings(const FSpawnTabArgs& Args);
 	TSharedRef<SDockTab>							SpawnTab_ParameterDefaults(const FSpawnTabArgs& Args);
-#if TEXTUREGRAPHEDITOR_ENABLE_OLD_SELECTION_PREVIEW
-	TSharedRef<SDockTab>							SpawnTab_SelectionPreview(const FSpawnTabArgs& Args);
-#endif // TEXTUREGRAPHEDITOR_ENABLE_OLD_SELECTION_PREVIEW
-#if TEXTUREGRAPHEDITOR_ENABLE_NEW_NODE_PREVIEW
 	TSharedRef<SDockTab>							SpawnTab_NodePreview(const FSpawnTabArgs& Args);
-#endif // TEXTUREGRAPHEDITOR_ENABLE_NEW_NODE_PREVIEW
 	TSharedRef<SDockTab>							SpawnTab_Output(const FSpawnTabArgs& Args);
 	TSharedRef<SDockTab>							SpawnTab_Settings(const FSpawnTabArgs& Args);
 	TSharedRef<SDockTab>							SpawnTab_Errors(const FSpawnTabArgs& Args);
@@ -313,12 +303,7 @@ private:
 	TSharedPtr<SGraphEditor>						GraphEditorWidget;
 
 	/** Node preview widget being displayed*/
-#if TEXTUREGRAPHEDITOR_ENABLE_OLD_SELECTION_PREVIEW
-	TSharedPtr<class STG_SelectionPreview>			SelectionPreview;
-#endif // TEXTUREGRAPHEDITOR_ENABLE_OLD_SELECTION_PREVIEW
-#if TEXTUREGRAPHEDITOR_ENABLE_NEW_NODE_PREVIEW
 	TSharedPtr<STG_NodePreviewWidget>				NodePreview;
-#endif // TEXTUREGRAPHEDITOR_ENABLE_NEW_NODE_PREVIEW
 
 	/** Texture details widget being displayed*/
 	TSharedPtr<class STG_TextureDetails>			TextureDetails;
diff --git a/Engine/Plugins/Experimental/TextureGraph/Source/TextureGraphEditor/Private/TG_EditorTabs.cpp b/Engine/Plugins/Experimental/TextureGraph/Source/TextureGraphEditor/Private/TG_EditorTabs.cpp
--- a/Engine/Plugins/Experimental/TextureGraph/Source/TextureGraphEditor/Private/TG_EditorTabs.cpp
+++ b/Engine/Plugins/Experimental/TextureGraph/Source/TextureGraphEditor/Private/TG_EditorTabs.cpp
@@ -12,12 +12,7 @@ const FName FTG_EditorTabs::FindTabId					(TEXT("TG_Editor_Find"));
 const FName FTG_EditorTabs::GraphEditorId				(TEXT("TG_Editor_GraphEditor"));
 const FName FTG_EditorTabs::PreviewSceneSettingsTabId	(TEXT("TG_Editor_PreviewSceneSettings"));
 const FName FTG_EditorTabs::ParameterDefaultsTabId		(TEXT("TG_Editor_ParameterDefaults"));
-#if TEXTUREGRAPHEDITOR_ENABLE_OLD_SELECTION_PREVIEW
-const FName FTG_EditorTabs::SelectionPreviewTabId		(TEXT("TG_Editor_SelectionPreview"));
-#endif // TEXTUREGRAPHEDITOR_ENABLE_OLD_SELECTION_PREVIEW
-#if TEXTUREGRAPHEDITOR_ENABLE_NEW_NODE_PREVIEW
 const FName FTG_EditorTabs::NodePreviewTabId			(TEXT("TG_Editor_NodePreview"));
-#endif // TEXTUREGRAPHEDITOR_ENABLE_NEW_NODE_PREVIEW
 const FName FTG_EditorTabs::OutputTabId					(TEXT("TG_Editor_Output"));
 const FName FTG_EditorTabs::PreviewSettingsTabId		(TEXT("TG_Editor_PreviewSettings"));
 const FName FTG_EditorTabs::ErrorsTabId					(TEXT("TG_Editor_Stats"));
diff --git a/Engine/Plugins/Experimental/TextureGraph/Source/TextureGraphEditor/Private/TG_EditorTabs.h b/Engine/Plugins/Experimental/TextureGraph/Source/TextureGraphEditor/Private/TG_EditorTabs.h
--- a/Engine/Plugins/Experimental/TextureGraph/Source/TextureGraphEditor/Private/TG_EditorTabs.h
+++ b/Engine/Plugins/Experimental/TextureGraph/Source/TextureGraphEditor/Private/TG_EditorTabs.h
@@ -15,12 +15,7 @@ struct FTG_EditorTabs
 	static const FName GraphEditorId;
 	static const FName PreviewSceneSettingsTabId;
 	static const FName ParameterDefaultsTabId;
-#if TEXTUREGRAPHEDITOR_ENABLE_OLD_SELECTION_PREVIEW
-	static const FName SelectionPreviewTabId;
-#endif // TEXTUREGRAPHEDITOR_ENABLE_OLD_SELECTION_PREVIEW
-#if TEXTUREGRAPHEDITOR_ENABLE_NEW_NODE_PREVIEW
 	static const FName NodePreviewTabId;
-#endif // TEXTUREGRAPHEDITOR_ENABLE_NEW_NODE_PREVIEW
 	static const FName OutputTabId;
 	static const FName PreviewSettingsTabId;
 	static const FName ErrorsTabId;
diff --git a/Engine/Plugins/Experimental/TextureGraph/Source/TextureGraphEditor/Private/TG_TexturePreviewViewportClient.cpp b/Engine/Plugins/Experimental/TextureGraph/Source/TextureGraphEditor/Private/TG_TexturePreviewViewportClient.cpp
deleted file mode 100644
diff --git a/Engine/Plugins/Experimental/TextureGraph/Source/TextureGraphEditor/Private/TG_TexturePreviewViewportClient.h b/Engine/Plugins/Experimental/TextureGraph/Source/TextureGraphEditor/Private/TG_TexturePreviewViewportClient.h
deleted file mode 100644
diff --git a/Engine/Plugins/Experimental/TextureGraph/Source/TextureGraphEditor/TextureGraphEditor.Build.cs b/Engine/Plugins/Experimental/TextureGraph/Source/TextureGraphEditor/TextureGraphEditor.Build.cs
--- a/Engine/Plugins/Experimental/TextureGraph/Source/TextureGraphEditor/TextureGraphEditor.Build.cs
+++ b/Engine/Plugins/Experimental/TextureGraph/Source/TextureGraphEditor/TextureGraphEditor.Build.cs
@@ -3,10 +3,6 @@ using UnrealBuildTool;
 public class TextureGraphEditor : ModuleRules
 {
 	// Flag that enables the new node preview.
-	private const bool EnableNewNodePreview = true;
-	// Flag that disables the old selection preview. This is only used if the EnableNewNodePreview is set to true.
-	private const bool DisableOldSelectionPreview = false;
-
 	public TextureGraphEditor(ReadOnlyTargetRules Target) : base(Target)
 	{
 		PCHUsage = ModuleRules.PCHUsageMode.UseExplicitOrSharedPCHs;
@@ -64,6 +60,7 @@ public class TextureGraphEditor : ModuleRules
 				"MessageLog",
 				"EditorWidgets",
 				"AssetDefinition",
+				"ImageWidgets"
 				// ... add private dependencies that you statically link with here ...	
 			}
 			);
@@ -73,17 +70,5 @@ public class TextureGraphEditor : ModuleRules
 				// ... add any modules that your module loads dynamically here ...
 			}
 			);
-
-		if (EnableNewNodePreview)
-		{
-			PrivateDependencyModuleNames.Add("ImageWidgets");
-		}
-
-		PublicDefinitions.Add(EnableNewNodePreview
-			? "TEXTUREGRAPHEDITOR_ENABLE_NEW_NODE_PREVIEW=1"
-			: "TEXTUREGRAPHEDITOR_ENABLE_NEW_NODE_PREVIEW=0");
-		PublicDefinitions.Add(EnableNewNodePreview && DisableOldSelectionPreview
-			? "TEXTUREGRAPHEDITOR_ENABLE_OLD_SELECTION_PREVIEW=0"
-			: "TEXTUREGRAPHEDITOR_ENABLE_OLD_SELECTION_PREVIEW=1");
 	}
 }
\ No newline at end of file
diff --git a/Engine/Plugins/Experimental/TextureGraph/Source/TextureGraphEngine/Data/RawBuffer.cpp b/Engine/Plugins/Experimental/TextureGraph/Source/TextureGraphEngine/Data/RawBuffer.cpp
--- a/Engine/Plugins/Experimental/TextureGraph/Source/TextureGraphEngine/Data/RawBuffer.cpp
+++ b/Engine/Plugins/Experimental/TextureGraph/Source/TextureGraphEngine/Data/RawBuffer.cpp
@@ -807,7 +807,7 @@ bool RawBuffer::UncompressGeneric(FName InName)
 	return FCompression::UncompressMemory(InName, (void*)Data, (int32)Length, (const void*)CompressedData, (int32)CompressedLength);
 }
 
-void RawBuffer::GetAsLinearColor(TArray<FLinearColor>& Pixels)
+void RawBuffer::GetAsLinearColor(TArray<FLinearColor>& Pixels) const
 {
 	Pixels.AddDefaulted(Desc.Width * Desc.Height);
 
@@ -914,7 +914,7 @@ void RawBuffer::GetAsLinearColor(TArray<FLinearColor>& Pixels)
 	}*/
 }
 
-FLinearColor RawBuffer::GetAsLinearColor(int PixelIndex)
+FLinearColor RawBuffer::GetAsLinearColor(int PixelIndex) const
 {
 	FLinearColor LinearColor = FLinearColor::Black;
 
diff --git a/Engine/Plugins/Experimental/TextureGraph/Source/TextureGraphEngine/Data/RawBuffer.h b/Engine/Plugins/Experimental/TextureGraph/Source/TextureGraphEngine/Data/RawBuffer.h
--- a/Engine/Plugins/Experimental/TextureGraph/Source/TextureGraphEngine/Data/RawBuffer.h
+++ b/Engine/Plugins/Experimental/TextureGraph/Source/TextureGraphEngine/Data/RawBuffer.h
@@ -209,8 +209,8 @@ public:
 
 	AsyncRawBufferP					LoadRawBuffer(bool bInDoUncompress, bool bFreeMemory = true);
 
-	void							GetAsLinearColor(TArray<FLinearColor>& Pixels);
-	FLinearColor					GetAsLinearColor(int PixelIndex);
+	void							GetAsLinearColor(TArray<FLinearColor>& Pixels) const;
+	FLinearColor					GetAsLinearColor(int PixelIndex) const;
 
 	bool							IsPadded() const;
 	size_t							GetUnpaddedSize();
diff --git a/Engine/Plugins/Experimental/TextureGraph/Source/TextureGraphEngine/Device/DeviceBuffer.cpp b/Engine/Plugins/Experimental/TextureGraph/Source/TextureGraphEngine/Device/DeviceBuffer.cpp
--- a/Engine/Plugins/Experimental/TextureGraph/Source/TextureGraphEngine/Device/DeviceBuffer.cpp
+++ b/Engine/Plugins/Experimental/TextureGraph/Source/TextureGraphEngine/Device/DeviceBuffer.cpp
@@ -100,6 +100,7 @@ AsyncRawBufferPtr DeviceBuffer::Raw()
 		return cti::make_ready_continuable(RawData);
 
 	check(!IsNull());
+	FetchingRaw = true;
 
 	return GetOwnerDevice()->Use()
 		.then([this](int32) mutable
@@ -111,6 +112,7 @@ AsyncRawBufferPtr DeviceBuffer::Raw()
 
 			/// Add to blobber
 			NewHash = TextureGraphEngine::GetBlobber()->AddGloballyUniqueHash(NewHash);
+			FetchingRaw = false;
 
 			return PromiseUtil::OnGameThread();
 
diff --git a/Engine/Plugins/Experimental/TextureGraph/Source/TextureGraphEngine/Device/DeviceBuffer.h b/Engine/Plugins/Experimental/TextureGraph/Source/TextureGraphEngine/Device/DeviceBuffer.h
--- a/Engine/Plugins/Experimental/TextureGraph/Source/TextureGraphEngine/Device/DeviceBuffer.h
+++ b/Engine/Plugins/Experimental/TextureGraph/Source/TextureGraphEngine/Device/DeviceBuffer.h
@@ -53,6 +53,7 @@ protected:
 
 	mutable CHashPtr				HashValue;					/// The actual RawObj buffer HashValue
 	RawBufferPtr					RawData;					/// Raw buffer pointer that we keep cached (If we happen to have it)
+	bool							FetchingRaw = false;		/// Whether we're already fetching raw or not
 
 	bool							bMarkedForCollection = false; /// Whether this instance is marked for collection or not
 	bool							bIsPersistent = false;		/// Whether this is persistent or not
@@ -139,6 +140,7 @@ public:
 
 	FORCEINLINE const bool*			TransferChain() const { return Chain; }
 	FORCEINLINE bool				IsPersistent() const { return bIsPersistent; }
+	FORCEINLINE bool				IsFetchingRaw() const { return FetchingRaw; }
 };
 
 //typedef std::shared_ptr<DeviceBuffer>	DeviceBufferPtr;
