Texture Graph Export Window updates
#tests 66438172003d0e1c2e0678e8

diff --git a/Engine/Plugins/Experimental/TextureGraph/Source/TextureGraph/Public/TG_OutputSettings.h b/Engine/Plugins/Experimental/TextureGraph/Source/TextureGraph/Public/TG_OutputSettings.h
--- a/Engine/Plugins/Experimental/TextureGraph/Source/TextureGraph/Public/TG_OutputSettings.h
+++ b/Engine/Plugins/Experimental/TextureGraph/Source/TextureGraph/Public/TG_OutputSettings.h
@@ -2,6 +2,7 @@
 
 #pragma once
 
+#include "TG_SystemTypes.h"
 #include "Model/Mix/MixSettings.h"
 #include "Misc/OutputDeviceNull.h"
 #include "TG_OutputSettings.generated.h"
@@ -10,6 +11,34 @@ class UTG_Expression_Output;
 class UTG_Node;
 class UTextureGraph;
 
+
+USTRUCT()
+struct TEXTUREGRAPH_API FOutputExpressionInfo
+{
+	GENERATED_BODY()
+
+public:
+	UPROPERTY()
+	FName OutputName;
+
+	FTG_Id OutputNodeID;
+
+	bool operator==(const FOutputExpressionInfo& Other) const
+	{
+		return OutputName == Other.OutputName;
+	}
+};
+
+UCLASS()
+class TEXTUREGRAPH_API UTG_ExportSettings : public UObject
+{
+	GENERATED_BODY()
+
+public:
+	UPROPERTY(EditAnywhere, EditFixedSize, Category = NoCategory, meta = (HideItemCount, NoResetToDefault, EditFixedOrder, ShowOnlyInnerProperties, FullyExpand))
+	TArray<FOutputExpressionInfo>	OutputExpressionsInfos;
+};
+
 USTRUCT(BlueprintType)
 struct TEXTUREGRAPH_API FTG_OutputSettings 
 {
diff --git a/Engine/Plugins/Experimental/TextureGraph/Source/TextureGraphEditor/Private/AssetDefinition_TextureGraph.cpp b/Engine/Plugins/Experimental/TextureGraph/Source/TextureGraphEditor/Private/AssetDefinition_TextureGraph.cpp
--- a/Engine/Plugins/Experimental/TextureGraph/Source/TextureGraphEditor/Private/AssetDefinition_TextureGraph.cpp
+++ b/Engine/Plugins/Experimental/TextureGraph/Source/TextureGraphEditor/Private/AssetDefinition_TextureGraph.cpp
@@ -11,7 +11,6 @@
 #include "IAssetTools.h"
 #include "ToolMenuSection.h"
 #include "AssetRegistry/AssetRegistryHelpers.h"
-#include "STG_ExportDialog.h"
 
 #define LOCTEXT_NAMESPACE "UAssetDefinition_TextureGraph"
 
@@ -42,10 +41,11 @@ namespace MenuExtension_TextureGraph
 		if (const UContentBrowserAssetContextMenuContext* CBContext = UContentBrowserAssetContextMenuContext::FindContextWithAssets(InContext))
 		{
 			const TArray<UTextureGraph*> &TextureGraphs = CBContext->LoadSelectedObjects<UTextureGraph>();
+			const FTextureGraphEditorModule& TextureGraphEditorModule = FModuleManager::Get().LoadModuleChecked<FTextureGraphEditorModule>(TEXT("TextureGraphEditor"));
 			
 			for (auto TextureGraphIt = TextureGraphs.CreateConstIterator(); TextureGraphIt; ++TextureGraphIt)
 			{
-				STG_ExportDialog::ShowWindow(*TextureGraphIt);
+				TextureGraphEditorModule.GetTextureExporter()->SetTextureGraphToExport(*TextureGraphIt);
 			}
 			
 		}
diff --git a/Engine/Plugins/Experimental/TextureGraph/Source/TextureGraphEditor/Private/Customizations/TG_OutputSettingsCustomization.h b/Engine/Plugins/Experimental/TextureGraph/Source/TextureGraphEditor/Private/Customizations/TG_OutputSettingsCustomization.h
--- a/Engine/Plugins/Experimental/TextureGraph/Source/TextureGraphEditor/Private/Customizations/TG_OutputSettingsCustomization.h
+++ b/Engine/Plugins/Experimental/TextureGraph/Source/TextureGraphEditor/Private/Customizations/TG_OutputSettingsCustomization.h
@@ -214,4 +214,55 @@ public:
 	}
 };
 
+
+class FTG_OutputExpressionInfoCustomization : public IPropertyTypeCustomization
+{
+public:
+	static TSharedRef<IPropertyTypeCustomization> Create()
+	{
+		return MakeShareable(new FTG_OutputExpressionInfoCustomization);
+	}
+
+	virtual void CustomizeHeader(TSharedRef<IPropertyHandle> PropertyHandle, FDetailWidgetRow& HeaderRow, IPropertyTypeCustomizationUtils& CustomizationUtils) override
+	{
+		HeaderRow.Visibility(EVisibility::Collapsed);
+	}
+
+	virtual void CustomizeChildren(TSharedRef<IPropertyHandle> PropertyHandle, IDetailChildrenBuilder& ChildBuilder, IPropertyTypeCustomizationUtils& CustomizationUtils) override
+	{
+		FPropertyEditorModule& PropertyModule = FModuleManager::LoadModuleChecked<FPropertyEditorModule>("PropertyEditor");
+		TArray<void*> RawData;
+		PropertyHandle->AccessRawData(RawData);
+
+		FOutputExpressionInfo* OutputExpressionInfo = reinterpret_cast<FOutputExpressionInfo*>(RawData[0]);
+		TArray<UObject*> ExternalObjects;
+
+		PropertyHandle->GetOuterObjects(ExternalObjects);
+		UObject* OutermostObject = ExternalObjects[0]->GetOutermostObject();
+		UTextureGraph* TextureGraph = CastChecked<UTextureGraph>(OutermostObject);
+
+		FTG_Id OutputNodeId = OutputExpressionInfo->OutputNodeID;
+
+		UTG_Expression_Output* OutputExpression = Cast<UTG_Expression_Output>(TextureGraph->Graph()->GetNode(OutputNodeId)->GetExpression());
+		
+		ExternalObjects.Empty();
+		ExternalObjects.Add(OutputExpression);
+		
+		IDetailPropertyRow* Row = ChildBuilder.AddExternalObjectProperty(ExternalObjects, GET_MEMBER_NAME_CHECKED(UTG_Expression_Output, OutputSettings), FAddPropertyParams().HideRootObjectNode(true).CreateCategoryNodes(false));
+		Row->DisplayName(FText::FromName(OutputExpression->GetTitleName()));
+		
+		// Custom edit condition for OutputSettings
+		TAttribute<bool> OutputSettingsEditCondition = TAttribute<bool>::Create(
+			[this, OutputExpression]()
+			{
+				return OutputExpression->OutputSettings.bExport;
+			});
+		
+		Row->EditCondition(OutputSettingsEditCondition, FOnBooleanValueChanged::CreateLambda([this, OutputExpression](bool NewValue)
+			{
+				OutputExpression->SetExport(NewValue);
+			}));
+	}
+};
+
 #undef LOCTEXT_NAMESPACE
diff --git a/Engine/Plugins/Experimental/TextureGraph/Source/TextureGraphEditor/Private/STG_EditorViewport.cpp b/Engine/Plugins/Experimental/TextureGraph/Source/TextureGraphEditor/Private/STG_EditorViewport.cpp
--- a/Engine/Plugins/Experimental/TextureGraph/Source/TextureGraphEditor/Private/STG_EditorViewport.cpp
+++ b/Engine/Plugins/Experimental/TextureGraph/Source/TextureGraphEditor/Private/STG_EditorViewport.cpp
@@ -40,12 +40,13 @@
 #include "Widgets/Text/STextBlock.h"
 #include "Materials/Material.h"
 #include "Model/Mix/ViewportSettings.h"
+#include "TextureGraph.h"
 
 
 class FTG_EditorViewportClient : public FEditorViewportClient
 {
 public:
-	FTG_EditorViewportClient(TWeakPtr<ITG_Editor> InTG_AssetEditor, FAdvancedPreviewScene& InPreviewScene, const TSharedRef<STG_EditorViewport>& InTG_AssetEditorViewport);
+	FTG_EditorViewportClient(TObjectPtr<UTextureGraph> InTextureGraphPtr, FAdvancedPreviewScene& InPreviewScene, const TSharedRef<STG_EditorViewport>& InTG_AssetEditorViewport);
 
 	// FEditorViewportClient interface
 	virtual bool InputKey(const FInputKeyEventArgs& EventArgs) override;
@@ -67,17 +68,17 @@ public:
 
 private:
 
-	/** Pointer back to the material editor tool that owns us */
-	TWeakPtr<ITG_Editor> TG_EditorPtr;
+	/** Pointer to the Edited Texture Graph */
+	TObjectPtr<UTextureGraph> TextureGraphPtr;
 
 	/** Preview Scene - uses advanced preview settings */
 	class FAdvancedPreviewScene* AdvancedPreviewScene;
 };
 
 
-FTG_EditorViewportClient::FTG_EditorViewportClient(TWeakPtr<ITG_Editor> InTG_Editor, FAdvancedPreviewScene& InPreviewScene, const TSharedRef<STG_EditorViewport>& InTG_EditorViewport)
+FTG_EditorViewportClient::FTG_EditorViewportClient(TObjectPtr<UTextureGraph> InTextureGraphPtr, FAdvancedPreviewScene& InPreviewScene, const TSharedRef<STG_EditorViewport>& InTG_EditorViewport)
 	: FEditorViewportClient(nullptr, &InPreviewScene, StaticCastSharedRef<SEditorViewport>(InTG_EditorViewport))
-	, TG_EditorPtr(InTG_Editor)
+	, TextureGraphPtr(InTextureGraphPtr)
 {
 	bUsesDrawHelper = true;
 	// Setup defaults for the common draw helper.
@@ -171,13 +172,8 @@ FLinearColor FTG_EditorViewportClient::GetBackgroundColor() const
 	}
 	else
 	{
-		FLinearColor BackgroundColor = FLinearColor::Black;
+		FLinearColor BackgroundColor = FColor(64, 64, 64);			
 		
-		if (TG_EditorPtr.IsValid())
-		{
-			BackgroundColor = FColor(64, 64, 64);			
-		}
-
 		return BackgroundColor;
 	}
 }
@@ -232,10 +228,8 @@ void STG_EditorViewport::Construct(const FArguments& InArgs)
 {
 	//_previewScene = new FAdvancedPreviewScene(FPreviewScene::ConstructionValues(), 
 
-	TG_EditorPtr = InArgs._TG_Editor;
+	TextureGraphPtr = InArgs._InTextureGraph;
 
-	TSharedPtr<ITG_Editor> PinnedEditor = TG_EditorPtr.Pin();
-	
 	PreviewScene = MakeShareable(new FAdvancedPreviewScene(FPreviewScene::ConstructionValues()));
 	/*if (StaticMesh)
 	{
@@ -324,7 +318,7 @@ TSharedRef<SEditorViewport> STG_EditorViewport::GetViewportWidget()
 
 TSharedRef<FEditorViewportClient> STG_EditorViewport::MakeEditorViewportClient()
 {
-	EditorViewportClient = MakeShareable(new FTG_EditorViewportClient(TG_EditorPtr, *PreviewScene.Get(), SharedThis(this)));
+	EditorViewportClient = MakeShareable(new FTG_EditorViewportClient(TextureGraphPtr, *PreviewScene.Get(), SharedThis(this)));
 	UAssetViewerSettings::Get()->OnAssetViewerSettingsChanged().AddRaw(this, &STG_EditorViewport::OnAssetViewerSettingsChanged);
 	EditorViewportClient->SetViewLocation(FVector::ZeroVector);
 	EditorViewportClient->SetViewRotation(FRotator(-25.0f, -135.0f, 0.0f));
@@ -379,9 +373,6 @@ void STG_EditorViewport::BindCommands()
 
 	const FTG_EditorCommands& Commands = FTG_EditorCommands::Get();
 
-	check(TG_EditorPtr.IsValid());
-	CommandList->Append(TG_EditorPtr.Pin()->GetToolkitCommands());
-
 	// Add the commands to the toolkit command list so that the toolbar buttons can find them
 	CommandList->MapAction(
 		Commands.SetCylinderPreview,
@@ -436,7 +427,7 @@ void STG_EditorViewport::SetRenderMode(FName InRenderModeName)
 	check(RenderModeMgr);
 	
 	RenderModeName = InRenderModeName;
-	RenderModeMgr->ChangeRenderMode(RenderModeName);
+	RenderModeMgr->ChangeRenderMode(RenderModeName, TextureGraphPtr);
 	
 	// set editor viewport view mode
 #if WITH_EDITOR
@@ -456,8 +447,9 @@ void STG_EditorViewport::SetRenderMode(FName InRenderModeName)
 void STG_EditorViewport::InitRenderModes(UTextureGraph* InTextureGraph)
 {
 	if (!RenderModeMgr)
-		RenderModeMgr = MakeShared<TG_RenderModeManager>(InTextureGraph);
+		RenderModeMgr = MakeShared<TG_RenderModeManager>();
 	
+	TextureGraphPtr = InTextureGraph;
 	RenderModeMgr->Clear();
 	UpdateRenderMode();
 }
@@ -471,7 +463,7 @@ void STG_EditorViewport::InitPreviewMesh()
 {
 	UStaticMesh* Primitive = GUnrealEd->GetThumbnailManager()->EditorCube;
 	SetPreviewAsset(Primitive);
-	TG_EditorPtr.Pin().Get()->GetTextureGraphInterface()->GetSettings()->SetPreviewMesh(Primitive);	
+	TextureGraphPtr->GetSettings()->SetPreviewMesh(Primitive);	
 			
 }
 
@@ -496,7 +488,7 @@ void STG_EditorViewport::OnSetPreviewPrimitive(EThumbnailPrimType PrimType, bool
 		if (Primitive != nullptr)
 		{
 			SetPreviewAsset(Primitive);
-			TG_EditorPtr.Pin().Get()->GetTextureGraphInterface()->GetSettings()->SetPreviewMesh(Primitive);	
+			TextureGraphPtr->GetSettings()->SetPreviewMesh(Primitive);	
 			
 			RefreshViewport();
 		}
@@ -508,7 +500,7 @@ void STG_EditorViewport::OnSetPreviewMeshFromSelection()
 	bool bFoundPreviewMesh = false;
 	FEditorDelegates::LoadSelectedAssetsIfNeeded.Broadcast();
 
-	UMixInterface* TG_Interface = TG_EditorPtr.Pin()->GetTextureGraphInterface();
+	UMixSettings* TextureGraphSettings = TextureGraphPtr->GetSettings();
 
 	// Look for a selected asset that can be converted to a mesh component
 	for (FSelectionIterator SelectionIt(*GEditor->GetSelectedObjects()); SelectionIt && !bFoundPreviewMesh; ++SelectionIt)
@@ -521,7 +513,7 @@ void STG_EditorViewport::OnSetPreviewMeshFromSelection()
 				if (ComponentClass->IsChildOf(UMeshComponent::StaticClass()))
 				{
 					SetPreviewAsset(TestAsset);
-					TG_Interface->GetSettings()->SetPreviewMesh(Cast<UStaticMesh>(TestAsset));
+					TextureGraphSettings->SetPreviewMesh(Cast<UStaticMesh>(TestAsset));
 					bFoundPreviewMesh = true;
 				}
 			}
@@ -706,29 +698,36 @@ bool STG_EditorViewport::SetPreviewAsset(UObject* InAsset)
 		PreviewScene->SetFloorOffset(-PreviewMeshComponent->Bounds.Origin.Z + PreviewMeshComponent->Bounds.BoxExtent.Z);
 
 	}
-	TG_EditorPtr.Pin()->SetMesh(PreviewMeshComponent, PreviewScene->GetWorld());
-
+	if (IsValid(TextureGraphPtr))
+	{
+		TextureGraphPtr->SetEditorMesh(Cast<UStaticMeshComponent>(PreviewMeshComponent), PreviewScene->GetWorld()).then([this]()
+		{
+			this->InitRenderModes(TextureGraphPtr);
+		});
+	}
 	return (PreviewMeshComponent != nullptr);
 }
 
 void STG_EditorViewport::GenerateRenderModesList()
 {
-	UMixInterface* TextureGraph = GetEditorPtr().Pin().Get()->GetTextureGraphInterface();
-	UMixSettings* Settings = TextureGraph->GetSettings();
-	
-	if(ensure(Settings->GetViewportSettings().Material))
+	if (IsValid(TextureGraphPtr))
 	{
-		const FName MaterialName = Settings->GetViewportSettings().GetMaterialName();
+		UMixSettings* Settings = TextureGraphPtr->GetSettings();
+	
+		if(ensure(Settings->GetViewportSettings().Material))
+		{
+			const FName MaterialName = Settings->GetViewportSettings().GetMaterialName();
 
-		CurrentMaterialName = MaterialName;
-		RenderModeName = MaterialName;
+			CurrentMaterialName = MaterialName;
+			RenderModeName = MaterialName;
 
-		RenderModesList.Empty();
-		RenderModesList.Add(MaterialName);
+			RenderModesList.Empty();
+			RenderModesList.Add(MaterialName);
 
-		for(FMaterialMappingInfo MaterialMappingInfo : Settings->GetViewportSettings().MaterialMappingInfos)
-		{
-			RenderModesList.Add(MaterialMappingInfo.MaterialInput);
+			for(FMaterialMappingInfo MaterialMappingInfo : Settings->GetViewportSettings().MaterialMappingInfos)
+			{
+				RenderModesList.Add(MaterialMappingInfo.MaterialInput);
+			}
 		}
 	}
 }
diff --git a/Engine/Plugins/Experimental/TextureGraph/Source/TextureGraphEditor/Private/STG_EditorViewport.h b/Engine/Plugins/Experimental/TextureGraph/Source/TextureGraphEditor/Private/STG_EditorViewport.h
--- a/Engine/Plugins/Experimental/TextureGraph/Source/TextureGraphEditor/Private/STG_EditorViewport.h
+++ b/Engine/Plugins/Experimental/TextureGraph/Source/TextureGraphEditor/Private/STG_EditorViewport.h
@@ -27,7 +27,7 @@ class STG_EditorViewport : public SAssetEditorViewport, public FGCObject, public
 public:
 	SLATE_BEGIN_ARGS(STG_EditorViewport)
 	{}
-	SLATE_ARGUMENT(TWeakPtr<ITG_Editor>, TG_Editor)
+	SLATE_ARGUMENT(UTextureGraph*, InTextureGraph)
 		SLATE_END_ARGS()
 
 
@@ -37,7 +37,8 @@ private:
 	FName									RenderModeName;	///Current Selected render mode name
 	TArray<FName>							RenderModesList;
 	TSharedPtr<class FEditorViewportClient> EditorViewportClient;	/// Editor viewport client
-	TWeakPtr<ITG_Editor>					TG_EditorPtr;			/// Pointer back to the TS Asset Editor tool that owns us
+	TObjectPtr<UTextureGraph>				TextureGraphPtr;		/// Pointer to the Texture Graph being shown
+	
 	EViewModeIndex							CurrentViewMode;		/// The currently selected view mode
 	TSharedPtr<FAdvancedPreviewScene>		PreviewScene;			/// The scene for this viewport
 	TObjectPtr<UMeshComponent>				PreviewMeshComponent;	/// Component for the preview mesh
@@ -47,7 +48,7 @@ private:
 	TSharedPtr<class STG_EditorViewportRenderModeToolBar>
 											RenderModeToolBar;
 
-	TSharedPtr<TG_RenderModeManager> RenderModeMgr;
+	TSharedPtr<TG_RenderModeManager>		RenderModeMgr;
 	
 protected:
 	void									GenerateRenderModesList();
@@ -55,8 +56,9 @@ protected:
 	virtual FReply							OnDrop( const FGeometry& MyGeometry, const FDragDropEvent& DragDropEvent ) override;
 public:
 	void									Construct(const FArguments& InArgs);
+	void									SetTextureGraph(const TObjectPtr<UTextureGraph>& InTextureGraph){ TextureGraphPtr = InTextureGraph; }
 
-											STG_EditorViewport();
+	STG_EditorViewport();
 	virtual									~STG_EditorViewport() override;
 
 	void									RefreshViewport();
@@ -78,7 +80,6 @@ public:
 	void									UpdateRenderMode();
 	FText									GetRenderModeName() { return FText::FromName(RenderModeName); }
 	FName									GetRenderModeFName() { return RenderModeName; }
-	TWeakPtr<ITG_Editor>					GetEditorPtr() { return TG_EditorPtr;}
 
 	void									InitPreviewMesh();
 
diff --git a/Engine/Plugins/Experimental/TextureGraph/Source/TextureGraphEditor/Private/STG_ExportDialog.cpp b/Engine/Plugins/Experimental/TextureGraph/Source/TextureGraphEditor/Private/STG_ExportDialog.cpp
deleted file mode 100644
diff --git a/Engine/Plugins/Experimental/TextureGraph/Source/TextureGraphEditor/Private/STG_ExportDialog.h b/Engine/Plugins/Experimental/TextureGraph/Source/TextureGraphEditor/Private/STG_ExportDialog.h
deleted file mode 100644
diff --git a/Engine/Plugins/Experimental/TextureGraph/Source/TextureGraphEditor/Private/STG_NodePreview.cpp b/Engine/Plugins/Experimental/TextureGraph/Source/TextureGraphEditor/Private/STG_NodePreview.cpp
--- a/Engine/Plugins/Experimental/TextureGraph/Source/TextureGraphEditor/Private/STG_NodePreview.cpp
+++ b/Engine/Plugins/Experimental/TextureGraph/Source/TextureGraphEditor/Private/STG_NodePreview.cpp
@@ -263,7 +263,7 @@ void STG_NodePreviewWidget::Construct(const FArguments& InArgs)
 		];
 }
 
-void STG_NodePreviewWidget::SelectionChanged(UTG_EdGraphNode* Node)
+void STG_NodePreviewWidget::SelectionChanged(UTG_Node* Node)
 {
 	const bool bUpdatePreview = LockedNode == nullptr && SelectedNode != Node;
 
@@ -275,7 +275,7 @@ void STG_NodePreviewWidget::SelectionChanged(UTG_EdGraphNode* Node)
 	}
 }
 
-void STG_NodePreviewWidget::NodeDeleted(const UTG_EdGraphNode* Node)
+void STG_NodePreviewWidget::NodeDeleted(const UTG_Node* Node)
 {
 	const bool bPreviewNodeDeleted = LockedNode == Node || (LockedNode == nullptr && SelectedNode == Node);
 
@@ -297,7 +297,7 @@ void STG_NodePreviewWidget::NodeDeleted(const UTG_EdGraphNode* Node)
 
 void STG_NodePreviewWidget::Update() const
 {
-	const UTG_EdGraphNode* PreviewNode = LockedNode ? LockedNode : SelectedNode;
+	const UTG_Node* PreviewNode = LockedNode ? LockedNode : SelectedNode;
 
 	// Get blob if preview node is valid or assign a nullptr if not.
 	const BlobPtr Blob = [PreviewNode]() -> BlobPtr
@@ -305,9 +305,7 @@ void STG_NodePreviewWidget::Update() const
 		if (PreviewNode)
 		{
 			TArray<FTG_Texture> OutTextures;
-			const UTG_Node* TGNode = PreviewNode->GetNode();
-			check(TGNode);
-			TGNode->GetAllOutputValues(OutTextures);
+			PreviewNode->GetAllOutputValues(OutTextures);
 
 			if (!OutTextures.IsEmpty())
 			{
diff --git a/Engine/Plugins/Experimental/TextureGraph/Source/TextureGraphEditor/Private/STG_NodePreview.h b/Engine/Plugins/Experimental/TextureGraph/Source/TextureGraphEditor/Private/STG_NodePreview.h
--- a/Engine/Plugins/Experimental/TextureGraph/Source/TextureGraphEditor/Private/STG_NodePreview.h
+++ b/Engine/Plugins/Experimental/TextureGraph/Source/TextureGraphEditor/Private/STG_NodePreview.h
@@ -14,7 +14,7 @@ namespace UE::ImageWidgets
 class FToolBarBuilder;
 class FUICommandList;
 class SHorizontalBox;
-class UTG_EdGraphNode;
+class UTG_Node;
 
 /**
  * Image viewer implementation that holds the node buffer information and draws the node texture.
@@ -94,10 +94,10 @@ public:
 	void Construct(const FArguments& InArgs);
 
 	/** Notify the preview about the node selection having changed. */
-	void SelectionChanged(UTG_EdGraphNode* Node);
+	void SelectionChanged(UTG_Node* Node);
 
 	/** Notify the preview about a node being deleted. */
-	void NodeDeleted(const UTG_EdGraphNode* Node);
+	void NodeDeleted(const UTG_Node* Node);
 
 	/** Trigger an update of the node preview after the displayed contents changed. */
 	void Update() const;
@@ -125,10 +125,12 @@ private:
 	TSharedPtr<FUICommandList> CommandList;
 
 	/** Pointer to the currently selected node, which might be different if the node preview is locked to another node. */
-	UTG_EdGraphNode* SelectedNode = nullptr;
+	// UTG_EdGraphNode* SelectedNode = nullptr;
+	UTG_Node* SelectedNode = nullptr;
 
 	/** Pointer to the node the node preview is currently locked on. */
-	UTG_EdGraphNode* LockedNode = nullptr;
+	// UTG_EdGraphNode* LockedNode = nullptr;
+	UTG_Node* LockedNode = nullptr;
 
 	/** Callback to notify about the node preview texture having changed to make sure that for example the histogram is synced with the node preview. */
 	FOnNodeBlobChanged OnNodeBlobChanged;
diff --git a/Engine/Plugins/Experimental/TextureGraph/Source/TextureGraphEditor/Private/TG_Editor.cpp b/Engine/Plugins/Experimental/TextureGraph/Source/TextureGraphEditor/Private/TG_Editor.cpp
--- a/Engine/Plugins/Experimental/TextureGraph/Source/TextureGraphEditor/Private/TG_Editor.cpp
+++ b/Engine/Plugins/Experimental/TextureGraph/Source/TextureGraphEditor/Private/TG_Editor.cpp
@@ -960,11 +960,10 @@ TSharedRef<SDockTab> FTG_Editor::SpawnTab_Viewport(const FSpawnTabArgs& Args)
 	TSharedRef< SDockTab > DockableTab =
 		SNew(SDockTab);
 
-	TWeakPtr<ITG_Editor> WeakSharedThis(SharedThis(this));
-	MakeViewportFunc = [WeakSharedThis](const FAssetEditorViewportConstructionArgs& InArgs)
+	MakeViewportFunc = [this](const FAssetEditorViewportConstructionArgs& InArgs)
 		{
 			return SNew(STG_EditorViewport)
-				.TG_Editor(WeakSharedThis);
+				.InTextureGraph(EditedTextureGraph);
 		};
 
 	// Create a new tab
@@ -1271,7 +1270,7 @@ bool FTG_Editor::DeleteNodes(TArray<UEdGraphNode*> NodesToDelete, bool ForceDele
 					UTG_Node* TGNode = TGEdGraphNode->GetNode();
 					check(TGNode);
 
-					NodePreview->NodeDeleted(TGEdGraphNode);
+					NodePreview->NodeDeleted(TGNode);
 					TextureGraph->RemoveNode(TGNode);
 				}
 				EdGraphNode->DestroyNode();
@@ -1927,7 +1926,7 @@ void FTG_Editor::RefreshNodePreview(const TSet<class UObject*>& NewSelection, co
 	}
 	else
 	{
-		NodePreview->SelectionChanged(SelectedNode);
+		NodePreview->SelectionChanged(SelectedNode ? SelectedNode->GetNode() : nullptr);
 	}
 }
 
diff --git a/Engine/Plugins/Experimental/TextureGraph/Source/TextureGraphEditor/Private/TG_EditorTabs.cpp b/Engine/Plugins/Experimental/TextureGraph/Source/TextureGraphEditor/Private/TG_EditorTabs.cpp
--- a/Engine/Plugins/Experimental/TextureGraph/Source/TextureGraphEditor/Private/TG_EditorTabs.cpp
+++ b/Engine/Plugins/Experimental/TextureGraph/Source/TextureGraphEditor/Private/TG_EditorTabs.cpp
@@ -16,4 +16,5 @@ const FName FTG_EditorTabs::NodePreviewTabId			(TEXT("TG_Editor_NodePreview"));
 const FName FTG_EditorTabs::OutputTabId					(TEXT("TG_Editor_Output"));
 const FName FTG_EditorTabs::PreviewSettingsTabId		(TEXT("TG_Editor_PreviewSettings"));
 const FName FTG_EditorTabs::ErrorsTabId					(TEXT("TG_Editor_Stats"));
-const FName FTG_EditorTabs::TextureDetailsTabId			(TEXT("TG_Editor_TextureDetails"));
\ No newline at end of file
+const FName FTG_EditorTabs::TextureDetailsTabId			(TEXT("TG_Editor_TextureDetails"));
+const FName FTG_EditorTabs::TextureExporterTabId		(TEXT("TG_Editor_TextureExporter"));
\ No newline at end of file
diff --git a/Engine/Plugins/Experimental/TextureGraph/Source/TextureGraphEditor/Private/TG_EditorTabs.h b/Engine/Plugins/Experimental/TextureGraph/Source/TextureGraphEditor/Private/TG_EditorTabs.h
--- a/Engine/Plugins/Experimental/TextureGraph/Source/TextureGraphEditor/Private/TG_EditorTabs.h
+++ b/Engine/Plugins/Experimental/TextureGraph/Source/TextureGraphEditor/Private/TG_EditorTabs.h
@@ -20,6 +20,7 @@ struct FTG_EditorTabs
 	static const FName PreviewSettingsTabId;
 	static const FName ErrorsTabId;
 	static const FName TextureDetailsTabId;
+	static const FName TextureExporterTabId;
 
 	// Disable default constructor
 	FTG_EditorTabs() = delete;
diff --git a/Engine/Plugins/Experimental/TextureGraph/Source/TextureGraphEditor/Private/TG_Exporter.cpp b/Engine/Plugins/Experimental/TextureGraph/Source/TextureGraphEditor/Private/TG_Exporter.cpp
new file mode 100644
--- /dev/null
+++ b/Engine/Plugins/Experimental/TextureGraph/Source/TextureGraphEditor/Private/TG_Exporter.cpp
@@ -0,0 +1,947 @@
+// Copyright Epic Games, Inc. All Rights Reserved.
+
+#include "TG_Exporter.h"
+
+#include "TG_EditorTabs.h"
+#include "CoreGlobals.h"
+#include "Delegates/Delegate.h"
+#include "TextureGraph.h"
+#include "TG_Node.h"
+#include "TG_Graph.h"
+#include "TG_Parameter.h"
+#include "Expressions/Output/TG_Expression_Output.h"
+#include "Framework/Application/SlateApplication.h"
+#include "Framework/Commands/Commands.h"
+#include "Framework/Commands/InputChord.h"
+#include "Framework/Commands/UIAction.h"
+#include "Framework/Commands/UICommandInfo.h"
+#include "Framework/Commands/UICommandList.h"
+#include "Framework/Docking/LayoutService.h"
+#include "Framework/Docking/TabManager.h"
+#include "Framework/MultiBox/MultiBoxBuilder.h"
+#include "HAL/Platform.h"
+#include "Internationalization/Internationalization.h"
+#include "Internationalization/Text.h"
+#include "Layout/Margin.h"
+#include "Misc/Attribute.h"
+#include "SlotBase.h"
+#include "Styling/AppStyle.h"
+#include "Styling/CoreStyle.h"
+#include "Styling/ISlateStyle.h"
+#include "Templates/SharedPointer.h"
+#include "Textures/SlateIcon.h"
+#include "Types/SlateEnums.h"
+#include "UObject/NameTypes.h"
+#include "UObject/UnrealNames.h"
+#include "UObject/WeakObjectPtr.h"
+#include "UObject/WeakObjectPtrTemplates.h"
+#include "Widgets/DeclarativeSyntaxSupport.h"
+#include "Widgets/Docking/SDockTab.h"
+#include "Widgets/Layout/SBorder.h"
+#include "Widgets/SBoxPanel.h"
+#include "Widgets/Input/SComboBox.h"
+#include "Widgets/Layout/SScrollBox.h"
+#include "WorkspaceMenuStructure.h"
+#include "WorkspaceMenuStructureModule.h"
+#include "AssetEditorViewportLayout.h"
+#include "STG_EditorViewport.h"
+#include "EditorViewportTabContent.h"
+#include "AdvancedPreviewSceneModule.h"
+#include "SPrimaryButton.h"
+#include "TG_OutputSettings.h"
+#include "STG_NodePreview.h"
+#include "TG_HelperFunctions.h"
+#include "PropertyEditorModule.h"
+#include "DetailsViewArgs.h"
+#include "IDetailsView.h"
+
+class SWidget;
+class SWindow;
+
+#define LOCTEXT_NAMESPACE "TextureGraphExporter"
+
+struct FTG_ExporterCommands : public TCommands<FTG_ExporterCommands>
+{
+	FTG_ExporterCommands()
+		: TCommands<FTG_ExporterCommands>(
+			TEXT("TextureGraphExporter"), // Context name for fast lookup
+			LOCTEXT("TextureGraphDebugger", "Texture Graph Exporter"), // Localized context name for displaying
+			NAME_None, // Parent
+			FCoreStyle::Get().GetStyleSetName() // Icon Style Set
+		)
+	{
+	}
+
+	// TCommand<> interface
+	virtual void RegisterCommands() override;
+	// End of TCommand<> interface
+
+	TSharedPtr<FUICommandInfo> ShowOutputPreview;
+	TSharedPtr<FUICommandInfo> Show3DPreview;
+	TSharedPtr<FUICommandInfo> Show3DPreviewSettings;
+	TSharedPtr<FUICommandInfo> ShowParameters;
+	TSharedPtr<FUICommandInfo> ShowExportSettings;
+
+};
+
+
+struct FTG_ExporterImpl: public FTickableGameObject, public FGCObject
+{
+public:
+	FTG_ExporterImpl();
+	~FTG_ExporterImpl();
+
+	/** Function registered with tab manager to create the Texture Graph Exporter */
+	TSharedRef<SDockTab> CreateTGExporterTab(const FSpawnTabArgs& Args);
+
+	void Cleanup();
+	/** Sets the Texture Graph to be exported in the Exporter */
+	void SetTextureGraphToExport(UTextureGraph* InTextureGraph);
+	void OnGraphChanged(UTG_Graph* InGraph, UTG_Node* InNode, bool Tweaking);
+
+	TSharedPtr<FTabManager>							TGExporterTabManager;
+	TSharedPtr<FTabManager::FLayout>				TGExporterLayout;
+
+	TObjectPtr<UTextureGraph>						TextureGraphPtr;
+	TSharedPtr<STG_NodePreviewWidget>				NodePreview;
+	
+	TSharedPtr<class IDetailsView>					ParametersView;
+	TSharedPtr<class IDetailsView>					ExportSettingsView;
+	TSharedPtr<class IDetailsView>					PreviewSettingsView;
+// Tracking the active viewports in this editor.
+	TSharedPtr<class FEditorViewportTabContent>		ViewportTabContent;
+
+	TObjectPtr<UTG_Parameters>						Parameters;
+	TObjectPtr<UTG_ExportSettings>					ExportSettings;
+	FExportSettings                                 TargetExportSettings;
+	
+	TArray<TSharedPtr<FName>>						OutputNodesList;
+	TSharedPtr<SComboBox<TSharedPtr<FName>>>		OutputNodesComboBoxWidget;
+	UTG_Node*										SelectedNode = nullptr;
+	// Inherited via FTickableGameObject
+	virtual void									Tick(float DeltaTime) override;
+
+	virtual ETickableTickType						GetTickableTickType() const override { return ETickableTickType::Always; }
+
+	virtual bool									IsTickableWhenPaused() const override { return true; }
+
+	virtual bool									IsTickableInEditor() const override { return true; }
+
+	virtual TStatId									GetStatId() const override;
+
+	// Inherited via FGCObject
+	virtual void									AddReferencedObjects(FReferenceCollector& Collector) override;
+	virtual FString									GetReferencerName() const override { return TEXT("FTextureGraphExporter");}
+
+private:
+	TSharedRef<SDockTab>							SpawnTab_Viewport(const FSpawnTabArgs& Args);
+	TSharedRef<SDockTab>							SpawnTab_ParameterDefaults(const FSpawnTabArgs& Args);
+	TSharedRef<SDockTab>							SpawnTab_NodePreview(const FSpawnTabArgs& Args);
+	TSharedRef<SDockTab>							SpawnTab_ExportSettings(const FSpawnTabArgs& Args);
+	TSharedRef<SDockTab>							SpawnTab_PreviewSettings(const FSpawnTabArgs& Args);
+	void 											RegisterTabSpawners(const TSharedPtr<FTabManager>);
+	void  											UnregisterTabSpawners(const TSharedPtr<FTabManager>);
+	TSharedPtr<class STG_EditorViewport>			GetEditorViewport() const;
+	void 											SetViewportPreviewMesh();	
+	bool 											SetPreviewAsset(UObject* InAsset);
+	TSharedRef<class IDetailsView>					GetPreviewSettingsView() const { return PreviewSettingsView.ToSharedRef(); }
+	TSharedRef<class IDetailsView>					GetExportSettingsView() const { return ExportSettingsView.ToSharedRef(); }
+	TSharedRef<class IDetailsView>					GetParametersView() const { return ParametersView.ToSharedRef(); }
+	void											OnViewportMaterialChanged();
+	void											OnMaterialMappingChanged();
+	void											UpdateExportSettingsUI();
+	void											UpdateParametersUI();
+	virtual void									RefreshViewport();
+	void											OnRenderingDone(UMixInterface* TextureGraph, const FInvalidationDetails* Details);
+	TSharedRef<SWidget>								GenerateOutputComboItem(TSharedPtr<FName> String);
+	void											OnOutputSelectionChanged(TSharedPtr<FName> String, ESelectInfo::Type Arg);
+	FReply											OnExportClicked(EAppReturnType::Type ButtonID);
+
+	// prevent copying:
+	FTG_ExporterImpl(const FTG_ExporterImpl&);
+	FTG_ExporterImpl(FTG_ExporterImpl&&);
+	FTG_ExporterImpl& operator=(FTG_ExporterImpl const&);
+	FTG_ExporterImpl& operator=(FTG_ExporterImpl&&);
+};
+
+void FTG_ExporterCommands::RegisterCommands()
+{
+	UI_COMMAND(ShowOutputPreview, "Node Preview", "Toggles visibility of the Output Preview", EUserInterfaceActionType::Check, FInputChord());
+	UI_COMMAND(Show3DPreview, "3D Preview", "Toggles visibility of the 3D Preview window", EUserInterfaceActionType::Check, FInputChord());
+	UI_COMMAND(Show3DPreviewSettings, "3D Preview Settings", "Toggles visibility of the 3D Preview Settings window", EUserInterfaceActionType::Check, FInputChord());
+	UI_COMMAND(ShowParameters, "Parameters", "Toggles visibility of the Parameters window", EUserInterfaceActionType::Check, FInputChord());
+	UI_COMMAND(ShowExportSettings, "Export Settings", "Toggles visibility of the Export Settings window", EUserInterfaceActionType::Check, FInputChord());
+}
+FTG_ExporterImpl::FTG_ExporterImpl()
+{
+	const IWorkspaceMenuStructure& MenuStructure = WorkspaceMenu::GetMenuStructure();
+
+	FTG_ExporterCommands::Register();
+	FGlobalTabmanager::Get()->RegisterNomadTabSpawner(FTG_EditorTabs::TextureExporterTabId, FOnSpawnTab::CreateRaw(this, &FTG_ExporterImpl::CreateTGExporterTab))
+		.SetDisplayName(NSLOCTEXT("TextureGraphExporter", "TabTitle", "Texture Graph Exporter"))
+		.SetTooltipText(NSLOCTEXT("TextureGraphExporter", "TooltipText", "Open the Texture Graph Exporter tab."))
+		.SetGroup(MenuStructure.GetDeveloperToolsMiscCategory())
+		.SetIcon(FSlateIcon(FAppStyle::GetAppStyleSetName(), "ClassIcon.Texture2D"));
+}
+
+FTG_ExporterImpl::~FTG_ExporterImpl()
+{
+	if (FSlateApplication::IsInitialized())
+	{
+		FGlobalTabmanager::Get()->UnregisterNomadTabSpawner(FTG_EditorTabs::TextureExporterTabId);
+	}
+
+	if (TGExporterTabManager.IsValid())
+	{
+		UnregisterTabSpawners(TGExporterTabManager);
+		FGlobalTabmanager::Get()->UnregisterTabSpawner(FTG_EditorTabs::TextureExporterTabId);
+		TGExporterLayout = TSharedPtr<FTabManager::FLayout>();
+		TGExporterTabManager = TSharedPtr<FTabManager>();
+
+		Cleanup();
+		
+		ExportSettingsView.Reset();
+		PreviewSettingsView.Reset();
+		ParametersView.Reset();
+		
+		// cleanup UI
+		if (Parameters)
+		{
+			Parameters->Parameters.Empty();
+			Parameters = nullptr;
+		}
+		if (ExportSettings)
+		{
+			ExportSettings->OutputExpressionsInfos.Empty();
+			ExportSettings = nullptr;
+		}
+		OutputNodesComboBoxWidget.Reset();
+	}
+	FTG_ExporterCommands::Unregister();
+}
+
+void FTG_ExporterImpl::Cleanup()
+{
+	if (TextureGraphPtr)
+	{
+		// cleanup events
+		TextureGraphPtr->GetSettings()->GetViewportSettings().OnViewportMaterialChangedEvent.RemoveAll(this);
+		TextureGraphPtr->GetSettings()->GetViewportSettings().OnMaterialMappingChangedEvent.RemoveAll(this);
+		TextureGraphPtr->GetSettings()->OnPreviewMeshChangedEvent.RemoveAll(this);
+		TextureGraphPtr->Graph()->OnGraphChangedDelegate.RemoveAll(this);
+		TextureGraphPtr->OnRenderDone.Unbind();
+		TextureGraphPtr = nullptr;
+	}
+
+	
+	OutputNodesList.Empty();
+	SelectedNode = nullptr;
+}
+TSharedRef<SDockTab> FTG_ExporterImpl::CreateTGExporterTab(const FSpawnTabArgs& Args)
+{
+	const TSharedRef<SDockTab> NomadTab = SNew(SDockTab)
+		.TabRole(ETabRole::NomadTab)
+		.Label(NSLOCTEXT("TextureGraphExporter", "TabTitle", "Texture Graph Exporter"));
+
+	TGExporterTabManager = FGlobalTabmanager::Get()->NewTabManager(NomadTab);
+	// on persist layout will handle saving layout if the editor is shut down:
+	TGExporterTabManager->SetOnPersistLayout(
+		FTabManager::FOnPersistLayout::CreateStatic(
+			[](const TSharedRef<FTabManager::FLayout>& InLayout)
+			{
+				if (InLayout->GetPrimaryArea().Pin().IsValid())
+				{
+					FLayoutSaveRestore::SaveToConfig(GEditorLayoutIni, InLayout);
+				}
+			}
+		)
+	);
+
+
+	TWeakPtr<FTabManager> TGExporterTabManagerWeak = TGExporterTabManager;
+	// On tab close will save the layout if the exporter window itself is closed,
+	// this handler also cleans up any floating controls. If we don't close
+	// all areas we need to add some logic to the tab manager to reuse existing tabs:
+	NomadTab->SetOnTabClosed(SDockTab::FOnTabClosedCallback::CreateStatic(
+		[](TSharedRef<SDockTab> Self, TWeakPtr<FTabManager> TabManager)
+		{
+			TSharedPtr<FTabManager> OwningTabManager = TabManager.Pin();
+			if (OwningTabManager.IsValid())
+			{
+				FLayoutSaveRestore::SaveToConfig(GEditorLayoutIni, OwningTabManager->PersistLayout());
+				OwningTabManager->CloseAllAreas();
+			}
+		}
+		, TGExporterTabManagerWeak
+	));
+
+
+	NodePreview = SNew(STG_NodePreviewWidget);
+
+	FPropertyEditorModule& PropertyEditorModule = FModuleManager::GetModuleChecked<FPropertyEditorModule>("PropertyEditor");
+	FDetailsViewArgs ParameterViewArgs;
+	ParameterViewArgs.NameAreaSettings = FDetailsViewArgs::HideNameArea;
+	ParameterViewArgs.bHideSelectionTip = true;
+	ParameterViewArgs.ColumnWidth = 0.70;
+	ParametersView = PropertyEditorModule.CreateDetailView(ParameterViewArgs);
+
+	// Settings details view
+	FDetailsViewArgs ExportSettingsViewArgs;
+	ExportSettingsViewArgs.NameAreaSettings = FDetailsViewArgs::HideNameArea;
+	ExportSettingsViewArgs.bHideSelectionTip = true;
+	ExportSettingsView = PropertyEditorModule.CreateDetailView(ExportSettingsViewArgs);
+
+	// Settings details view
+	FDetailsViewArgs SettingsViewArgs;
+	SettingsViewArgs.NameAreaSettings = FDetailsViewArgs::HideNameArea;
+	SettingsViewArgs.bHideSelectionTip = true;
+	PreviewSettingsView = PropertyEditorModule.CreateDetailView(SettingsViewArgs);
+	
+	RegisterTabSpawners(TGExporterTabManager);
+
+	
+	TGExporterLayout = FTabManager::NewLayout("Standalone_TextureGraphExporter_Layout_v1")
+	->AddArea
+	(
+		FTabManager::NewPrimaryArea()->SetOrientation(Orient_Vertical)
+			->Split
+			(
+				FTabManager::NewSplitter()->SetOrientation(Orient_Horizontal)->SetSizeCoefficient(0.9f)
+				->Split
+				(
+					FTabManager::NewSplitter()->SetOrientation(Orient_Vertical)->SetSizeCoefficient(0.25f)
+					->Split
+					(
+						FTabManager::NewStack()
+						->AddTab(FTG_EditorTabs::ParameterDefaultsTabId, ETabState::OpenedTab)
+					)
+					->Split
+					(
+						FTabManager::NewStack()
+						->AddTab(FTG_EditorTabs::OutputTabId, ETabState::OpenedTab)
+					)
+				)
+				->Split
+				(
+					FTabManager::NewSplitter()->SetOrientation(Orient_Vertical)->SetSizeCoefficient(0.25f)
+					->Split
+					(
+						FTabManager::NewStack()
+						->SetHideTabWell(true)
+						->AddTab(FTG_EditorTabs::NodePreviewTabId, ETabState::OpenedTab)
+					)
+					->Split
+					(
+						FTabManager::NewStack()
+						->AddTab(FTG_EditorTabs::ViewportTabId, ETabState::OpenedTab)
+						->AddTab(FTG_EditorTabs::PreviewSettingsTabId, ETabState::OpenedTab)
+						->SetForegroundTab(FTG_EditorTabs::ViewportTabId)
+					)
+				)
+			)
+		);
+
+	
+	TGExporterLayout = FLayoutSaveRestore::LoadFromConfig(GEditorLayoutIni, TGExporterLayout.ToSharedRef());
+	
+	TSharedRef<SWidget> TabContents = TGExporterTabManager->RestoreFrom(TGExporterLayout.ToSharedRef(), TSharedPtr<SWindow>()).ToSharedRef();
+	
+	// build command list for tab restoration menu:
+	TSharedPtr<FUICommandList> CommandList = MakeShareable(new FUICommandList());
+	
+	TWeakPtr<FTabManager> TGExportManagerWeak = TGExporterTabManager;
+	
+	const auto ToggleTabVisibility = [](TWeakPtr<FTabManager> InTGExportManagerWeak, FName InTabName)
+	{
+		TSharedPtr<FTabManager> InTGExportManager = InTGExportManagerWeak.Pin();
+		if (InTGExportManager.IsValid())
+		{
+			TSharedPtr<SDockTab> ExistingTab = InTGExportManager->FindExistingLiveTab(InTabName);
+			if (ExistingTab.IsValid())
+			{
+				ExistingTab->RequestCloseTab();
+			}
+			else
+			{
+				InTGExportManager->TryInvokeTab(InTabName);
+			}
+		}
+	};
+	
+	const auto IsTabVisible = [](TWeakPtr<FTabManager> InTGExportManagerWeak, FName InTabName)
+	{
+		TSharedPtr<FTabManager> InTGExportManager = InTGExportManagerWeak.Pin();
+		if (InTGExportManager.IsValid())
+		{
+			return InTGExportManager->FindExistingLiveTab(InTabName).IsValid();
+		}
+		return false;
+	};
+	
+	bool bViewportIsOff = !ViewportTabContent.IsValid();
+	TSharedPtr<SDockTab> ViewportTab; 
+	// check here if 3d viewport is turned off, we need to turn it on temporarily to initialize our systems correctly
+	if (bViewportIsOff)
+	{
+		ViewportTab = TGExporterTabManager->TryInvokeTab(FTG_EditorTabs::ViewportTabId);	
+	}
+	
+	// Set the preview mesh for the material.  This call must occur after the toolbar is initialized.
+	SetViewportPreviewMesh();
+
+	CommandList->MapAction(
+		FTG_ExporterCommands::Get().Show3DPreview,
+		FExecuteAction::CreateStatic(
+			ToggleTabVisibility,
+			TGExportManagerWeak,
+			FTG_EditorTabs::ViewportTabId
+		),
+		FCanExecuteAction::CreateStatic(
+			[]() { return true; }
+		),
+		FIsActionChecked::CreateStatic(
+			IsTabVisible,
+			TGExportManagerWeak,
+			FTG_EditorTabs::ViewportTabId
+		)
+	);
+	
+	CommandList->MapAction(
+		FTG_ExporterCommands::Get().Show3DPreviewSettings,
+		FExecuteAction::CreateStatic(
+			ToggleTabVisibility,
+			TGExportManagerWeak,
+			FTG_EditorTabs::PreviewSettingsTabId
+		),
+		FCanExecuteAction::CreateStatic(
+			[]() { return true; }
+		),
+		FIsActionChecked::CreateStatic(
+			IsTabVisible,
+			TGExportManagerWeak,
+			FTG_EditorTabs::PreviewSettingsTabId
+		)
+	);
+
+	CommandList->MapAction(
+			FTG_ExporterCommands::Get().ShowParameters,
+			FExecuteAction::CreateStatic(
+				ToggleTabVisibility,
+				TGExportManagerWeak,
+				FTG_EditorTabs::ParameterDefaultsTabId
+			),
+			FCanExecuteAction::CreateStatic(
+				[]() { return true; }
+			),
+			FIsActionChecked::CreateStatic(
+				IsTabVisible,
+				TGExportManagerWeak,
+				FTG_EditorTabs::ParameterDefaultsTabId
+			)
+		);
+
+	CommandList->MapAction(
+			FTG_ExporterCommands::Get().ShowOutputPreview,
+			FExecuteAction::CreateStatic(
+				ToggleTabVisibility,
+				TGExportManagerWeak,
+				FTG_EditorTabs::NodePreviewTabId
+			),
+			FCanExecuteAction::CreateStatic(
+				[]() { return true; }
+			),
+			FIsActionChecked::CreateStatic(
+				IsTabVisible,
+				TGExportManagerWeak,
+				FTG_EditorTabs::NodePreviewTabId
+			)
+		);
+
+	CommandList->MapAction(
+			FTG_ExporterCommands::Get().ShowExportSettings,
+			FExecuteAction::CreateStatic(
+				ToggleTabVisibility,
+				TGExportManagerWeak,
+				FTG_EditorTabs::OutputTabId
+			),
+			FCanExecuteAction::CreateStatic(
+				[]() { return true; }
+			),
+			FIsActionChecked::CreateStatic(
+				IsTabVisible,
+				TGExportManagerWeak,
+				FTG_EditorTabs::OutputTabId
+			)
+		);
+	
+	FMenuBarBuilder MenuBarBuilder(CommandList);
+	MenuBarBuilder.AddPullDownMenu(
+		LOCTEXT("WindowMenuLabel", "Window"),
+		FText::GetEmpty(),
+		FNewMenuDelegate::CreateLambda([](FMenuBuilder& Builder) {
+			Builder.AddMenuEntry(FTG_ExporterCommands::Get().ShowOutputPreview);
+			Builder.AddMenuEntry(FTG_ExporterCommands::Get().ShowParameters);
+			Builder.AddMenuEntry(FTG_ExporterCommands::Get().Show3DPreviewSettings);
+			Builder.AddMenuEntry(FTG_ExporterCommands::Get().Show3DPreview);
+			Builder.AddMenuEntry(FTG_ExporterCommands::Get().ShowExportSettings);
+			})
+	);
+	
+	
+	TSharedRef<SWidget> MenuBarWidget = MenuBarBuilder.MakeWidget();
+	
+	NomadTab->SetContent(
+		SNew(SVerticalBox)
+		+ SVerticalBox::Slot()
+		.AutoHeight()
+		[
+			MenuBarWidget
+		]
+		+SVerticalBox::Slot()
+		[
+			SNew(SBorder)
+			.BorderImage(FAppStyle::GetBrush("ToolPanel.DarkGroupBorder"))
+			.Padding(FMargin(0.f, 2.f))
+			[
+				TabContents
+			]
+		]
+	);
+	
+	
+	// Tell tab-manager about the multi-box for platforms with a global menu bar
+	TGExporterTabManager->SetMenuMultiBox(MenuBarBuilder.GetMultiBox(), MenuBarWidget);
+	
+	return NomadTab;
+}
+
+
+void FTG_ExporterImpl::RegisterTabSpawners(const TSharedPtr<FTabManager> InTabManager)
+{
+	InTabManager->RegisterTabSpawner(FTG_EditorTabs::ViewportTabId, FOnSpawnTab::CreateRaw(this, &FTG_ExporterImpl::SpawnTab_Viewport))
+		.SetDisplayName(LOCTEXT("ViewportTab", "3D Preview"))
+		.SetIcon(FSlateIcon(FAppStyle::GetAppStyleSetName(), "LevelEditor.Tabs.Viewports"));
+
+	InTabManager->RegisterTabSpawner(FTG_EditorTabs::ParameterDefaultsTabId, FOnSpawnTab::CreateRaw(this, &FTG_ExporterImpl::SpawnTab_ParameterDefaults))
+		.SetDisplayName(LOCTEXT("ParametersTab", "Parameters"))
+		.SetIcon(FSlateIcon(FAppStyle::GetAppStyleSetName(), "LevelEditor.Tabs.Details"));
+
+	InTabManager->RegisterTabSpawner(FTG_EditorTabs::NodePreviewTabId, FOnSpawnTab::CreateRaw(this, &FTG_ExporterImpl::SpawnTab_NodePreview))
+		.SetDisplayName(LOCTEXT("NodePreviewTab", "Node Preview"))
+		.SetIcon(FSlateIcon(FAppStyle::GetAppStyleSetName(), "LevelEditor.Tabs.Viewports"));
+
+	InTabManager->RegisterTabSpawner(FTG_EditorTabs::PreviewSettingsTabId, FOnSpawnTab::CreateRaw(this, &FTG_ExporterImpl::SpawnTab_PreviewSettings))
+		.SetDisplayName(LOCTEXT("PreviewSettingsTab", "3D Preview Settings"))
+		.SetIcon(FSlateIcon(FAppStyle::GetAppStyleSetName(), "LevelEditor.Tabs.Details"));
+
+	InTabManager->RegisterTabSpawner(FTG_EditorTabs::OutputTabId, FOnSpawnTab::CreateRaw(this, &FTG_ExporterImpl::SpawnTab_ExportSettings))
+		.SetDisplayName(LOCTEXT("ExportSettingsTab", "Export Settings"))
+		.SetIcon(FSlateIcon(FAppStyle::GetAppStyleSetName(), "LevelEditor.Tabs.Details"));
+}
+
+void FTG_ExporterImpl::UnregisterTabSpawners(const TSharedPtr<FTabManager> InTabManager)
+{
+	InTabManager->UnregisterTabSpawner(FTG_EditorTabs::ViewportTabId);
+	InTabManager->UnregisterTabSpawner(FTG_EditorTabs::ParameterDefaultsTabId);
+	InTabManager->UnregisterTabSpawner(FTG_EditorTabs::NodePreviewTabId);
+	InTabManager->UnregisterTabSpawner(FTG_EditorTabs::OutputTabId);
+	InTabManager->UnregisterTabSpawner(FTG_EditorTabs::PreviewSettingsTabId);
+}
+TSharedRef<SDockTab> FTG_ExporterImpl::SpawnTab_PreviewSettings(const FSpawnTabArgs& Args)
+{
+	check(Args.GetTabId() == FTG_EditorTabs::PreviewSettingsTabId);
+
+	TSharedPtr<SDockTab> SettingsTab = SNew(SDockTab)
+		[
+			PreviewSettingsView.ToSharedRef()
+		];
+
+	if (TextureGraphPtr)
+	{
+		GetPreviewSettingsView()->SetObject(TextureGraphPtr->GetSettings(), true);
+	}
+
+	return SettingsTab.ToSharedRef();
+}
+TSharedRef<SDockTab> FTG_ExporterImpl::SpawnTab_ExportSettings(const FSpawnTabArgs& Args)
+{
+	check(Args.GetTabId() == FTG_EditorTabs::OutputTabId);
+
+	TSharedPtr<SDockTab> SettingsTab = SNew(SDockTab)
+		[
+			SNew(SVerticalBox)
+			
+			+ SVerticalBox::Slot()
+			.FillHeight(1)
+			[
+				SNew(SBorder)
+				.BorderImage(FAppStyle::GetBrush("ToolPanel.GroupBorder"))
+				[
+					SNew(SScrollBox)
+						+SScrollBox::Slot()
+						.VAlign(VAlign_Fill)
+						.FillSize(1.0)
+						[
+							ExportSettingsView.ToSharedRef()
+						]
+				]
+			]
+
+			+ SVerticalBox::Slot()
+			.AutoHeight()
+			[
+				SNew(SPrimaryButton)
+				.Text(LOCTEXT("Export", "Export"))
+				.OnClicked_Raw(this, &FTG_ExporterImpl::OnExportClicked, EAppReturnType::Ok)
+			]
+			
+		];
+
+	if (TextureGraphPtr)
+	{
+		GetExportSettingsView()->SetObject(ExportSettings, true);
+	}
+
+	return SettingsTab.ToSharedRef();
+}
+TSharedRef<SDockTab> FTG_ExporterImpl::SpawnTab_Viewport(const FSpawnTabArgs& Args)
+{
+	check(Args.GetTabId() == FTG_EditorTabs::ViewportTabId);
+
+	TSharedRef< SDockTab > DockableTab =
+		SNew(SDockTab);
+
+	AssetEditorViewportFactoryFunction MakeViewportFunc = [this](const FAssetEditorViewportConstructionArgs& InArgs)
+	{
+		return SNew(STG_EditorViewport)
+			.InTextureGraph(TextureGraphPtr);
+	};
+
+	// Create a new tab
+	ViewportTabContent = MakeShareable(new FEditorViewportTabContent());
+	//ViewportTabContent->OnViewportTabContentLayoutChanged().AddRaw(this, &FTG_ExporterImpl::OnEditorLayoutChanged);
+
+	const FString LayoutId = FString("TG_EditorViewport");
+	ViewportTabContent->Initialize(MakeViewportFunc, DockableTab, LayoutId);
+
+	// This call must occur after the toolbar is initialized.
+	SetViewportPreviewMesh();
+	
+	return DockableTab;
+}
+
+TSharedRef<SDockTab> FTG_ExporterImpl::SpawnTab_ParameterDefaults(const FSpawnTabArgs& Args)
+{
+	check(Args.GetTabId() == FTG_EditorTabs::ParameterDefaultsTabId);
+
+	return SNew(SDockTab)
+		[
+			SNew(SBox)
+				[
+					ParametersView.ToSharedRef()
+				]
+		];
+}
+
+
+TSharedRef<SDockTab> FTG_ExporterImpl::SpawnTab_NodePreview(const FSpawnTabArgs& Args)
+{
+	check(Args.GetTabId() == FTG_EditorTabs::NodePreviewTabId);
+
+	return SNew(SDockTab)
+		[
+			SNew(SVerticalBox)
+			
+			+ SVerticalBox::Slot()
+			.AutoHeight()
+			[
+				SAssignNew(OutputNodesComboBoxWidget, SComboBox<TSharedPtr<FName>>)
+				.OptionsSource(&OutputNodesList)
+				.OnGenerateWidget_Raw(this, &FTG_ExporterImpl::GenerateOutputComboItem)
+				.OnSelectionChanged_Raw(this, &FTG_ExporterImpl::OnOutputSelectionChanged)
+				[
+					SNew(STextBlock)
+					.Text_Lambda([this] ()
+					{
+						FText ComboTitleText = FText::FromString(TEXT("No TextureGraph selected, or it has no Outputs")); 
+						if (IsValid(SelectedNode))
+						{
+							ComboTitleText = FText::FromName(SelectedNode->GetExpression()->GetTitleName());
+						}
+						else if (!OutputNodesList.IsEmpty())
+						{
+							ComboTitleText = FText::FromName(*OutputNodesList[0]);
+						}
+						return ComboTitleText;
+					})
+				]
+			]
+			
+			+ SVerticalBox::Slot()
+			[
+		 		NodePreview.ToSharedRef()
+			]
+		];
+}
+void FTG_ExporterImpl::SetViewportPreviewMesh()
+{
+	if (TextureGraphPtr)
+	{
+		TObjectPtr<UStaticMesh> PreviewMesh = TextureGraphPtr->GetSettings()->GetPreviewMesh();
+		// Set the preview mesh for the material.  
+		if (!PreviewMesh || !SetPreviewAsset(PreviewMesh))
+		{
+			// The material preview mesh couldn't be found or isn't loaded. Fallback to the one of the primitive types.
+			GetEditorViewport()->InitPreviewMesh();
+		}
+	}
+}
+TSharedPtr<STG_EditorViewport> FTG_ExporterImpl::GetEditorViewport() const
+{
+	if (ViewportTabContent.IsValid())
+	{
+		// we can use static cast here b/c we know in this editor we will have a static mesh viewport 
+		return StaticCastSharedPtr<STG_EditorViewport>(ViewportTabContent->GetFirstViewport());
+	}
+
+	return nullptr;
+}
+bool FTG_ExporterImpl::SetPreviewAsset(UObject* InAsset)
+{
+	if (GetEditorViewport().IsValid())
+	{
+		return GetEditorViewport()->SetPreviewAsset(InAsset);
+	}
+	return false;
+}
+// Function to generate combo box items
+TSharedRef<SWidget> FTG_ExporterImpl::GenerateOutputComboItem(TSharedPtr<FName> InItem)
+{
+	return SNew(STextBlock).Text(FText::FromName(*InItem));
+}
+
+// Function called when the selection changes
+void FTG_ExporterImpl::OnOutputSelectionChanged(TSharedPtr<FName> SelectedItem, ESelectInfo::Type SelectInfo)
+{
+	if (SelectedItem.IsValid())
+	{
+		FName SelectedNodeName = *SelectedItem;
+		FTG_Id SelectedNodeId = FTG_Id::INVALID;
+
+		TextureGraphPtr->Graph()->ForEachNodes([this, &SelectedNodeId, SelectedNodeName](const UTG_Node* Node, uint32 Index)
+		{
+			if (Node && Node->GetExpression()->IsA(UTG_Expression_Output::StaticClass()))
+			{
+				// choose a default node
+				if (!SelectedNodeId.IsValid())
+					SelectedNodeId = Node->GetId();
+
+				// check if this is our selected node
+				if (Node->GetExpression()->GetTitleName() == SelectedNodeName)
+				{
+					SelectedNodeId = Node->GetId();
+				}
+			}
+		});
+		
+		SelectedNode = TextureGraphPtr->Graph()->GetNode(SelectedNodeId);
+		NodePreview->SelectionChanged(SelectedNode);
+	
+	}
+}
+
+FReply FTG_ExporterImpl::OnExportClicked(EAppReturnType::Type ButtonID)
+{
+	if (TextureGraphPtr)
+	{
+		FTG_HelperFunctions::ExportAsync(TextureGraphPtr, "", "", this->TargetExportSettings, false);
+	}
+	return FReply::Handled();
+}
+void FTG_ExporterImpl::SetTextureGraphToExport(UTextureGraph* InTextureGraph)
+{
+	// clear out previous handles
+	Cleanup();
+	
+	// overwrite the original TextureGraph in place by constructing a new one with the same name
+	FObjectDuplicationParameters Params = InitStaticDuplicateObjectParams(InTextureGraph, InTextureGraph->GetOuter(), NAME_None,
+	~RF_Standalone, UTextureGraph::StaticClass(), EDuplicateMode::Normal, EInternalObjectFlags::None);
+	
+	TextureGraphPtr = Cast<UTextureGraph>(StaticDuplicateObjectEx(Params));
+	//Exporter gets notified when rendering is done
+	TextureGraphPtr->OnRenderDone.BindRaw(this, &FTG_ExporterImpl::OnRenderingDone);
+
+	// force open export window
+	FGlobalTabmanager::Get()->TryInvokeTab(FTG_EditorTabs::TextureExporterTabId);
+	
+	
+	UpdateParametersUI();
+
+	UpdateExportSettingsUI();
+	
+	// Update list of output nodes in 2d View
+	OutputNodesList.Empty();
+	TextureGraphPtr->Graph()->ForEachNodes(
+			[&](const UTG_Node* node, uint32 index)
+	{
+		if(UTG_Expression_Output* OutputExpression = Cast<UTG_Expression_Output>(node->GetExpression()))
+		{
+			OutputNodesList.Add(MakeShared<FName>(OutputExpression->GetTitleName()));			
+		}
+	});
+	if (!OutputNodesList.IsEmpty())
+	{
+		OutputNodesComboBoxWidget->SetSelectedItem(OutputNodesList[0]);
+	}
+	OutputNodesComboBoxWidget->RefreshOptions();
+	
+	FViewportSettings& ViewportSettings = TextureGraphPtr->GetSettings()->GetViewportSettings();
+
+	ViewportSettings.OnViewportMaterialChangedEvent.AddRaw(this, &FTG_ExporterImpl::OnViewportMaterialChanged);
+	ViewportSettings.OnMaterialMappingChangedEvent.AddRaw(this, &FTG_ExporterImpl::OnMaterialMappingChanged);
+	TextureGraphPtr->GetSettings()->OnPreviewMeshChangedEvent.AddRaw(this, &FTG_ExporterImpl::SetViewportPreviewMesh);
+	TextureGraphPtr->Graph()->OnGraphChangedDelegate.AddRaw(this, &FTG_ExporterImpl::OnGraphChanged);
+	GetPreviewSettingsView()->SetObject(TextureGraphPtr->GetSettings(), true);
+
+	GetEditorViewport()->SetTextureGraph(TextureGraphPtr);
+	OnViewportMaterialChanged();
+	SetViewportPreviewMesh();
+	
+	TextureGraphEngine::RegisterErrorReporter(TextureGraphPtr, std::make_shared<FTextureGraphErrorReporter>());
+
+}
+void FTG_ExporterImpl::OnGraphChanged(UTG_Graph* InGraph, UTG_Node* InNode, bool Tweaking)
+{
+	if (TextureGraphPtr)
+	{
+		TextureGraphPtr->TriggerUpdate(Tweaking);
+	
+		// UpdateParametersUI();
+
+		if (InNode->IsA<UTG_Expression_Output>())
+		{
+			TextureGraphPtr->UpdateGlobalTGSettings();
+		}
+		RefreshViewport();
+	}
+	// OutputView->ForceRefresh();
+}
+
+void FTG_ExporterImpl::OnRenderingDone(UMixInterface* TextureGraph, const FInvalidationDetails* Details)
+{
+	if (TextureGraph != nullptr && TextureGraph == TextureGraphPtr)
+	{
+		// refresh node preview
+		NodePreview->Update();
+	}
+}
+void FTG_ExporterImpl::OnViewportMaterialChanged()
+{
+	const UTG_Node* FirstTargetNode = nullptr;
+	TextureGraphPtr->Graph()->ForEachNodes([&](const UTG_Node* node, uint32 index)
+		{
+			if (UTG_Expression_Output* OutputExpression = dynamic_cast<UTG_Expression_Output*>(node->GetExpression()))
+			{
+				if (!FirstTargetNode)
+				{
+					FirstTargetNode = node;
+				}
+			}
+		});
+
+	FViewportSettings& ViewportSettings = TextureGraphPtr->GetSettings()->GetViewportSettings();
+
+	if (FirstTargetNode && ViewportSettings.MaterialMappingInfos.Num() > 0)
+	{
+		ViewportSettings.SetDefaultTarget(FirstTargetNode->GetNodeName());
+	}
+
+	GetEditorViewport()->GenerateRendermodeToolbar();
+	GetEditorViewport()->InitRenderModes(TextureGraphPtr);
+}
+
+void FTG_ExporterImpl::OnMaterialMappingChanged()
+{
+	GetEditorViewport()->UpdateRenderMode();
+}
+
+void FTG_ExporterImpl::UpdateExportSettingsUI()
+{
+	// recreate export settings UI
+	ExportSettings = NewObject<UTG_ExportSettings>(TextureGraphPtr);
+	
+	UTG_Graph* Graph = TextureGraphPtr->Graph();
+
+	Graph->ForEachNodes([&](const UTG_Node* node, uint32 index)
+	{
+		if (UTG_Expression_Output* OutputExpression = dynamic_cast<UTG_Expression_Output*>(node->GetExpression()))
+		{
+			ExportSettings->OutputExpressionsInfos.Add(FOutputExpressionInfo{OutputExpression->GetTitleName(), node->GetId()});
+		}
+	});
+
+	GetExportSettingsView()->SetObject(ExportSettings);
+}
+void FTG_ExporterImpl::UpdateParametersUI()
+{
+	UTG_Graph* Graph = TextureGraphPtr->Graph();
+	FTG_Ids Ids = Graph->GetParamIds();
+
+	// Create a new object to set for the view.
+	Parameters = NewObject<UTG_Parameters>();
+
+	for (const FTG_Id& id : Ids)
+	{
+		UTG_Pin* Pin = Graph->GetPin(id);
+
+		if (Pin && (Pin->IsInput() || Pin->IsSetting()))
+		{
+			FTG_ParameterInfo Info{ id, Pin->GetAliasName() };
+			Parameters->Parameters.Add(Info);
+		}
+	}
+	Parameters->TextureGraph = Graph;
+
+	GetParametersView()->SetObject(Parameters);
+}
+void FTG_ExporterImpl::Tick(float DeltaTime)
+{
+	RefreshViewport();
+}
+
+TStatId FTG_ExporterImpl::GetStatId() const
+{
+	return TStatId();
+}
+
+void FTG_ExporterImpl::AddReferencedObjects(FReferenceCollector& Collector)
+{
+	Collector.AddReferencedObject(Parameters);
+	Collector.AddReferencedObject(TextureGraphPtr);
+	Collector.AddReferencedObject(ExportSettings);
+}
+
+void FTG_ExporterImpl::RefreshViewport()
+{
+	if (GetEditorViewport().IsValid())
+	{
+		GetEditorViewport()->RefreshViewport();
+	}
+}
+
+///////////////////////////////////////////////////////////////////////////////////
+FTG_Exporter::FTG_Exporter()
+	: Impl(MakeUnique<FTG_ExporterImpl>())
+{
+}
+
+FTG_Exporter::~FTG_Exporter()
+{
+}
+
+void FTG_Exporter::SetTextureGraphToExport(UTextureGraph* InTextureGraph)
+{
+	Impl->SetTextureGraphToExport(InTextureGraph);
+}
+
+#undef LOCTEXT_NAMESPACE 
diff --git a/Engine/Plugins/Experimental/TextureGraph/Source/TextureGraphEditor/Private/TG_RenderModeManager.cpp b/Engine/Plugins/Experimental/TextureGraph/Source/TextureGraphEditor/Private/TG_RenderModeManager.cpp
--- a/Engine/Plugins/Experimental/TextureGraph/Source/TextureGraphEditor/Private/TG_RenderModeManager.cpp
+++ b/Engine/Plugins/Experimental/TextureGraph/Source/TextureGraphEditor/Private/TG_RenderModeManager.cpp
@@ -13,10 +13,8 @@
 #include "Materials/Material.h"
 #include "Model/Mix/ViewportSettings.h"
 
-TG_RenderModeManager::TG_RenderModeManager(UTextureGraph* InTextureGraph) :
-	TextureGraph(InTextureGraph)
+TG_RenderModeManager::TG_RenderModeManager()
 {
-	check(InTextureGraph);
 }
 
 TG_RenderModeManager::~TG_RenderModeManager()
@@ -24,14 +22,14 @@ TG_RenderModeManager::~TG_RenderModeManager()
 	
 }
 
-void TG_RenderModeManager::ChangeRenderMode(FName NewRenderMode)
+void TG_RenderModeManager::ChangeRenderMode(FName NewRenderMode, UTextureGraph* TextureGraph)
 {
 	_lastRenderMode = _currentRenderMode;
 	_currentRenderMode = NewRenderMode;
-	UpdateRenderMode();
+	UpdateRenderMode(TextureGraph);
 }
 
-void TG_RenderModeManager::UpdateRenderMode()
+void TG_RenderModeManager::UpdateRenderMode(UTextureGraph* TextureGraph)
 {
 	UMixSettings* settings = TextureGraph->GetSettings();
 	size_t numTargets = settings->NumTargets();
@@ -39,7 +37,7 @@ void TG_RenderModeManager::UpdateRenderMode()
 	
 	if (numTargets != _renderModeMaterials.Num())
 	{
-		InitializeDefaultMaterials(numTargets);
+		InitializeDefaultMaterials(numTargets, TextureGraph);
 	}
 
 	for (size_t ti = 0; ti < numTargets; ti++)
@@ -93,7 +91,7 @@ bool TG_RenderModeManager::IsCurrentRenderModelLit(int targetId /* = 0*/)
 }
 
 
-void TG_RenderModeManager::InitializeDefaultMaterials(int totalTargets)
+void TG_RenderModeManager::InitializeDefaultMaterials(int totalTargets, UTextureGraph* TextureGraph)
 {
 	_renderModeMaterials.Empty();
 
diff --git a/Engine/Plugins/Experimental/TextureGraph/Source/TextureGraphEditor/Private/TG_RenderModeManager.h b/Engine/Plugins/Experimental/TextureGraph/Source/TextureGraphEditor/Private/TG_RenderModeManager.h
--- a/Engine/Plugins/Experimental/TextureGraph/Source/TextureGraphEditor/Private/TG_RenderModeManager.h
+++ b/Engine/Plugins/Experimental/TextureGraph/Source/TextureGraphEditor/Private/TG_RenderModeManager.h
@@ -41,16 +41,15 @@ protected:
 	FName							_currentRenderMode;			/// Currently applied render mode	
 	TMap<int, MaterialMap>			_renderModeMaterials;
 	TArray<FName>					_renderModesString;
-	UTextureGraph*					TextureGraph;
 
-	void							InitializeDefaultMaterials(int totalTargets);
+	void							InitializeDefaultMaterials(int totalTargets, UTextureGraph* TextureGraph);
 	RenderMaterial_BPPtr			GetTargetRenderModeMaterial(int targetId, FName renderMode);
 	void							BindBlobToMaterial(RenderMaterialPtr renderMaterial, TiledBlobPtr blobToBind, const FName& targetName);
+	virtual void					UpdateRenderMode(UTextureGraph* TextureGraph);
 public:
-									TG_RenderModeManager(UTextureGraph* InTextureGraph);
+									TG_RenderModeManager();
 	virtual							~TG_RenderModeManager();		
-	virtual void					ChangeRenderMode(FName NewRenderMode);
-	virtual void					UpdateRenderMode();
+	virtual void					ChangeRenderMode(FName NewRenderMode, UTextureGraph* TextureGraph);
 	void							Clear();
 	bool							IsCurrentRenderModelLit(int targetId = 0);
 	//////////////////////////////////////////////////////////////////////////
diff --git a/Engine/Plugins/Experimental/TextureGraph/Source/TextureGraphEditor/Private/TextureGraphEditorModule.cpp b/Engine/Plugins/Experimental/TextureGraph/Source/TextureGraphEditor/Private/TextureGraphEditorModule.cpp
--- a/Engine/Plugins/Experimental/TextureGraph/Source/TextureGraphEditor/Private/TextureGraphEditorModule.cpp
+++ b/Engine/Plugins/Experimental/TextureGraph/Source/TextureGraphEditor/Private/TextureGraphEditorModule.cpp
@@ -34,13 +34,14 @@ void FTextureGraphEditorModule::StartupModule()
 	PropertyEditorModule.RegisterCustomPropertyTypeLayout("TG_ParameterInfo", FOnGetPropertyTypeCustomizationInstance::CreateStatic(&FTG_ParameterInfoCustomization::Create));
 	PropertyEditorModule.RegisterCustomPropertyTypeLayout("TG_Texture", FOnGetPropertyTypeCustomizationInstance::CreateStatic(&FTG_TextureCustomization::Create));
 	PropertyEditorModule.RegisterCustomPropertyTypeLayout("TG_Variant", FOnGetPropertyTypeCustomizationInstance::CreateStatic(&FTG_VariantCustomization::Create));
-	auto ScalarIdentifier = MakeShared<FTG_ScalarTypeIdentifier>();
-	PropertyEditorModule.RegisterCustomPropertyTypeLayout("FloatProperty", FOnGetPropertyTypeCustomizationInstance::CreateStatic(&FTG_ScalarCustomization::Create), ScalarIdentifier);
+	PropertyEditorModule.RegisterCustomPropertyTypeLayout("FloatProperty", FOnGetPropertyTypeCustomizationInstance::CreateStatic(&FTG_ScalarCustomization::Create), MakeShared<FTG_ScalarTypeIdentifier>());
 	PropertyEditorModule.RegisterCustomPropertyTypeLayout("MaterialMappingInfo", FOnGetPropertyTypeCustomizationInstance::CreateStatic(&FTG_MaterialMappingInfoCustomization::Create));
 	PropertyEditorModule.RegisterCustomPropertyTypeLayout("ViewportSettings", FOnGetPropertyTypeCustomizationInstance::CreateStatic(&FTG_ViewportSettingsCustomization::Create));
 	PropertyEditorModule.RegisterCustomPropertyTypeLayout("OutputSettings", FOnGetPropertyTypeCustomizationInstance::CreateStatic(&FTG_OutputSettingsCustomization::Create));
 	PropertyEditorModule.RegisterCustomPropertyTypeLayout("TG_LevelsSettings", FOnGetPropertyTypeCustomizationInstance::CreateStatic(&FTG_LevelsSettingsCustomization::Create));
-
+	PropertyEditorModule.RegisterCustomPropertyTypeLayout("OutputExpressionInfo", FOnGetPropertyTypeCustomizationInstance::CreateStatic(&FTG_OutputExpressionInfoCustomization::Create));
+	PropertyEditorModule.NotifyCustomizationModuleChanged();
+	
 	// Register slate style overrides
 	FTG_Style::Register();
 	FTG_EditorCommands::Register();
@@ -51,6 +52,8 @@ void FTextureGraphEditorModule::StartupModule()
 	GraphPanelPinFactory = MakeShared<FTG_EditorGraphPanelPinFactory>();
 	FEdGraphUtilities::RegisterVisualPinFactory(GraphPanelPinFactory);
 	
+	TG_Exporter = MakeUnique<FTG_Exporter>();
+	
 	StartTextureGraphEngine();
 }
 
@@ -62,14 +65,20 @@ void FTextureGraphEditorModule::ShutdownModule()
 	FTG_EditorCommands::Unregister();
 	
 	FPropertyEditorModule& PropertyEditorModule = FModuleManager::GetModuleChecked<FPropertyEditorModule>("PropertyEditor");
-	PropertyEditorModule.UnregisterCustomPropertyTypeLayout("TG_ParameterInfo");
-	PropertyEditorModule.UnregisterCustomPropertyTypeLayout("TG_Texture");
-	auto ScalarIdentifier = MakeShared<FTG_ScalarTypeIdentifier>();
-	PropertyEditorModule.UnregisterCustomPropertyTypeLayout("FloatProperty", ScalarIdentifier);
-	PropertyEditorModule.UnregisterCustomPropertyTypeLayout("MaterialMappingInfo");
-	PropertyEditorModule.UnregisterCustomPropertyTypeLayout("OutputSettings");
-	PropertyEditorModule.UnregisterCustomPropertyTypeLayout("TG_LevelsSettings");
-
+	// Unregister the details customization
+	if (FModuleManager::Get().IsModuleLoaded("PropertyEditor"))
+	{
+		PropertyEditorModule.UnregisterCustomPropertyTypeLayout("TG_ParameterInfo");
+		PropertyEditorModule.UnregisterCustomPropertyTypeLayout("TG_Texture");
+		auto ScalarIdentifier = MakeShared<FTG_ScalarTypeIdentifier>();
+		PropertyEditorModule.UnregisterCustomPropertyTypeLayout("FloatProperty", ScalarIdentifier);
+		PropertyEditorModule.UnregisterCustomPropertyTypeLayout("MaterialMappingInfo");
+		PropertyEditorModule.UnregisterCustomPropertyTypeLayout("OutputSettings");
+		PropertyEditorModule.UnregisterCustomPropertyTypeLayout("TG_LevelsSettings");
+		PropertyEditorModule.UnregisterCustomPropertyTypeLayout("TG_OutputExpressionInfoCustomization");
+		
+		PropertyEditorModule.NotifyCustomizationModuleChanged();
+	}
 	FEdGraphUtilities::UnregisterVisualPinFactory(GraphPanelPinFactory);
 	FEdGraphUtilities::UnregisterVisualNodeFactory(GraphNodeFactory);
 	// Unregister slate style overrides
diff --git a/Engine/Plugins/Experimental/TextureGraph/Source/TextureGraphEditor/Public/TG_Exporter.h b/Engine/Plugins/Experimental/TextureGraph/Source/TextureGraphEditor/Public/TG_Exporter.h
new file mode 100644
--- /dev/null
+++ b/Engine/Plugins/Experimental/TextureGraph/Source/TextureGraphEditor/Public/TG_Exporter.h
@@ -0,0 +1,29 @@
+// Copyright Epic Games, Inc. All Rights Reserved.
+
+#pragma once
+
+#include "Templates/UniquePtr.h"
+
+class UTextureGraph;
+
+struct FTG_Exporter
+{
+	// Initializes the global state of the exporter (commands, tab spawners, etc):
+	FTG_Exporter();
+
+	// Destructor declaration purely so that we can pimpl:
+	~FTG_Exporter();
+
+	/** Sets the current texture graph to be used with the exporter */
+	void SetTextureGraphToExport(UTextureGraph* InTextureGraph);
+
+private:
+	TUniquePtr< struct FTG_ExporterImpl > Impl;
+
+	// prevent copying:
+	FTG_Exporter(const FTG_Exporter&);
+	FTG_Exporter(FTG_Exporter&&);
+	FTG_Exporter& operator=(FTG_Exporter const&);
+	FTG_Exporter& operator=(FTG_Exporter&&);
+};
+
diff --git a/Engine/Plugins/Experimental/TextureGraph/Source/TextureGraphEditor/Public/TextureGraphEditorModule.h b/Engine/Plugins/Experimental/TextureGraph/Source/TextureGraphEditor/Public/TextureGraphEditorModule.h
--- a/Engine/Plugins/Experimental/TextureGraph/Source/TextureGraphEditor/Public/TextureGraphEditorModule.h
+++ b/Engine/Plugins/Experimental/TextureGraph/Source/TextureGraphEditor/Public/TextureGraphEditorModule.h
@@ -4,6 +4,7 @@
 #include "CoreMinimal.h"
 #include "Modules/ModuleManager.h"
 #include "Toolkits/IToolkit.h"
+#include "TG_Exporter.h"
 class IAssetTools;
 class IAssetTypeActions;
 class ITG_Editor;
@@ -22,7 +23,7 @@ private:
 	FTickerDelegate					TickDelegate;
 	/** Handle of the delegate to run the Tick method in charge of running TextureGraphEngine update*/
 	FTSTicker::FDelegateHandle		TickDelegateHandle;
-	
+	TUniquePtr<FTG_Exporter>		TG_Exporter;
 
 protected:
 	TSharedPtr<FTG_EditorGraphNodeFactory> GraphNodeFactory;
@@ -31,10 +32,15 @@ public:
 	/** IModuleInterface implementation */
 	virtual void					StartupModule() override;
 	virtual void					ShutdownModule() override;
+
+	
 	virtual void					StartTextureGraphEngine();
 	virtual void					ShutdownTextureGraphEngine();
 	bool							Tick(float deltaTime);
 	void							RegisterAssetTypeAction(IAssetTools& AssetTools, TSharedRef<IAssetTypeActions> Action);
 	void							UnRegisterAllAssetTypeActions();
 	static TSharedRef<ITG_Editor>	CreateTextureGraphEditor(const EToolkitMode::Type Mode, const TSharedPtr< class IToolkitHost >& InitToolkitHost, UTextureGraph* InTextureGraph);
+
+	/** Returns a reference to the Blueprint Debugger state object */
+	const TUniquePtr<FTG_Exporter>& GetTextureExporter() const { return TG_Exporter; }
 };
diff --git a/Engine/Plugins/Experimental/TextureGraph/Source/TextureGraphEditor/TextureGraphEditor.Build.cs b/Engine/Plugins/Experimental/TextureGraph/Source/TextureGraphEditor/TextureGraphEditor.Build.cs
--- a/Engine/Plugins/Experimental/TextureGraph/Source/TextureGraphEditor/TextureGraphEditor.Build.cs
+++ b/Engine/Plugins/Experimental/TextureGraph/Source/TextureGraphEditor/TextureGraphEditor.Build.cs
@@ -29,6 +29,7 @@ public class TextureGraphEditor : ModuleRules
 				"Projects",
 				"ToolWidgets",
 				"KismetWidgets",
+				"WorkspaceMenuStructure"
 				// ... add other public dependencies that you statically link with here ...
 			}
 			);
