Trigonometry enum setting change.

diff --git a/Engine/Plugins/Experimental/TextureGraph/Source/TextureGraph/Public/Expressions/Maths/TG_Expression_Maths_OneInput.h b/Engine/Plugins/Experimental/TextureGraph/Source/TextureGraph/Public/Expressions/Maths/TG_Expression_Maths_OneInput.h
--- a/Engine/Plugins/Experimental/TextureGraph/Source/TextureGraph/Public/Expressions/Maths/TG_Expression_Maths_OneInput.h
+++ b/Engine/Plugins/Experimental/TextureGraph/Source/TextureGraph/Public/Expressions/Maths/TG_Expression_Maths_OneInput.h
@@ -28,8 +28,8 @@ public:
 	FTG_Variant							Angle;
 
 	// What trigonometric function to apply at the angle
-	UPROPERTY(meta = (TGType = "TG_Input"))
-	TEnumAsByte<ETrigFunction>			Function = ETrigFunction::Sin;
+	UPROPERTY(EditAnywhere, BlueprintReadWrite, Category = NoCategory, meta = (TGType = "TG_Setting", RegenPinsOnChange))
+	ETrigFunction						Function = ETrigFunction::Sin;
 
 	virtual FText						GetTooltipText() const override { return FText::FromString(TEXT("A variety of Trigonometric functions.")); } 
 	virtual FName						GetCategory() const override { return TG_Category::Maths; } 
diff --git a/Engine/Plugins/Experimental/TextureGraph/Source/TextureGraphEngine/Transform/Expressions/T_Maths_OneInput.h b/Engine/Plugins/Experimental/TextureGraph/Source/TextureGraphEngine/Transform/Expressions/T_Maths_OneInput.h
--- a/Engine/Plugins/Experimental/TextureGraph/Source/TextureGraphEngine/Transform/Expressions/T_Maths_OneInput.h
+++ b/Engine/Plugins/Experimental/TextureGraph/Source/TextureGraphEngine/Transform/Expressions/T_Maths_OneInput.h
@@ -54,7 +54,7 @@ DECLARE_EMPTY_GLOBAL_SHADER_DERIVED_FROM(FSH_Round, FSH_MathsOp_OneInput);
 /// Trigonometric functions
 //////////////////////////////////////////////////////////////////////////
 UENUM(BlueprintType)
-enum ETrigFunction : int
+enum class ETrigFunction : uint8
 {
 	Sin							UMETA(DisplayName = "Sine"),
 	Cos							UMETA(DisplayName = "Cosine"),
