Warp fixes

diff --git a/Engine/Plugins/Experimental/TextureGraph/Shaders/Expressions/Expression_Warp.usf b/Engine/Plugins/Experimental/TextureGraph/Shaders/Expressions/Expression_Warp.usf
--- a/Engine/Plugins/Experimental/TextureGraph/Shaders/Expressions/Expression_Warp.usf
+++ b/Engine/Plugins/Experimental/TextureGraph/Shaders/Expressions/Expression_Warp.usf
@@ -35,7 +35,7 @@ float4 FSH_DirectionalWarp(float2 InUV_ : TEXCOORD0) : SV_Target0
 	/// We need to calculate the UV that will get warped to the current UV. So we need
 	/// to move in the opposite direction of the angle to find the UV that will end
 	/// replacing the current pixel. Hence we add 180 degrees to it (or half PI in radians)
-	float AngleValue = AngleRad + HALF_PI;
+	float AngleValue = AngleRad * HALF_PI;
 	float C = cos(AngleValue);
 	float S = sin(AngleValue);
 	float2 Dir = float2(C, S);
@@ -56,7 +56,7 @@ float4 FSH_NormalWarp(float2 InUV_ : TEXCOORD0) : SV_Target0
 	/// We need to calculate the UV that will get warped to the current UV. So we need
 	/// to move in the opposite direction of the angle to find the UV that will end
 	/// replacing the current pixel. Hence we add 180 degrees to it (or half PI in radians)
-	float AngleValue = MaskValue * QUARTER_PI + HALF_PI;
+	float AngleValue = MaskValue * HALF_PI;
 	float C = cos(AngleValue);
 	float S = sin(AngleValue);
 	float2 Dir = float2(C, S);
