Fix material ID crash.

diff --git a/Engine/Plugins/Experimental/TextureGraph/Source/TextureGraphEngine/Data/Blob.cpp b/Engine/Plugins/Experimental/TextureGraph/Source/TextureGraphEngine/Data/Blob.cpp
--- a/Engine/Plugins/Experimental/TextureGraph/Source/TextureGraphEngine/Data/Blob.cpp
+++ b/Engine/Plugins/Experimental/TextureGraph/Source/TextureGraphEngine/Data/Blob.cpp
@@ -245,7 +245,6 @@ AsyncRawBufferPtr Blob::Raw()
 	/// Could possibly be going to another thread
 	/// Save the current temp Hash
 	CHashPtr PrevHash = Buffer->Hash(false);
-	check(!PrevHash || PrevHash->IsTemp());
 
 	return Buffer->Raw().then([this, PrevHash](RawBufferPtr RawObj)
 	{
@@ -261,10 +260,6 @@ AsyncRawBufferPtr Blob::Raw()
 				*BufferDesc.Name, PrevHashValue, Hash->Value(), BufferDesc.Width, BufferDesc.Height);
 		}
 
-		/// TODO: do we really need this?
-		if (PrevHash != nullptr)
-			TextureGraphEngine::GetBlobber()->UpdateHash(PrevHash->Value(), Hash);
-
 		return Buffer->Raw_Now();
 	});
 }
