Changing enums (tagged with RegenPinsOnChange meta) via Nodes UI, now correctly regenerates pins and the node UI.

diff --git a/Engine/Plugins/Experimental/TextureGraph/Source/TextureGraphEditor/Private/EdGraph/TG_EdGraphNode.cpp b/Engine/Plugins/Experimental/TextureGraph/Source/TextureGraphEditor/Private/EdGraph/TG_EdGraphNode.cpp
--- a/Engine/Plugins/Experimental/TextureGraph/Source/TextureGraphEditor/Private/EdGraph/TG_EdGraphNode.cpp
+++ b/Engine/Plugins/Experimental/TextureGraph/Source/TextureGraphEditor/Private/EdGraph/TG_EdGraphNode.cpp
@@ -524,6 +524,14 @@ void UTG_EdGraphNode::PinDefaultValueChangedWithTweaking(UEdGraphPin* Pin, bool
 	{
 		Pin->DefaultObject = TGPin->EditSelfVar()->GetAs<TObjectPtr<UObject>>();
 	}
+
+#if WITH_EDITOR
+	// Update Node UI if the property has meta RegenPinsOnChange
+	if(Property != nullptr && Property->HasMetaData("RegenPinsOnChange"))
+	{
+		UpdatePinsAndReconstructNode();
+	}
+#endif
 }
 
 void UTG_EdGraphNode::PinConnectionListChanged(UEdGraphPin* Pin)
@@ -660,6 +668,19 @@ bool UTG_EdGraphNode::UpdateEdPinDefaultValue(UEdGraphPin* EdPin, const UTG_EdGr
 	return bShouldUpdatePinsVisibility;
 }
 
+void UTG_EdGraphNode::UpdatePinsAndReconstructNode()
+{
+	UpdateInputPinsVisibility();
+		
+	// cache advanced visibility
+	TEnumAsByte<ENodeAdvancedPins::Type> CachedAdvancedPinDisplay = this->AdvancedPinDisplay;
+	ReconstructNode();
+	this->AdvancedPinDisplay = CachedAdvancedPinDisplay;
+
+	//Tell Editor to update details
+	Cast<UTG_EdGraph>(GetGraph())->RefreshEditorDetails();
+}
+
 void UTG_EdGraphNode::OnNodeChanged(UTG_Node* InNode)
 {
 	bool bShouldUpdatePinsVisibility = false;
@@ -673,15 +694,7 @@ void UTG_EdGraphNode::OnNodeChanged(UTG_Node* InNode)
 	// if there was a pin which triggered regeneration of all other input pins (used meta "RegenPinsOnChange"), we update them
 	if (bShouldUpdatePinsVisibility)
 	{
-		UpdateInputPinsVisibility();
-		
-		// cache advanced visibility
-		TEnumAsByte<ENodeAdvancedPins::Type> CachedAdvancedPinDisplay = this->AdvancedPinDisplay;
-		ReconstructNode();
-		this->AdvancedPinDisplay = CachedAdvancedPinDisplay;
-
-		//Tell Editor to update details
-		Cast<UTG_EdGraph>(GetGraph())->RefreshEditorDetails();
+		UpdatePinsAndReconstructNode();
 	}
 #endif
 	
diff --git a/Engine/Plugins/Experimental/TextureGraph/Source/TextureGraphEditor/Public/EdGraph/TG_EdGraphNode.h b/Engine/Plugins/Experimental/TextureGraph/Source/TextureGraphEditor/Public/EdGraph/TG_EdGraphNode.h
--- a/Engine/Plugins/Experimental/TextureGraph/Source/TextureGraphEditor/Public/EdGraph/TG_EdGraphNode.h
+++ b/Engine/Plugins/Experimental/TextureGraph/Source/TextureGraphEditor/Public/EdGraph/TG_EdGraphNode.h
@@ -103,6 +103,7 @@ protected:
 
 	void OnNodePostEvaluate(const FTG_EvaluationContext* Fts_EvaluationContext);
 	bool UpdateEdPinDefaultValue(UEdGraphPin* EdPin, const UTG_EdGraphSchema* Schema);
+	void UpdatePinsAndReconstructNode();
 
 	void OnUpdateCommentText(const FString& NewComment) override;
 	void OnCommentBubbleToggled(bool bInCommentBubbleVisible) override;
