Implement Levels and auto levels better

diff --git a/Engine/Plugins/Experimental/TextureGraph/Shaders/Expressions/Expression_Levels.usf b/Engine/Plugins/Experimental/TextureGraph/Shaders/Expressions/Expression_Levels.usf
--- a/Engine/Plugins/Experimental/TextureGraph/Shaders/Expressions/Expression_Levels.usf
+++ b/Engine/Plugins/Experimental/TextureGraph/Shaders/Expressions/Expression_Levels.usf
@@ -1,36 +1,78 @@
 // Copyright Epic Games, Inc. All Rights Reserved.
 #include "/Engine/Public/Platform.ush"
 #include "/Plugin/TextureGraph/SamplerStates.ush"
+#include "/Plugin/TextureGraph/ShaderUtil.ush"
 
-Texture2D Input;
-float MinValue;
-float MaxValue;
-float Gamma;
 
-float3 GammaCorrection(float3 Clr, float GammaValue)
-{
-	float invGamma = 1.0 / GammaValue;
-	return pow(Clr, float3(invGamma, invGamma, invGamma));
-}
+#ifndef CONVERT_TO_GRAYSCALE
+#define CONVERT_TO_GRAYSCALE 0
+#endif
+
+#ifndef AUTO_LEVELS
+#define AUTO_LEVELS 0
+#endif
+
+Texture2D SourceTexture;
+float LowValue;
+float HighValue;
+float MidValue;
 
-float3 LevelsControlInputRange(float3 Clr, float3 MinInput, float3 MaxInput)
+float DoAutoLevel;
+float MidPercentage;
+
+Texture2D Histogram;
+
+
+float LevelsControlInputRange(float Value, float LowInput, float HighInput)
 {
-	return min(max(Clr - MinInput, float3(0, 0, 0)) / (float3(MaxInput) - MinInput), float3(1, 1, 1));
+	return min(max(Value - LowInput, float(0)) / (float(HighInput) - LowInput), float(1));
 }
 
-float3 LevelsControlInput(float3 Clr, float3 MinInput, float GammaValue, float3 MaxInput)
+float LevelsControlInputCurve(float Value, float LowInput, float HighInput, float MidInput)
 {
-	return GammaCorrection(LevelsControlInputRange(Clr, MinInput, MaxInput), GammaValue);
+	float ValueRanged = LevelsControlInputRange(Value, LowInput, HighInput);
+
+	// 0.5 = EvalRange(Mid) ^ Exponent
+	// thus
+	float MidRanged = saturate((MidInput - LowInput) / (HighInput - LowInput));
+	MidRanged = max(0.001, min(0.9999, MidRanged));
+
+	float InvGammaValue = log(MidRanged) / log(0.5);
+
+	return pow(ValueRanged, InvGammaValue);
 }
 
 float4 FSH_Levels(float2 uv : TEXCOORD0) : SV_Target0
 {
-	float4 Clr = Input.Sample(SamplerStates_NoBorder, uv);
-	// float4 MinInput = MinValue.Sample(SamplerStates_NoBorder, uv);
-	// float4 MaxInput = MaxValue.Sample(SamplerStates_NoBorder, uv);
+	// fetch the source texture value
+	float4 Color = SourceTexture.Sample(SamplerStates_NoBorder, uv);
+
+	// isolate the scalar value
+#if CONVERT_TO_GRAYSCALE
+	float SourceValue = Grayscale(Color.rgb);
+#else
+	float SourceValue = Color.x;
+#endif
+
+	// Standard path, Low and High come from specified uniforms
+	float LowInput = LowValue;
+	float HighInput = HighValue;
+
+	float MidInput = MidValue;
+
+#if AUTO_LEVELS
+
+	float4 HistoMinCount	= Histogram.Load(int3(0, 1, 0));
+	float4 HistoMaxCount	= Histogram.Load(int3(1, 1, 0));
+	float4 HistoMin			= Histogram.Load(int3(2, 1, 0));
+	float4 HistoMax			= Histogram.Load(int3(3, 1, 0));
+		
+	LowInput = 	HistoMin.w;
+	HighInput = HistoMax.w;
+#endif
 
-	float3 LevelClr = LevelsControlInput(Clr, float3(MinValue, MinValue, MinValue), Gamma, float3(MaxValue, MaxValue, MaxValue));
+	float3 LeveledColor = LevelsControlInputCurve(SourceValue, LowInput, HighInput, MidValue);
 	
-	return float4(LevelClr.rgb, Clr.a);
+	return float4(LeveledColor.rgb, 1.0);
 }
 
diff --git a/Engine/Plugins/Experimental/TextureGraph/Source/TextureGraph/Private/Expressions/FIlter/TG_Expression_Levels.cpp b/Engine/Plugins/Experimental/TextureGraph/Source/TextureGraph/Private/Expressions/FIlter/TG_Expression_Levels.cpp
--- a/Engine/Plugins/Experimental/TextureGraph/Source/TextureGraph/Private/Expressions/FIlter/TG_Expression_Levels.cpp
+++ b/Engine/Plugins/Experimental/TextureGraph/Source/TextureGraph/Private/Expressions/FIlter/TG_Expression_Levels.cpp
@@ -3,7 +3,7 @@
 #include "Expressions/Filter/TG_Expression_Levels.h"
 #include "FxMat/MaterialManager.h"
 #include "Job/JobBatch.h"
-#include "Transform/Expressions/T_Color.h"
+#include "Transform/Expressions/T_Levels.h"
 
 void FTG_LevelsSettings_VarPropertySerialize(FTG_Var::VarPropertySerialInfo& Info)
 {
@@ -89,32 +89,145 @@ bool FTG_LevelsSettings::SetMidFromMidExponent(float InExponent)
 	return false;
 }
 
+
+
+
+#if WITH_EDITOR
+
+void UTG_Expression_Levels::PostEditChangeProperty(FPropertyChangedEvent& PropertyChangedEvent)
+{
+	// First catch if Material changes
+	if (PropertyChangedEvent.GetPropertyName() == GET_MEMBER_NAME_CHECKED(UTG_Expression_Levels, LowValue))
+	{
+		SetLowValue(LowValue);
+		FeedbackPinValue(GET_MEMBER_NAME_CHECKED(UTG_Expression_Levels, MidValue), MidValue);
+	}
+	// Second catch if AttributeName changes
+	if (PropertyChangedEvent.GetPropertyName() == GET_MEMBER_NAME_CHECKED(UTG_Expression_Levels, HighValue))
+	{
+		SetHighValue(HighValue);
+		FeedbackPinValue(GET_MEMBER_NAME_CHECKED(UTG_Expression_Levels, MidValue), MidValue);
+	}
+
+	Super::PostEditChangeProperty(PropertyChangedEvent);
+}
+
+bool UTG_Expression_Levels::CanEditChange(const FProperty* InProperty) const
+{
+	bool bEditCondition = Super::CanEditChange(InProperty);
+	// if already set to false Or InProperty not directly owned by us, early out
+	if (!bEditCondition || this->GetClass() != InProperty->GetOwnerClass())
+	{
+		return bEditCondition;
+	}
+
+	const FName PropertyName = InProperty->GetFName();
+
+	// Specific logic associated with Property
+	if (PropertyName == GET_MEMBER_NAME_CHECKED(UTG_Expression_Levels, LowValue))
+	{
+		bEditCondition = (!IsAutoLevel());
+	}
+	if (PropertyName == GET_MEMBER_NAME_CHECKED(UTG_Expression_Levels, MidValue))
+	{
+		bEditCondition = (!IsAutoLevel());
+	}
+	if (PropertyName == GET_MEMBER_NAME_CHECKED(UTG_Expression_Levels, HighValue))
+	{
+		bEditCondition = (!IsAutoLevel());
+	}
+	if (PropertyName == GET_MEMBER_NAME_CHECKED(UTG_Expression_Levels, MidAutoLevels))
+	{
+		bEditCondition = (IsAutoLevel());
+	}
+
+	return bEditCondition;
+}
+
+#endif
+
+void UTG_Expression_Levels::SetLowValue(float InValue)
+{
+	if (Levels.SetLow(InValue))
+	{
+		LowValue = Levels.Low;
+		MidValue = Levels.Mid;
+		FeedbackPinValue(GET_MEMBER_NAME_CHECKED(UTG_Expression_Levels, MidValue), MidValue);
+	}
+	else
+	{
+		LowValue = Levels.Low;
+		FeedbackPinValue(GET_MEMBER_NAME_CHECKED(UTG_Expression_Levels, LowValue), LowValue);
+	}
+}
+
+void UTG_Expression_Levels::SetMidValue(float InValue)
+{
+	if (Levels.SetMid(InValue))
+	{
+		MidValue = Levels.Mid;
+	}
+	else
+	{
+		MidValue = Levels.Mid;
+		FeedbackPinValue(GET_MEMBER_NAME_CHECKED(UTG_Expression_Levels, MidValue), MidValue);
+	}
+}
+
+void UTG_Expression_Levels::SetHighValue(float InValue)
+{
+	if (Levels.SetHigh(InValue))
+	{
+		MidValue = Levels.Mid;
+		HighValue = Levels.High;
+		FeedbackPinValue(GET_MEMBER_NAME_CHECKED(UTG_Expression_Levels, MidValue), MidValue);
+	}
+	else
+	{
+		HighValue = Levels.High;
+		FeedbackPinValue(GET_MEMBER_NAME_CHECKED(UTG_Expression_Levels, HighValue), HighValue);
+	}
+}
+
 void UTG_Expression_Levels::Evaluate(FTG_EvaluationContext* InContext)
 {
 	Super::Evaluate(InContext);
 
-	const RenderMaterial_FXPtr RenderMaterial = TextureGraphEngine::GetMaterialManager()->CreateMaterial_FX<VSH_Simple, FSH_Levels>(TEXT("T_Levels"));
+	if (!Input) // No Input, black Output
+	{
+		Output = FTG_Texture::GetBlack();
+		return;
+	}
+
+	LevelsControl = MakeShared<FLevels>();
+
+	switch (LevelsExpressionType)
+	{
+	case ELevelsExpressionType::LowMidHigh:
+		LevelsControl->InitFromLowMidHigh(LowValue, MidValue, HighValue);
+		break;
+	case ELevelsExpressionType::AutoLowHigh:
+		LevelsControl->InitFromAutoLevels(MidAutoLevels);
+		break;
+	}
+
+	BufferDescriptor Desc = Output.GetBufferDescriptor();
+	Output = T_Levels::Create(InContext->Cycle, Desc, Input.RasterBlob, LevelsControl, InContext->TargetId);
+}
 
-	check(RenderMaterial);
+void UTG_Expression_HistogramScan::Evaluate(FTG_EvaluationContext* InContext)
+{
+	Super::Evaluate(InContext);
 
-	if (!Input)
+	if (!Input) // No Input, black Output
 	{
 		Output = FTG_Texture::GetBlack();
 		return;
 	}
-	
-	JobUPtr RenderJob = std::make_unique<Job>(InContext->Cycle->GetMix(), InContext->TargetId, std::static_pointer_cast<BlobTransform>(RenderMaterial));
+	LevelsControl = MakeShared<FLevels>();
+	LevelsControl->InitFromPositionContrast(Position, Contrast);
 
-	RenderJob
-		->AddArg(ARG_BLOB(Input, "Input"))
-		->AddArg(ARG_FLOAT(Levels.Low, "MinValue"))
-		->AddArg(ARG_FLOAT(Levels.High, "MaxValue"))
-		->AddArg(ARG_FLOAT(Levels.EvalMidExponent(), "Gamma"))
-		;
 
-	const FString Name = TEXT("Levels"); 
 	BufferDescriptor Desc = Output.GetBufferDescriptor();
-
-	Output = RenderJob->InitResult(Name, &Desc);
-	InContext->Cycle->AddJob(InContext->TargetId, std::move(RenderJob));
+	Output = T_Levels::Create(InContext->Cycle, Desc, Input.RasterBlob, LevelsControl, InContext->TargetId);
 }
diff --git a/Engine/Plugins/Experimental/TextureGraph/Source/TextureGraph/Private/TG_Var.cpp b/Engine/Plugins/Experimental/TextureGraph/Source/TextureGraph/Private/TG_Var.cpp
--- a/Engine/Plugins/Experimental/TextureGraph/Source/TextureGraph/Private/TG_Var.cpp
+++ b/Engine/Plugins/Experimental/TextureGraph/Source/TextureGraph/Private/TG_Var.cpp
@@ -149,7 +149,10 @@ void Generic_Simple_Serializer(FTG_Var::VarPropertySerialInfo& Info)
 	const T_ValueType VarValue = Info.Var->GetAs<T_ValueType>();
 
 	if (Info.CopyVarToProperty)
+	{
 		TProperty->SetPropertyValue(TProperty->template ContainerPtrToValuePtr<T_ValueType>(Info.Owner, Info.ClampedIndex()), VarValue);
+		//TProperty->SetValue_InContainer(Info.Owner, VarValue);
+	}
 	else
 		Info.Var->EditAs<T_ValueType>() = TProperty->GetPropertyValue(TProperty->template ContainerPtrToValuePtr<T_ValueType>(Info.Owner, Info.ClampedIndex()));
 }
@@ -217,7 +220,19 @@ void VarPropertySerializer_float(FTG_Var::VarPropertySerialInfo& Info)
 
 void VarPropertySerializer_bool(FTG_Var::VarPropertySerialInfo& Info)
 {
-	Generic_Simple_Serializer<FBoolProperty, bool>(Info);
+	//Generic_Simple_Serializer<FBoolProperty, bool>(Info);
+	FProperty* Property = Info.Owner->GetClass()->FindPropertyByName(Info.Arg.GetName());
+
+	const FBoolProperty* TProperty = CastField<FBoolProperty>(Property);
+	const bool VarValue = Info.Var->GetAs<bool>();
+
+	if (Info.CopyVarToProperty)
+	{
+		TProperty->SetPropertyValue(TProperty->template ContainerPtrToValuePtr<bool>(Info.Owner, Info.ClampedIndex()), VarValue);
+		//TProperty->SetValue_InContainer(Info.Owner, VarValue);
+	}
+	else
+		Info.Var->EditAs<bool>() = TProperty->GetPropertyValue(TProperty->template ContainerPtrToValuePtr<bool>(Info.Owner, Info.ClampedIndex()));
 }
 
 void VarPropertySerializer_FName(FTG_Var::VarPropertySerialInfo& Info)
diff --git a/Engine/Plugins/Experimental/TextureGraph/Source/TextureGraph/Public/Expressions/Filter/TG_Expression_Levels.h b/Engine/Plugins/Experimental/TextureGraph/Source/TextureGraph/Public/Expressions/Filter/TG_Expression_Levels.h
--- a/Engine/Plugins/Experimental/TextureGraph/Source/TextureGraph/Public/Expressions/Filter/TG_Expression_Levels.h
+++ b/Engine/Plugins/Experimental/TextureGraph/Source/TextureGraph/Public/Expressions/Filter/TG_Expression_Levels.h
@@ -9,6 +9,17 @@
 
 #include "TG_Expression_Levels.generated.h"
 
+
+
+UENUM(BlueprintType)
+enum class ELevelsExpressionType : uint8
+{
+	LowMidHigh = 0			UMETA(DisplayName = "Manual"),
+	AutoLowHigh				UMETA(DisplayName = "Auto Levels"),
+};
+
+struct FLevels;
+
 USTRUCT(BlueprintType)
 struct TEXTUREGRAPH_API FTG_LevelsSettings
 {
@@ -64,22 +75,49 @@ class TEXTUREGRAPH_API UTG_Expression_Levels : public UTG_Expression
 {
 	GENERATED_BODY()
 
+	FTG_LevelsSettings					Levels;
+
+	TSharedPtr<FLevels> 				LevelsControl;
+
 public:
+#if WITH_EDITOR
+	virtual void 						PostEditChangeProperty(FPropertyChangedEvent& PropertyChangedEvent) override;
+
+	// Used to implement EditCondition logic for both Node UI and Details View
+	virtual bool						CanEditChange(const FProperty* InProperty) const override;
+#endif
 
 	TG_DECLARE_EXPRESSION(TG_Category::Adjustment);
 	virtual void						Evaluate(FTG_EvaluationContext* InContext) override;
 
+	UPROPERTY(EditAnywhere, Category = NoCategory, meta = (TGType = "TG_Setting", PinDisplayName = "Type", PinNotConnectable = true, RegenPinsOnChange))
+	ELevelsExpressionType				LevelsExpressionType = ELevelsExpressionType::LowMidHigh;
+
 	// The input image to adjust the levels for
-	UPROPERTY(EditAnywhere, Category = NoCategory, meta = (TGType = "TG_Input", PinDisplayName = "", MD_HistogramLuminance))
+	UPROPERTY(EditAnywhere, Category = NoCategory, meta = (TGType = "TG_Input", PinDisplayName = ""))
 	FTG_Texture							Input;
 
-	// The Low-Mid-High settings of the Levels expression
-	UPROPERTY(EditAnywhere, Category = NoCategory, meta = (TGType = "TG_Setting", PinDisplayName = "", MD_LevelsSettings, PinNotConnectable = true))
-	FTG_LevelsSettings					Levels;
+	// The Low value of the Levels adjustment, any pixel under that value is set to black. Default is 0.
+	UPROPERTY(EditAnywhere, Setter,  Category = NoCategory, meta = (TGType = "TG_Setting", PinDisplayName = "Low", UIMin = "0", ClampMin = "0", UIMax = "1", ClampMax = "1", EditConditionHides, MD_ScalarEditor))
+	float								LowValue = 0;
+	void SetLowValue(float InValue);
+
+	// The mid value of the Levels adjustment, must be in the range [Min, Max] and the Default is 0.5.
+	// The mid value determine where the smoothing curve applying the midpoint filter is crossing 0.5.
+	UPROPERTY(EditAnywhere, Setter, Category = NoCategory, meta = (TGType = "TG_Setting", PinDisplayName = "Mid", UIMin = "0", ClampMin = "0", UIMax = "1", ClampMax = "1", EditConditionHides, MD_ScalarEditor))
+	float								MidValue = 0.5f; // Midtones
+	void SetMidValue(float InValue);
+
+	// The High value of the Levels adjustment, any pixel above that value is set to white. Default is 1.
+	UPROPERTY(EditAnywhere, Setter, Category = NoCategory, meta = (TGType = "TG_Setting", PinDisplayName = "High", UIMin = "0", ClampMin = "0", UIMax = "1", ClampMax = "1", EditConditionHides, MD_ScalarEditor))
+	float								HighValue = 1;
+	void SetHighValue(float InValue);
+
+
+	// The High value of the Levels adjustment, any pixel above that value is set to white. Default is 1.
+	UPROPERTY(EditAnywhere, Category = NoCategory, meta = (TGType = "TG_Setting", PinDisplayName = "Mid Auto Levels", UIMin = "0", ClampMin = "0", UIMax = "1", ClampMax = "1", EditConditionHides, MD_ScalarEditor))
+	float								MidAutoLevels = 0.5;
 
-	// TODO: Show Property after 5.4 release
-	//	UPROPERTY(EditAnywhere, Category = NoCategory, meta = (TGType = "TG_Setting", PinNotConnectable = true))
-	bool								AutoLevels = false;
 
 	// The output image filtered result of the Levels operator 
 	UPROPERTY(EditAnywhere, Category = NoCategory, meta = (TGType = "TG_Output", PinDisplayName = ""))
@@ -87,4 +125,36 @@ public:
 
 	virtual FText						GetTooltipText() const override { return FText::FromString(TEXT("Remaps shadows and highlights of the input. Any values less or equal to Low are mapped to black, any values, greater or equal to High are mapped to white, and any values inbetween have Gamma applied as an exponent.")); } 
 
+	bool IsAutoLevel() const { return LevelsExpressionType == ELevelsExpressionType::AutoLowHigh; }
 };
+
+
+UCLASS()
+class TEXTUREGRAPH_API UTG_Expression_HistogramScan : public UTG_Expression
+{
+	GENERATED_BODY()
+
+	TSharedPtr<FLevels> 				LevelsControl;
+
+public:
+
+	TG_DECLARE_EXPRESSION(TG_Category::Adjustment);
+	virtual void						Evaluate(FTG_EvaluationContext* InContext) override;
+
+	// The input image to adjust the levels for
+	UPROPERTY(EditAnywhere, Category = NoCategory, meta = (TGType = "TG_Input", PinDisplayName = "", MD_HistogramLuminance))
+	FTG_Texture							Input;
+
+	UPROPERTY(EditAnywhere, Category = NoCategory, meta = (TGType = "TG_Setting", UIMin = "0", ClampMin = "0", UIMax = "1", ClampMax = "1", MD_ScalarEditor))
+	float 								Position = 0.5f;
+
+	UPROPERTY(EditAnywhere, Category = NoCategory, meta = (TGType = "TG_Setting", UIMin = "0", ClampMin = "0", UIMax = "1", ClampMax = "1", MD_ScalarEditor))
+	float 								Contrast = 0.5f; 
+
+	// The output image filtered result of the Levels operator 
+	UPROPERTY(EditAnywhere, Category = NoCategory, meta = (TGType = "TG_Output", PinDisplayName = ""))
+	FTG_Texture							Output;
+
+	virtual FText						GetTooltipText() const override { return FText::FromString(TEXT("Remaps shadows and highlights of the input. Any values less or equal to Low are mapped to black, any values, greater or equal to High are mapped to white, and any values inbetween have Gamma applied as an exponent.")); } 
+
+};
\ No newline at end of file
diff --git a/Engine/Plugins/Experimental/TextureGraph/Source/TextureGraphEngine/Device/DeviceBuffer.cpp b/Engine/Plugins/Experimental/TextureGraph/Source/TextureGraphEngine/Device/DeviceBuffer.cpp
--- a/Engine/Plugins/Experimental/TextureGraph/Source/TextureGraphEngine/Device/DeviceBuffer.cpp
+++ b/Engine/Plugins/Experimental/TextureGraph/Source/TextureGraphEngine/Device/DeviceBuffer.cpp
@@ -129,6 +129,28 @@ AsyncRawBufferPtr DeviceBuffer::Raw()
 		});
 }
 
+AsyncRawBufferPtr DeviceBuffer::GetRawOrMaketIt()
+{
+	if (!IsFetchingRaw())
+	{
+		return Raw();
+	}
+	else
+	{
+		RawBufferPtr RawObj;
+		return GetOwnerDevice()->Use()
+			.then([this, RawObj](int32) mutable
+				{
+					RawObj = Raw_Now();
+					return PromiseUtil::OnGameThread();
+				})
+			.then([RawObj]()
+				{
+					return RawObj;
+				});
+	}
+}
+
 CHashPtr DeviceBuffer::CalcHash()
 {
 	/// Don't recalculate if the NewHash has already been calculated
diff --git a/Engine/Plugins/Experimental/TextureGraph/Source/TextureGraphEngine/Device/DeviceBuffer.h b/Engine/Plugins/Experimental/TextureGraph/Source/TextureGraphEngine/Device/DeviceBuffer.h
--- a/Engine/Plugins/Experimental/TextureGraph/Source/TextureGraphEngine/Device/DeviceBuffer.h
+++ b/Engine/Plugins/Experimental/TextureGraph/Source/TextureGraphEngine/Device/DeviceBuffer.h
@@ -97,6 +97,8 @@ public:
 	virtual RawBufferPtr			Raw_Now() = 0;	/// SLOW: Read above
 	virtual AsyncRawBufferPtr		Raw();			/// SLOW: Read above [use Blob::Raw instead. This is to be used from Blob only]
 
+	AsyncRawBufferPtr				GetRawOrMaketIt();
+
 protected:
 	virtual AsyncBufferResultPtr	TransferFrom(DeviceBufferRef& Source) = 0;
 
diff --git a/Engine/Plugins/Experimental/TextureGraph/Source/TextureGraphEngine/Transform/Expressions/T_Color.cpp b/Engine/Plugins/Experimental/TextureGraph/Source/TextureGraphEngine/Transform/Expressions/T_Color.cpp
--- a/Engine/Plugins/Experimental/TextureGraph/Source/TextureGraphEngine/Transform/Expressions/T_Color.cpp
+++ b/Engine/Plugins/Experimental/TextureGraph/Source/TextureGraphEngine/Transform/Expressions/T_Color.cpp
@@ -9,7 +9,6 @@
 #include "Helper/GraphicsUtil.h"
 
 IMPLEMENT_GLOBAL_SHADER(FSH_Grayscale, "/Plugin/TextureGraph/Expressions/Expression_Grayscale.usf", "FSH_Grayscale", SF_Pixel);
-IMPLEMENT_GLOBAL_SHADER(FSH_Levels, "/Plugin/TextureGraph/Expressions/Expression_Levels.usf", "FSH_Levels", SF_Pixel);
 IMPLEMENT_GLOBAL_SHADER(FSH_Threshold, "/Plugin/TextureGraph/Expressions/Expression_Threshold.usf", "FSH_Threshold", SF_Pixel);
 IMPLEMENT_GLOBAL_SHADER(FSH_HSV, "/Plugin/TextureGraph/Expressions/Expression_HSV.usf", "FSH_HSV", SF_Pixel);
 IMPLEMENT_GLOBAL_SHADER(FSH_RGB2HSV, "/Plugin/TextureGraph/Expressions/Expression_HSV.usf", "FSH_RGB2HSV", SF_Pixel);
diff --git a/Engine/Plugins/Experimental/TextureGraph/Source/TextureGraphEngine/Transform/Expressions/T_Color.h b/Engine/Plugins/Experimental/TextureGraph/Source/TextureGraphEngine/Transform/Expressions/T_Color.h
--- a/Engine/Plugins/Experimental/TextureGraph/Source/TextureGraphEngine/Transform/Expressions/T_Color.h
+++ b/Engine/Plugins/Experimental/TextureGraph/Source/TextureGraphEngine/Transform/Expressions/T_Color.h
@@ -28,31 +28,6 @@ public:
 	static void						ModifyCompilationEnvironment(const FGlobalShaderPermutationParameters& params, FShaderCompilerEnvironment& env) {}
 };
 
-//////////////////////////////////////////////////////////////////////////
-/// Simple levels shader
-//////////////////////////////////////////////////////////////////////////
-class TEXTUREGRAPHENGINE_API FSH_Levels : public FSH_Base
-{
-public:
-	SHADER_USE_PARAMETER_STRUCT(FSH_Levels, FSH_Base);
-	DECLARE_GLOBAL_SHADER(FSH_Levels);
-
-	BEGIN_SHADER_PARAMETER_STRUCT(FParameters, )
-		SHADER_PARAMETER_STRUCT(FStandardSamplerStates, SamplerStates)
-		SHADER_PARAMETER_TEXTURE(Texture2D, Input)
-		SHADER_PARAMETER(float, MinValue)
-		SHADER_PARAMETER(float, MaxValue)
-		SHADER_PARAMETER(float, Gamma)
-	END_SHADER_PARAMETER_STRUCT()
-
-	static bool						ShouldCompilePermutation(const FGlobalShaderPermutationParameters& Parameters) 
-	{
-		return IsPCPlatform(Parameters.Platform) && EnumHasAllFlags(Parameters.Flags, EShaderPermutationFlags::HasEditorOnlyData);
-	}
-	static void						ModifyCompilationEnvironment(const FGlobalShaderPermutationParameters& params, FShaderCompilerEnvironment& env) {}
-};
-
-
 //////////////////////////////////////////////////////////////////////////
 /// Simple levels shader
 //////////////////////////////////////////////////////////////////////////
diff --git a/Engine/Plugins/Experimental/TextureGraph/Source/TextureGraphEngine/Transform/Expressions/T_Levels.cpp b/Engine/Plugins/Experimental/TextureGraph/Source/TextureGraphEngine/Transform/Expressions/T_Levels.cpp
new file mode 100644
--- /dev/null
+++ b/Engine/Plugins/Experimental/TextureGraph/Source/TextureGraphEngine/Transform/Expressions/T_Levels.cpp
@@ -0,0 +1,329 @@
+// Copyright Epic Games, Inc. All Rights Reserved.
+
+#include "T_Levels.h"
+#include "Job/JobArgs.h"
+#include "TextureGraphEngine.h"
+#include "Math/Vector.h"
+#include "FxMat/MaterialManager.h"
+#include "Job/JobBatch.h"
+#include "Helper/GraphicsUtil.h"
+#include "Transform/Utility/T_TextureHistogram.h"
+
+IMPLEMENT_GLOBAL_SHADER(FSH_Levels, "/Plugin/TextureGraph/Expressions/Expression_Levels.usf", "FSH_Levels", SF_Pixel);
+
+
+bool FLevels::SetHigh(float InValue)
+{
+	float NewValue = std::min(1.f, std::max(InValue, Low));
+	if (NewValue != High)
+	{
+		float CurveExponent = EvalMidExponent();
+		High = NewValue;
+		return SetMidFromMidExponent(CurveExponent);
+	}
+	return false;
+}
+
+bool FLevels::SetLow(float InValue)
+{
+	float NewValue = std::max(0.f, std::min(InValue, High));
+	if (NewValue != Low)
+	{
+		float CurveExponent = EvalMidExponent();
+		Low = NewValue;
+		return SetMidFromMidExponent(CurveExponent);
+	}
+	return false;
+}
+
+bool FLevels::SetMid(float InValue)
+{
+	float NewValue = std::max(Low, std::min(InValue, High));
+	if (NewValue != Mid)
+	{
+		Mid = NewValue;
+		return true;
+	}
+	return false;
+}
+
+float FLevels::EvalRange(float Val) const
+{
+	return std::max(0.f, std::min((Val - Low)/(High - Low), 1.f));
+}
+float FLevels::EvalRangeInv(float Val) const
+{
+	return Val * (High - Low) + Low;
+}
+
+float FLevels::EvalMidExponent() const
+{
+	// 0.5 = EvalRange(Mid) ^ Exponent
+	// thus
+	float MidRanged = EvalRange(Mid);
+	MidRanged = std::max(0.001f, std::min(9.999f, MidRanged));
+
+	return log(0.5) / log(MidRanged);
+}
+
+bool FLevels::SetMidFromMidExponent(float InExponent)
+{
+	// 0.5 = EvalRange(Mid) ^ Exponent
+	// thus
+	float NewValue = EvalRangeInv(pow(0.5, 1.0/InExponent));
+
+	if (NewValue != Mid)
+	{
+		Mid = NewValue;
+		return true;
+	}
+	return false;
+}
+
+void FLevels::InitFromLowMidHigh(float LowValue, float MidValue, float HighValue)
+{
+	Low = std::max(0.0f, LowValue);
+	Mid = MidValue;
+	High = std::min(1.0f, HighValue);
+
+	IsAutoLevels = false;
+	IsHistogramScan = false;
+
+}
+
+void FLevels::InitFromAutoLevels(float InMidPercentage)
+{
+	IsAutoLevels = true;
+	MidPercentage = std::min(std::max(0.0f, InMidPercentage), 1.0f);
+
+	IsHistogramScan = false;
+}
+
+void FLevels::InitFromPositionContrast(float InPosition, float InContrast)
+{
+	float ContrastRange = 1.0 - InContrast * 0.5f;
+
+	Low = std::max(0.0f, InPosition - ContrastRange);
+	Mid = InPosition;
+	High = std::min(1.0f, InPosition + ContrastRange);
+	
+	IsAutoLevels = false;
+
+	IsHistogramScan = true;
+	Position = InPosition;
+	Contrast = InContrast;
+}
+
+
+
+
+class TEXTUREGRAPHENGINE_API RenderMaterial_FX_Levels : public RenderMaterial_FX
+{
+	using Super = RenderMaterial_FX;
+private:
+	FLevelsPtr Levels;
+
+public:
+	RenderMaterial_FX_Levels(FString InName, FxMaterialPtr InMaterial, const FLevelsPtr& InLevels);
+	virtual ~RenderMaterial_FX_Levels() override;
+
+	virtual AsyncPrepareResult		PrepareResources(const TransformArgs& Args) override;
+	virtual BlobTransformPtr		DuplicateInstance(FString InName) override;
+
+	virtual AsyncTransformResultPtr	Exec(const TransformArgs& Args) override;
+};
+
+
+RenderMaterial_FX_Levels::RenderMaterial_FX_Levels(FString InName, FxMaterialPtr InMaterial, const FLevelsPtr& InLevels) :
+	RenderMaterial_FX(InName, InMaterial),
+	Levels(InLevels)
+{
+
+}
+
+RenderMaterial_FX_Levels::~RenderMaterial_FX_Levels()
+{
+
+}
+
+AsyncPrepareResult		RenderMaterial_FX_Levels::PrepareResources(const TransformArgs& Args)
+{
+	return Super::PrepareResources(Args);
+}
+
+
+std::shared_ptr<BlobTransform> RenderMaterial_FX_Levels::DuplicateInstance(FString InName)
+{
+	if (InName.IsEmpty())
+		InName = Name;
+
+	check(FXMaterial);
+	FxMaterialPtr Clone = FXMaterial->Clone();
+	check(Clone);
+
+	return std::static_pointer_cast<BlobTransform>(std::make_shared<RenderMaterial_FX_Levels>(InName, Clone, Levels));
+}
+
+AsyncTransformResultPtr	RenderMaterial_FX_Levels::Exec(const TransformArgs& Args)
+{
+	if (Levels->IsAutoLevels)
+	{
+		float Low = Levels->HistogramData.HistogramData[256 + 2].X;
+		float High = Levels->HistogramData.HistogramData[256 + 3].X;
+		float Mid = Levels->MidPercentage * (High - Low) + Low;
+
+		// feedback the min max mid values in the Levels struct
+		Levels->Low = Low;
+		Levels->High = High;
+		Levels->Mid = Mid;
+
+		//Pass on uniform values to the shader
+		SetFloat(TEXT("LowValue"), Low);
+		SetFloat(TEXT("MidValue"), Mid);
+		SetFloat(TEXT("HighValue"), High);
+
+	}
+	return Super::Exec(Args);
+}
+
+
+//////////////////////////////////////////////////////////////////////////
+class TEXTUREGRAPHENGINE_API Job_Levels : public Job
+{
+private:
+	TiledBlobRef Histogram;
+	FLevelsPtr Levels;
+
+public:
+	Job_Levels(UMixInterface* InMix, int32 targetId, BlobTransformPtr InTransform, TiledBlobRef InHistogram, const FLevelsPtr& InLevels, UObject* InErrorOwner = nullptr, uint16 priority = (uint16)E_Priority::kHigh, uint64 id = 0);
+
+protected:
+
+	virtual cti::continuable<int32>	PreExecAsync(ENamedThreads::Type execThread, ENamedThreads::Type returnThread) override;
+};
+
+Job_Levels::Job_Levels(UMixInterface* InMix, int32 targetId, BlobTransformPtr InTransform, TiledBlobRef InSourceHistogram, const FLevelsPtr& InLevels,
+											   UObject* InErrorOwner /*= nullptr*/, uint16 priority /*= (uint16)E_Priority::kHigh*/, uint64 id /*= 0*/)
+: Job(InMix, targetId, InTransform, InErrorOwner, priority)
+, Histogram(InSourceHistogram)
+, Levels(InLevels)
+{
+	DeviceNativeTask::Name = TEXT("Levels");
+}
+
+cti::continuable<int32> Job_Levels::PreExecAsync(ENamedThreads::Type execThread, ENamedThreads::Type returnThread)
+{
+	//TODO: Need to cache the Result so that the job do not execute every time with matching arguments
+/*	BlobPtr cachedResult = TextureGraphEngine::GetBlobber()->Find(Job::Hash()->Value());
+	if (cachedResult)
+	{
+		bIsCulled = true;
+		bIsDone = true;
+
+		return cti::make_ready_continuable(0);
+	}*/
+
+	// Early exit if download not needed
+	if (!Levels->IsAutoLevels)
+	{
+		return cti::make_ready_continuable(0);
+	}
+
+	// Now the job is really needed, let's go on with the Histogram download:
+	BlobRef HistoBuffer = Histogram->GetTile(0,0);
+	if (HistoBuffer)
+	{
+		/*
+		* // TODO ITRied to use this path, but what happen when the download of the raw is in flight ?
+		if (!HistoBuffer->GetBufferRef()->HasRaw())
+		{
+			/// If it's already fetching raw, then we don't want to do it again
+			if (HistoBuffer->GetBufferRef()->IsFetchingRaw())
+				return cti::make_ready_continuable(0);
+				*/
+			return HistoBuffer->GetBufferRef()->GetRawOrMaketIt().then([this](RawBufferPtr Raw)
+			{
+				// TODO: This should NOT happen ?
+				if (!Raw)
+				{
+					check(Raw)
+					return cti::make_ready_continuable(0);
+				}
+				check(Raw->GetDescriptor().Width == 256)
+				check(Raw->GetDescriptor().Height == 2)
+				check(Raw->GetDescriptor().Size() == Raw->GetLength())
+				check(Raw->HasData())
+	
+				int32 NumBins = 256;
+				int32 NumMetaBins = 5;
+				const uint8* RawData = Raw->GetData();	
+				if (RawData)
+				{
+					Levels->HistogramData.HistogramData.Append((FVector4f*)RawData, (NumBins + NumMetaBins));
+				}	
+				return cti::make_ready_continuable(0);
+			});
+
+		//}
+	}
+
+	return cti::make_ready_continuable(0);
+}
+
+
+TiledBlobPtr	T_Levels::Create(MixUpdateCyclePtr Cycle, BufferDescriptor DesiredOutputDesc, TiledBlobPtr Source, const FLevelsPtr& InLevels, int32 TargetId)
+{
+	bool NeedsConvertToGrayscale = (Source->GetDescriptor().ItemsPerPoint > 1) ;
+	
+	FSH_Levels::FPermutationDomain PermutationVector;
+	PermutationVector.Set<FSH_Levels::FConvertToGrayscale>(NeedsConvertToGrayscale);
+	//PermutationVector.Set<FSH_Levels::FIsAutoLevels>(InLevels->IsAutoLevels);
+	PermutationVector.Set<FSH_Levels::FIsAutoLevels>(false);
+
+
+	if (!Source)
+	{
+		return TextureHelper::GBlack;
+	}
+
+	bool NeedHistogramRaw = false;
+	TiledBlobPtr Histogram = TextureHelper::GBlack;
+	if (InLevels->IsAutoLevels)
+	{
+		NeedHistogramRaw = true;
+		Histogram = T_TextureHistogram::Create(Cycle, Source, TargetId);
+	}
+
+	if (!Histogram)
+	{
+		Histogram = TextureHelper::GBlack;
+		NeedHistogramRaw = false;
+	}
+	
+	std::shared_ptr<FxMaterial_Normal<VSH_Simple, FSH_Levels>> Mat = std::make_shared<FxMaterial_Normal<VSH_Simple, FSH_Levels>>();
+	const RenderMaterial_FXPtr RenderMaterial = std::make_shared<RenderMaterial_FX_Levels>(TEXT("T_Levels"), Mat, InLevels);
+	check(RenderMaterial);
+
+	JobUPtr RenderJob = std::make_unique<Job_Levels>(Cycle->GetMix(), TargetId, std::static_pointer_cast<BlobTransform>(RenderMaterial), Histogram, InLevels);
+	RenderJob
+		->AddArg(ARG_BLOB(Source, "SourceTexture"))
+		->AddArg(ARG_FLOAT(InLevels->Low, "LowValue"))
+		->AddArg(ARG_FLOAT(InLevels->High, "HighValue"))
+		->AddArg(ARG_FLOAT(InLevels->Mid, "MidValue"))
+		->AddArg(ARG_FLOAT(float(InLevels->IsAutoLevels), "DoAutoLevel"))
+		->AddArg(ARG_FLOAT(InLevels->MidPercentage, "MidPercentage"))
+		->AddArg(std::make_shared<JobArg_Blob>(JobArg_Blob(Histogram, "Histogram").WithNotHandleTiles()))
+		;
+
+
+	if (DesiredOutputDesc.Format == BufferFormat::Auto)
+	{
+		DesiredOutputDesc.Format = Source->GetDescriptor().Format;
+	}
+
+	TiledBlobPtr Result = RenderJob->InitResult(TEXT("Levels"), &DesiredOutputDesc);
+	Cycle->AddJob(TargetId, std::move(RenderJob));
+
+	return Result;
+
+}
diff --git a/Engine/Plugins/Experimental/TextureGraph/Source/TextureGraphEngine/Transform/Expressions/T_Levels.h b/Engine/Plugins/Experimental/TextureGraph/Source/TextureGraphEngine/Transform/Expressions/T_Levels.h
new file mode 100644
--- /dev/null
+++ b/Engine/Plugins/Experimental/TextureGraph/Source/TextureGraphEngine/Transform/Expressions/T_Levels.h
@@ -0,0 +1,96 @@
+// Copyright Epic Games, Inc. All Rights Reserved.
+#pragma once
+
+#include "CoreMinimal.h"
+#include "Job/Job.h"
+#include "FxMat/FxMaterial.h"
+#include "Model/Mix/MixUpdateCycle.h"
+#include <DataDrivenShaderPlatformInfo.h>
+
+
+//////////////////////////////////////////////////////////////////////////
+/// Simple levels shader
+//////////////////////////////////////////////////////////////////////////
+class TEXTUREGRAPHENGINE_API FSH_Levels : public FSH_Base
+{
+public:
+	SHADER_USE_PARAMETER_STRUCT(FSH_Levels, FSH_Base);
+	DECLARE_GLOBAL_SHADER(FSH_Levels);
+
+	BEGIN_SHADER_PARAMETER_STRUCT(FParameters, )
+		SHADER_PARAMETER_STRUCT(FStandardSamplerStates, SamplerStates)
+		SHADER_PARAMETER_TEXTURE(Texture2D, SourceTexture)
+		SHADER_PARAMETER(float, LowValue)
+		SHADER_PARAMETER(float, HighValue)
+		SHADER_PARAMETER(float, MidValue)
+		SHADER_PARAMETER(float, DoAutoLevel)
+		SHADER_PARAMETER(float, MidPercentage)
+		SHADER_PARAMETER_TEXTURE(Texture2D, Histogram)
+	END_SHADER_PARAMETER_STRUCT()
+
+	class FConvertToGrayscale : SHADER_PERMUTATION_BOOL("CONVERT_TO_GRAYSCALE");
+	class FIsAutoLevels : SHADER_PERMUTATION_BOOL("AUTO_LEVELS");
+	using FPermutationDomain = TShaderPermutationDomain<FConvertToGrayscale, FIsAutoLevels>;
+
+	static bool						ShouldCompilePermutation(const FGlobalShaderPermutationParameters& Parameters) 
+	{
+		return IsPCPlatform(Parameters.Platform) && EnumHasAllFlags(Parameters.Flags, EShaderPermutationFlags::HasEditorOnlyData);
+	}
+	static void						ModifyCompilationEnvironment(const FGlobalShaderPermutationParameters& params, FShaderCompilerEnvironment& env) {}
+};
+
+
+struct TEXTUREGRAPHENGINE_API FLevelsHistogramData
+{
+	TArray<FVector4f> HistogramData;
+};
+
+
+struct TEXTUREGRAPHENGINE_API FLevels
+{
+
+	// The Low value of the Levels adjustment, any pixel under that value is set to black. Default is 0.
+	float								Low = 0;
+
+	// The mid value of the Levels adjustment, must be in the range [Min, Max] and the Default is 0.5.
+	// The mid value determine where the smoothing curve applying the midpoint filter is crossing 0.5.
+	float								Mid = 0.5f; // Midtones
+
+	// The High value of the Levels adjustment, any pixel above that value is set to white. Default is 1.
+	float								High = 1;
+
+	bool 								IsAutoLevels = false;
+	float 								MidPercentage = 0.5f;
+
+	bool 								IsHistogramScan = false;
+	float 								Position = 0.5f;
+	float 								Contrast = 0.5f;
+
+	bool SetLow(float InValue);
+	bool SetMid(float InValue);
+	bool SetHigh(float InValue);
+
+	// Eval Low-High range mapping on value
+	float EvalRange(float Val) const;
+	// Eval reverse Low-High range mapping on value
+	float EvalRangeInv(float Val) const;
+
+	// Evaluate the MidExponent of the power curve applying the midpoint filter
+	float EvalMidExponent() const;
+	bool SetMidFromMidExponent(float InExponent);
+
+	void InitFromLowMidHigh(float LowValue, float MidValue, float HighValue);
+	void InitFromAutoLevels(float MidPercentage);
+	void InitFromPositionContrast(float MidValue, float Contrast);
+
+	FLevelsHistogramData HistogramData;
+};
+
+using FLevelsPtr = TSharedPtr< FLevels>;
+
+class TEXTUREGRAPHENGINE_API T_Levels
+{
+public:
+	static TiledBlobPtr				Create(MixUpdateCyclePtr Cycle, BufferDescriptor DesiredOutputDesc, TiledBlobPtr Source, const FLevelsPtr& Levels, int32 TargetId);
+};
+
