Update the Histogram Scan algorithm with correct Levels configuration

Fix the the dangling  RawBufferPtr in the GetRawOrMakeIt call issue https://horde.devtools.epicgames.com/job/6650c0eff97e86098df1f3fb?step=b6f5&issue=478565

diff --git a/Engine/Plugins/Experimental/TextureGraph/Source/TextureGraphEngine/Device/DeviceBuffer.cpp b/Engine/Plugins/Experimental/TextureGraph/Source/TextureGraphEngine/Device/DeviceBuffer.cpp
--- a/Engine/Plugins/Experimental/TextureGraph/Source/TextureGraphEngine/Device/DeviceBuffer.cpp
+++ b/Engine/Plugins/Experimental/TextureGraph/Source/TextureGraphEngine/Device/DeviceBuffer.cpp
@@ -131,23 +131,20 @@ AsyncRawBufferPtr DeviceBuffer::Raw()
 
 AsyncRawBufferPtr DeviceBuffer::GetRawOrMaketIt()
 {
+	check(IsInGameThread());
 	if (!IsFetchingRaw())
 	{
 		return Raw();
 	}
 	else
 	{
-		RawBufferPtr RawObj;
-		return GetOwnerDevice()->Use()
-			.then([this, RawObj](int32) mutable
-				{
-					RawObj = Raw_Now();
-					return PromiseUtil::OnGameThread();
-				})
-			.then([RawObj]()
-				{
-					return RawObj;
-				});
+		return GetOwnerDevice()->Use().then([this](int32) mutable
+			{
+				RawBufferPtr RawObj = Raw_Now();
+
+				// Return to game thread with the payload
+				return PromiseUtil::OnGameThread(std::move(RawObj));
+			});
 	}
 }
 
diff --git a/Engine/Plugins/Experimental/TextureGraph/Source/TextureGraphEngine/Helper/Promise.h b/Engine/Plugins/Experimental/TextureGraph/Source/TextureGraphEngine/Helper/Promise.h
--- a/Engine/Plugins/Experimental/TextureGraph/Source/TextureGraphEngine/Helper/Promise.h
+++ b/Engine/Plugins/Experimental/TextureGraph/Source/TextureGraphEngine/Helper/Promise.h
@@ -91,6 +91,16 @@ struct TEXTUREGRAPHENGINE_API PromiseUtil
 		});
 	}
 
+	template<typename T>
+	static cti::continuable<T>  OnGameThread(T&& Arg)
+	{
+		if (IsInGameThread())
+			return cti::make_ready_continuable(std::forward<T>(Arg));
+
+		return OnThread(ENamedThreads::GameThread, std::forward<T>(Arg));
+	}
+
+
 	static AsyncInt					OnRenderingThread()
 	{
 		if (IsInRenderingThread())
diff --git a/Engine/Plugins/Experimental/TextureGraph/Source/TextureGraphEngine/Transform/Expressions/T_Levels.cpp b/Engine/Plugins/Experimental/TextureGraph/Source/TextureGraphEngine/Transform/Expressions/T_Levels.cpp
--- a/Engine/Plugins/Experimental/TextureGraph/Source/TextureGraphEngine/Transform/Expressions/T_Levels.cpp
+++ b/Engine/Plugins/Experimental/TextureGraph/Source/TextureGraphEngine/Transform/Expressions/T_Levels.cpp
@@ -87,31 +87,26 @@ void FLevels::InitFromLowMidHigh(float LowValue, float MidValue, float HighValue
 	High = std::min(1.0f, HighValue);
 
 	IsAutoLevels = false;
-	IsHistogramScan = false;
-
 }
 
 void FLevels::InitFromAutoLevels(float InMidPercentage)
 {
 	IsAutoLevels = true;
 	MidPercentage = std::min(std::max(0.0f, InMidPercentage), 1.0f);
-
-	IsHistogramScan = false;
 }
 
+// Histogram scan
 void FLevels::InitFromPositionContrast(float InPosition, float InContrast)
 {
-	float ContrastRange = 1.0 - InContrast * 0.5f;
-
-	Low = std::max(0.0f, InPosition - ContrastRange);
-	Mid = InPosition;
-	High = std::min(1.0f, InPosition + ContrastRange);
+	float C = InContrast * 0.5f;
+	float P = 1.0f - FMath::Clamp(InPosition, 0, 1);
+	float P1 = (FMath::Max(P, 0.5f) - 0.5f) * 2.0f;
+	float P2 = FMath::Min(P * 2.0f, 1.0f);
+	Low = FMath::Lerp(P1, P2, C);
+	High = FMath::Lerp(P2, P1, C);
+	Mid = Low + (High - Low) * 0.5f;
 	
 	IsAutoLevels = false;
-
-	IsHistogramScan = true;
-	Position = InPosition;
-	Contrast = InContrast;
 }
 
 
diff --git a/Engine/Plugins/Experimental/TextureGraph/Source/TextureGraphEngine/Transform/Expressions/T_Levels.h b/Engine/Plugins/Experimental/TextureGraph/Source/TextureGraphEngine/Transform/Expressions/T_Levels.h
--- a/Engine/Plugins/Experimental/TextureGraph/Source/TextureGraphEngine/Transform/Expressions/T_Levels.h
+++ b/Engine/Plugins/Experimental/TextureGraph/Source/TextureGraphEngine/Transform/Expressions/T_Levels.h
@@ -62,10 +62,6 @@ struct TEXTUREGRAPHENGINE_API FLevels
 	bool 								IsAutoLevels = false;
 	float 								MidPercentage = 0.5f;
 
-	bool 								IsHistogramScan = false;
-	float 								Position = 0.5f;
-	float 								Contrast = 0.5f;
-
 	bool SetLow(float InValue);
 	bool SetMid(float InValue);
 	bool SetHigh(float InValue);
