Histogram scan node

diff --git a/Engine/Plugins/Experimental/TextureGraph/Source/TextureGraph/Private/Expressions/FIlter/TG_Expression_HistogramScan.cpp b/Engine/Plugins/Experimental/TextureGraph/Source/TextureGraph/Private/Expressions/FIlter/TG_Expression_HistogramScan.cpp
new file mode 100644
--- /dev/null
+++ b/Engine/Plugins/Experimental/TextureGraph/Source/TextureGraph/Private/Expressions/FIlter/TG_Expression_HistogramScan.cpp
@@ -0,0 +1,52 @@
+// Copyright Epic Games, Inc. All Rights Reserved.
+
+#include "Expressions/Filter/TG_Expression_HistogramScan.h"
+#include "Expressions/Filter/TG_Expression_Levels.h"
+#include "FxMat/MaterialManager.h"
+#include "Job/JobBatch.h"
+#include "Transform/Expressions/T_Color.h"
+
+void UTG_Expression_HistogramScan::Evaluate(FTG_EvaluationContext* InContext)
+{
+	Super::Evaluate(InContext);
+
+	if (!Input)
+	{
+		Output = FTG_Texture::GetBlack();
+		return;
+	}
+
+	BufferDescriptor SourceDesc = Input.GetBufferDescriptor();
+	BufferDescriptor DesiredDesc = BufferDescriptor::Combine(Output.GetBufferDescriptor(), SourceDesc);
+	DesiredDesc.ItemsPerPoint = 4;
+
+	const RenderMaterial_FXPtr RenderMaterial = TextureGraphEngine::GetMaterialManager()->CreateMaterial_FX<VSH_Simple, FSH_Levels>(TEXT("T_Levels"));
+
+	check(RenderMaterial);
+
+	FTG_LevelsSettings Levels;
+
+	float C = Contrast * 0.5f;
+	float P = 1.0f - FMath::Clamp(Position, 0, 1);
+	float P1 = (FMath::Max(P, 0.5f) - 0.5f) * 2.0f;
+	float P2 = FMath::Min(P * 2.0f, 1.0f);
+
+	Levels.Low = FMath::Lerp(P1, P2, C);
+	Levels.High = FMath::Lerp(P2, P1, C);
+	Levels.Mid = Levels.Low + (Levels.High - Levels.Low) * 0.5f;
+
+	JobUPtr RenderJob = std::make_unique<Job>(InContext->Cycle->GetMix(), InContext->TargetId, std::static_pointer_cast<BlobTransform>(RenderMaterial));
+
+	RenderJob
+		->AddArg(ARG_BLOB(Input, "Input"))
+		->AddArg(ARG_FLOAT(Levels.Low, "MinValue"))
+		->AddArg(ARG_FLOAT(Levels.High, "MaxValue"))
+		->AddArg(ARG_FLOAT(Levels.EvalMidExponent(), "Gamma"))
+		;
+
+	const FString Name = TEXT("HistogramScan"); 
+	BufferDescriptor Desc = Output.GetBufferDescriptor();
+
+	Output = RenderJob->InitResult(Name, &Desc);
+	InContext->Cycle->AddJob(InContext->TargetId, std::move(RenderJob));
+}
diff --git a/Engine/Plugins/Experimental/TextureGraph/Source/TextureGraph/Public/Expressions/Filter/TG_Expression_HistogramScan.h b/Engine/Plugins/Experimental/TextureGraph/Source/TextureGraph/Public/Expressions/Filter/TG_Expression_HistogramScan.h
new file mode 100644
--- /dev/null
+++ b/Engine/Plugins/Experimental/TextureGraph/Source/TextureGraph/Public/Expressions/Filter/TG_Expression_HistogramScan.h
@@ -0,0 +1,43 @@
+// Copyright Epic Games, Inc. All Rights Reserved.
+
+#pragma once
+
+#include "CoreMinimal.h"
+
+#include "Expressions/TG_Expression.h"
+#include "TG_Texture.h"
+#include "FxMat/FxMaterial.h"
+
+#include "TG_Expression_HistogramScan.generated.h"
+
+//////////////////////////////////////////////////////////////////////////
+/// Expression
+//////////////////////////////////////////////////////////////////////////
+UCLASS()
+class TEXTUREGRAPH_API UTG_Expression_HistogramScan : public UTG_Expression
+{
+	GENERATED_BODY()
+
+public:
+	TG_DECLARE_EXPRESSION(TG_Category::Adjustment);
+	virtual void							Evaluate(FTG_EvaluationContext* InContext) override;
+
+	// Drives the position of the histogram of the input image
+	UPROPERTY(EditAnywhere, Category = NoCategory, meta = (TGType = "TG_Setting", UIMin = "0", ClampMin = "0", UIMax = "1", ClampMax = "1"))
+		float								Position = 0;
+
+	// Drives the contrast of the histogram of the input image
+	UPROPERTY(EditAnywhere, Category = NoCategory, meta = (TGType = "TG_Setting", UIMin = "0", ClampMin = "0", UIMax = "1", ClampMax = "1"))
+		float								Contrast = 0;
+
+	// The input texture to adjust the brightness/contrast for
+	UPROPERTY(meta = (TGType = "TG_Input", PinDisplayName = ""))
+		FTG_Texture							Input;
+
+	// The output image
+	UPROPERTY(EditAnywhere, Category = NoCategory, meta = (TGType = "TG_Output", PinDisplayName = ""))
+		FTG_Texture							Output;
+	virtual FTG_Name						GetDefaultName() const override { return TEXT("HistogramScan");}
+	virtual FText							GetTooltipText() const override { return FText::FromString(TEXT("Lets you drive the contrast and position of the histogram. Input must be a grayscale image.")); } 
+};
+
