Propagate correctly the assignment of a value to a UProperty and call the Uproperty Setter if any.
Clean up the inner checks of values assignement for the Levels
Updated the test for Levels

diff --git a/Engine/Plugins/Experimental/TextureGraph/Source/TextureGraph/Private/Expressions/FIlter/TG_Expression_Levels.cpp b/Engine/Plugins/Experimental/TextureGraph/Source/TextureGraph/Private/Expressions/FIlter/TG_Expression_Levels.cpp
--- a/Engine/Plugins/Experimental/TextureGraph/Source/TextureGraph/Private/Expressions/FIlter/TG_Expression_Levels.cpp
+++ b/Engine/Plugins/Experimental/TextureGraph/Source/TextureGraph/Private/Expressions/FIlter/TG_Expression_Levels.cpp
@@ -90,23 +90,33 @@ bool FTG_LevelsSettings::SetMidFromMidExponent(float InExponent)
 }
 
 
-
+void UTG_Expression_Levels::PostLoad()
+{
+	// Restore LevelsSettings inner struct from saved values
+	Levels.Low = (LowValue);
+	Levels.High = (HighValue);
+	Levels.Mid = (MidValue);
+}
 
 #if WITH_EDITOR
 
 void UTG_Expression_Levels::PostEditChangeProperty(FPropertyChangedEvent& PropertyChangedEvent)
 {
-	// First catch if Material changes
+	// Catch if any of low / mid / high changes do the proper range check and feedback final values
+	// Low
 	if (PropertyChangedEvent.GetPropertyName() == GET_MEMBER_NAME_CHECKED(UTG_Expression_Levels, LowValue))
 	{
 		SetLowValue(LowValue);
-		FeedbackPinValue(GET_MEMBER_NAME_CHECKED(UTG_Expression_Levels, MidValue), MidValue);
 	}
-	// Second catch if AttributeName changes
+	// High
 	if (PropertyChangedEvent.GetPropertyName() == GET_MEMBER_NAME_CHECKED(UTG_Expression_Levels, HighValue))
 	{
 		SetHighValue(HighValue);
-		FeedbackPinValue(GET_MEMBER_NAME_CHECKED(UTG_Expression_Levels, MidValue), MidValue);
+	}
+	// Mid
+	if (PropertyChangedEvent.GetPropertyName() == GET_MEMBER_NAME_CHECKED(UTG_Expression_Levels, MidValue))
+	{
+		SetMidValue(MidValue);
 	}
 
 	Super::PostEditChangeProperty(PropertyChangedEvent);
@@ -157,7 +167,6 @@ void UTG_Expression_Levels::SetLowValue(float InValue)
 	else
 	{
 		LowValue = Levels.Low;
-		FeedbackPinValue(GET_MEMBER_NAME_CHECKED(UTG_Expression_Levels, LowValue), LowValue);
 	}
 }
 
@@ -170,7 +179,6 @@ void UTG_Expression_Levels::SetMidValue(float InValue)
 	else
 	{
 		MidValue = Levels.Mid;
-		FeedbackPinValue(GET_MEMBER_NAME_CHECKED(UTG_Expression_Levels, MidValue), MidValue);
 	}
 }
 
@@ -185,7 +193,6 @@ void UTG_Expression_Levels::SetHighValue(float InValue)
 	else
 	{
 		HighValue = Levels.High;
-		FeedbackPinValue(GET_MEMBER_NAME_CHECKED(UTG_Expression_Levels, HighValue), HighValue);
 	}
 }
 
diff --git a/Engine/Plugins/Experimental/TextureGraph/Source/TextureGraph/Private/TG_Var.cpp b/Engine/Plugins/Experimental/TextureGraph/Source/TextureGraph/Private/TG_Var.cpp
--- a/Engine/Plugins/Experimental/TextureGraph/Source/TextureGraph/Private/TG_Var.cpp
+++ b/Engine/Plugins/Experimental/TextureGraph/Source/TextureGraph/Private/TG_Var.cpp
@@ -150,8 +150,8 @@ void Generic_Simple_Serializer(FTG_Var::VarPropertySerialInfo& Info)
 
 	if (Info.CopyVarToProperty)
 	{
-		TProperty->SetPropertyValue(TProperty->template ContainerPtrToValuePtr<T_ValueType>(Info.Owner, Info.ClampedIndex()), VarValue);
-		//TProperty->SetValue_InContainer(Info.Owner, VarValue);
+		// This calls the Setter method if the UProperty has a Setter
+		TProperty->SetValue_InContainer(Info.Owner, VarValue);
 	}
 	else
 		Info.Var->EditAs<T_ValueType>() = TProperty->GetPropertyValue(TProperty->template ContainerPtrToValuePtr<T_ValueType>(Info.Owner, Info.ClampedIndex()));
diff --git a/Engine/Plugins/Experimental/TextureGraph/Source/TextureGraph/Public/Expressions/Filter/TG_Expression_Levels.h b/Engine/Plugins/Experimental/TextureGraph/Source/TextureGraph/Public/Expressions/Filter/TG_Expression_Levels.h
--- a/Engine/Plugins/Experimental/TextureGraph/Source/TextureGraph/Public/Expressions/Filter/TG_Expression_Levels.h
+++ b/Engine/Plugins/Experimental/TextureGraph/Source/TextureGraph/Public/Expressions/Filter/TG_Expression_Levels.h
@@ -79,7 +79,8 @@ class TEXTUREGRAPH_API UTG_Expression_Levels : public UTG_Expression
 
 	TSharedPtr<FLevels> 				LevelsControl;
 
-public:
+public:	
+	virtual void						PostLoad() override;
 #if WITH_EDITOR
 	virtual void 						PostEditChangeProperty(FPropertyChangedEvent& PropertyChangedEvent) override;
 
@@ -87,6 +88,8 @@ public:
 	virtual bool						CanEditChange(const FProperty* InProperty) const override;
 #endif
 
+
+
 	TG_DECLARE_EXPRESSION(TG_Category::Adjustment);
 	virtual void						Evaluate(FTG_EvaluationContext* InContext) override;
 
