Fixed details panel inconsistenty issues with multiple nodes of the same type which have dynamic pins (pins adjusted on changing type etc.)

diff --git a/Engine/Plugins/Experimental/TextureGraph/Source/TextureGraphEditor/Private/TG_Editor.cpp b/Engine/Plugins/Experimental/TextureGraph/Source/TextureGraphEditor/Private/TG_Editor.cpp
--- a/Engine/Plugins/Experimental/TextureGraph/Source/TextureGraphEditor/Private/TG_Editor.cpp
+++ b/Engine/Plugins/Experimental/TextureGraph/Source/TextureGraphEditor/Private/TG_Editor.cpp
@@ -1200,6 +1200,10 @@ void FTG_Editor::OnSelectedNodesChanged(const TSet<class UObject*>& NewSelection
 		if (UTG_EdGraphNode* GraphNode = Cast<UTG_EdGraphNode>(*SetIt))
 		{
 			SelectedObjects.Add(GraphNode->GetDetailsObject());
+
+			// calling this here to internally set EditCondition metadata properly for Details panel
+			// this is driven from the CanEditChange() method of the expression which can be customized as required
+			GraphNode->UpdateInputPinsVisibility();
 		}
 	}
 
