FIx warning

diff --git a/Engine/Plugins/Experimental/TextureGraph/Source/TextureGraph/Private/Expressions/FIlter/TG_Expression_Levels.cpp b/Engine/Plugins/Experimental/TextureGraph/Source/TextureGraph/Private/Expressions/FIlter/TG_Expression_Levels.cpp
--- a/Engine/Plugins/Experimental/TextureGraph/Source/TextureGraph/Private/Expressions/FIlter/TG_Expression_Levels.cpp
+++ b/Engine/Plugins/Experimental/TextureGraph/Source/TextureGraph/Private/Expressions/FIlter/TG_Expression_Levels.cpp
@@ -172,14 +172,8 @@ void UTG_Expression_Levels::SetLowValue(float InValue)
 
 void UTG_Expression_Levels::SetMidValue(float InValue)
 {
-	if (Levels.SetMid(InValue))
-	{
-		MidValue = Levels.Mid;
-	}
-	else
-	{
-		MidValue = Levels.Mid;
-	}
+	Levels.SetMid(InValue);
+	MidValue = Levels.Mid;
 }
 
 void UTG_Expression_Levels::SetHighValue(float InValue)
