Replaced FindPinChecked with FindPin and an if check for cases where EdPin isn't created for certain input pins.

diff --git a/Engine/Plugins/Experimental/TextureGraph/Source/TextureGraphEditor/Private/EdGraph/TG_EdGraphNode.cpp b/Engine/Plugins/Experimental/TextureGraph/Source/TextureGraphEditor/Private/EdGraph/TG_EdGraphNode.cpp
--- a/Engine/Plugins/Experimental/TextureGraph/Source/TextureGraphEditor/Private/EdGraph/TG_EdGraphNode.cpp
+++ b/Engine/Plugins/Experimental/TextureGraph/Source/TextureGraphEditor/Private/EdGraph/TG_EdGraphNode.cpp
@@ -503,10 +503,13 @@ void UTG_EdGraphNode::UpdateInputPinsVisibility() const
 		check(TGPin);
 
 		const FName& OtherPinName = TGPin->GetArgumentName();
-		UEdGraphPin* EdPin = FindPinChecked(OtherPinName, EEdGraphPinDirection::EGPD_Input);
+		UEdGraphPin* EdPin = FindPin(OtherPinName, EEdGraphPinDirection::EGPD_Input);
 
 		// check hide state
-		UpdatePinVisibility(EdPin, TGPin);
+		if (EdPin)
+		{
+			UpdatePinVisibility(EdPin, TGPin);
+		}
 	}
 }
 #endif
