Added a step function

diff --git a/Engine/Plugins/Experimental/TextureGraph/Shaders/Expressions/Expression_Maths_TwoInputs.usf b/Engine/Plugins/Experimental/TextureGraph/Shaders/Expressions/Expression_Maths_TwoInputs.usf
--- a/Engine/Plugins/Experimental/TextureGraph/Shaders/Expressions/Expression_Maths_TwoInputs.usf
+++ b/Engine/Plugins/Experimental/TextureGraph/Shaders/Expressions/Expression_Maths_TwoInputs.usf
@@ -56,3 +56,10 @@ float4 FSH_Pow(float2 uv : TEXCOORD0) : SV_Target0
 	return pow(op1, op2);
 }
 
+float4 FSH_Step(float2 uv : TEXCOORD0) : SV_Target0
+{
+	float4 op1 = Operand1.Sample(SamplerStates_Wrap, uv);
+	float4 op2 = Operand2.Sample(SamplerStates_Wrap, uv);
+	return step(op1, op2);
+}
+
diff --git a/Engine/Plugins/Experimental/TextureGraph/Shaders/Expressions/Expression_Warp.usf b/Engine/Plugins/Experimental/TextureGraph/Shaders/Expressions/Expression_Warp.usf
--- a/Engine/Plugins/Experimental/TextureGraph/Shaders/Expressions/Expression_Warp.usf
+++ b/Engine/Plugins/Experimental/TextureGraph/Shaders/Expressions/Expression_Warp.usf
@@ -35,7 +35,7 @@ float4 FSH_DirectionalWarp(float2 InUV_ : TEXCOORD0) : SV_Target0
 	/// We need to calculate the UV that will get warped to the current UV. So we need
 	/// to move in the opposite direction of the angle to find the UV that will end
 	/// replacing the current pixel. Hence we add 180 degrees to it (or half PI in radians)
-	float AngleValue = AngleRad * HALF_PI;
+	float AngleValue = AngleRad; // * HALF_PI;
 	float C = cos(AngleValue);
 	float S = sin(AngleValue);
 	float2 Dir = float2(C, S);
diff --git a/Engine/Plugins/Experimental/TextureGraph/Source/TextureGraph/Private/Expressions/Maths/TG_Expression_Maths_TwoInputs.cpp b/Engine/Plugins/Experimental/TextureGraph/Source/TextureGraph/Private/Expressions/Maths/TG_Expression_Maths_TwoInputs.cpp
--- a/Engine/Plugins/Experimental/TextureGraph/Source/TextureGraph/Private/Expressions/Maths/TG_Expression_Maths_TwoInputs.cpp
+++ b/Engine/Plugins/Experimental/TextureGraph/Source/TextureGraph/Private/Expressions/Maths/TG_Expression_Maths_TwoInputs.cpp
@@ -62,7 +62,9 @@ FTG_Texture	UTG_Expression_Subtract::EvaluateTexture(FTG_EvaluationContext* InCo
 	return GenericMathOp(InContext, T_Maths_TwoInputs::CreateSubtract, Output.EditTexture().GetBufferDescriptor(), Input1.GetTexture(InContext), Input2.GetTexture(InContext));
 }
 
+//////////////////////////////////////////////////////////////////////////
 /// Dot
+//////////////////////////////////////////////////////////////////////////
 float UTG_Expression_Dot::EvaluateScalar_WithValue(FTG_EvaluationContext* InContext, const float* const ValuePtr, size_t Count)
 {
 	return 0;
@@ -80,7 +82,9 @@ FTG_Texture	UTG_Expression_Dot::EvaluateTexture(FTG_EvaluationContext* InContext
 	return T_Maths_TwoInputs::CreateDot(InContext->Cycle, Output.EditTexture().GetBufferDescriptor(), InContext->TargetId, Input1.GetTexture(InContext), Input2.GetTexture(InContext));
 }
 
-/// Cross
+//////////////////////////////////////////////////////////////////////////
+/// Cross 
+//////////////////////////////////////////////////////////////////////////
 float UTG_Expression_Cross::EvaluateScalar_WithValue(FTG_EvaluationContext* InContext, const float* const ValuePtr, size_t Count)
 {
 	return 0;
@@ -97,6 +101,22 @@ FTG_Texture	UTG_Expression_Cross::EvaluateTexture(FTG_EvaluationContext* InConte
 	return T_Maths_TwoInputs::CreateCross(InContext->Cycle, Output.EditTexture().GetBufferDescriptor(), InContext->TargetId, Input1.GetTexture(InContext), Input2.GetTexture(InContext));
 }
 
+//////////////////////////////////////////////////////////////////////////
+/// Step
+//////////////////////////////////////////////////////////////////////////
+float UTG_Expression_Step::EvaluateScalar_WithValue(FTG_EvaluationContext* InContext, const float* const ValuePtr, size_t Count)
+{
+	check(ValuePtr && Count == 2);
+	float Y = ValuePtr[0];
+	float X = ValuePtr[1];
+	return X >= Y ? 1 : 0;
+}
+
+FTG_Texture	UTG_Expression_Step::EvaluateTexture(FTG_EvaluationContext* InContext)
+{
+	return T_Maths_TwoInputs::CreateStep(InContext->Cycle, Output.EditTexture().GetBufferDescriptor(), InContext->TargetId, Input1.GetTexture(InContext), Input2.GetTexture(InContext));
+}
+
 //////////////////////////////////////////////////////////////////////////
 /// Power
 //////////////////////////////////////////////////////////////////////////
diff --git a/Engine/Plugins/Experimental/TextureGraph/Source/TextureGraph/Public/Expressions/Maths/TG_Expression_Maths_TwoInputs.h b/Engine/Plugins/Experimental/TextureGraph/Source/TextureGraph/Public/Expressions/Maths/TG_Expression_Maths_TwoInputs.h
--- a/Engine/Plugins/Experimental/TextureGraph/Source/TextureGraph/Public/Expressions/Maths/TG_Expression_Maths_TwoInputs.h
+++ b/Engine/Plugins/Experimental/TextureGraph/Source/TextureGraph/Public/Expressions/Maths/TG_Expression_Maths_TwoInputs.h
@@ -129,6 +129,22 @@ protected:
 	virtual FTG_Texture					EvaluateTexture(FTG_EvaluationContext* InContext) override;
 };
 
+//////////////////////////////////////////////////////////////////////////
+/// Step function: Usual HLSL step style function
+//////////////////////////////////////////////////////////////////////////
+UCLASS()
+class TEXTUREGRAPH_API UTG_Expression_Step : public UTG_Expression_BasicMath_Op
+{
+	GENERATED_BODY()
+	
+public:
+	virtual FText						GetTooltipText() const override { return FText::FromString(TEXT("step(y, x) returns 1 if the x parameter is greater than or equal to the y parameter; otherwise, 0.")); } 
+
+protected:
+	virtual float						EvaluateScalar_WithValue(FTG_EvaluationContext* InContext, const float* const Value, size_t Count) override;
+	virtual FTG_Texture					EvaluateTexture(FTG_EvaluationContext* InContext) override;
+};
+
 //////////////////////////////////////////////////////////////////////////
 /// Power
 //////////////////////////////////////////////////////////////////////////
diff --git a/Engine/Plugins/Experimental/TextureGraph/Source/TextureGraphEngine/Transform/Expressions/T_Maths_TwoInputs.cpp b/Engine/Plugins/Experimental/TextureGraph/Source/TextureGraphEngine/Transform/Expressions/T_Maths_TwoInputs.cpp
--- a/Engine/Plugins/Experimental/TextureGraph/Source/TextureGraphEngine/Transform/Expressions/T_Maths_TwoInputs.cpp
+++ b/Engine/Plugins/Experimental/TextureGraph/Source/TextureGraphEngine/Transform/Expressions/T_Maths_TwoInputs.cpp
@@ -16,6 +16,7 @@ IMPLEMENT_GLOBAL_SHADER(FSH_Subtract, "/Plugin/TextureGraph/Expressions/Expressi
 IMPLEMENT_GLOBAL_SHADER(FSH_Dot, "/Plugin/TextureGraph/Expressions/Expression_Maths_TwoInputs.usf", "FSH_Dot", SF_Pixel);
 IMPLEMENT_GLOBAL_SHADER(FSH_Cross, "/Plugin/TextureGraph/Expressions/Expression_Maths_TwoInputs.usf", "FSH_Cross", SF_Pixel);
 IMPLEMENT_GLOBAL_SHADER(FSH_Pow, "/Plugin/TextureGraph/Expressions/Expression_Maths_TwoInputs.usf", "FSH_Pow", SF_Pixel);
+IMPLEMENT_GLOBAL_SHADER(FSH_Step, "/Plugin/TextureGraph/Expressions/Expression_Maths_TwoInputs.usf", "FSH_Step", SF_Pixel);
 
 IMPLEMENT_GLOBAL_SHADER(FSH_IfThenElse, "/Plugin/TextureGraph/Expressions/Expression_IfThenElse.usf", "FSH_IfThenElse_GT_Component", SF_Pixel);
 
@@ -99,6 +100,11 @@ TiledBlobPtr T_Maths_TwoInputs::CreatePow(MixUpdateCyclePtr Cycle, BufferDescrip
 	return CreateGenericMathOp<FSH_Pow>(Cycle, DesiredOutputDesc, TargetId, Operand1, Operand2, "T_Pow");
 }
 
+TiledBlobPtr T_Maths_TwoInputs::CreateStep(MixUpdateCyclePtr Cycle, BufferDescriptor DesiredOutputDesc, int32 TargetId, TiledBlobPtr Operand1, TiledBlobPtr Operand2)
+{
+	return CreateGenericMathOp<FSH_Step>(Cycle, DesiredOutputDesc, TargetId, Operand1, Operand2, "T_Step");
+}
+
 RenderMaterial_FXPtr GetIfThenElseMaterial(EIfThenElseOperator Operator, EIfThenElseType Type)
 {
 	switch (Operator)
diff --git a/Engine/Plugins/Experimental/TextureGraph/Source/TextureGraphEngine/Transform/Expressions/T_Maths_TwoInputs.h b/Engine/Plugins/Experimental/TextureGraph/Source/TextureGraphEngine/Transform/Expressions/T_Maths_TwoInputs.h
--- a/Engine/Plugins/Experimental/TextureGraph/Source/TextureGraphEngine/Transform/Expressions/T_Maths_TwoInputs.h
+++ b/Engine/Plugins/Experimental/TextureGraph/Source/TextureGraphEngine/Transform/Expressions/T_Maths_TwoInputs.h
@@ -35,6 +35,7 @@ DECLARE_EMPTY_GLOBAL_SHADER_DERIVED_FROM(FSH_Subtract, FSH_BasicMathOp);
 DECLARE_EMPTY_GLOBAL_SHADER_DERIVED_FROM(FSH_Dot, FSH_BasicMathOp);
 DECLARE_EMPTY_GLOBAL_SHADER_DERIVED_FROM(FSH_Cross, FSH_BasicMathOp);
 DECLARE_EMPTY_GLOBAL_SHADER_DERIVED_FROM(FSH_Pow, FSH_BasicMathOp);
+DECLARE_EMPTY_GLOBAL_SHADER_DERIVED_FROM(FSH_Step, FSH_BasicMathOp);
 
 //////////////////////////////////////////////////////////////////////////
 /// Comparison ops
@@ -119,6 +120,7 @@ public:
 	static TiledBlobPtr				CreateDot(MixUpdateCyclePtr Cycle, BufferDescriptor DesiredOutputDesc, int32 TargetId, TiledBlobPtr Operand1, TiledBlobPtr Operand2);
 	static TiledBlobPtr				CreateCross(MixUpdateCyclePtr Cycle, BufferDescriptor DesiredOutputDesc, int32 TargetId, TiledBlobPtr Operand1, TiledBlobPtr Operand2);
 	static TiledBlobPtr				CreatePow(MixUpdateCyclePtr Cycle, BufferDescriptor DesiredOutputDesc, int32 TargetId, TiledBlobPtr Operand1, TiledBlobPtr Operand2);
+	static TiledBlobPtr				CreateStep(MixUpdateCyclePtr Cycle, BufferDescriptor DesiredOutputDesc, int32 TargetId, TiledBlobPtr Operand1, TiledBlobPtr Operand2);
 	
 	static TiledBlobPtr				CreateIfThenElse(MixUpdateCyclePtr Cycle, BufferDescriptor DesiredOutputDesc, int32 TargetId,
 		TiledBlobPtr LHS, TiledBlobPtr RHS, TiledBlobPtr Then, TiledBlobPtr Else, EIfThenElseOperator Operator, EIfThenElseType Type);
