UE-216229 - Fix variant crash in some texture graphs.

diff --git a/Engine/Plugins/Experimental/TextureGraph/Source/TextureGraph/Private/Expressions/TG_Expression_Variant.cpp b/Engine/Plugins/Experimental/TextureGraph/Source/TextureGraph/Private/Expressions/TG_Expression_Variant.cpp
--- a/Engine/Plugins/Experimental/TextureGraph/Source/TextureGraph/Private/Expressions/TG_Expression_Variant.cpp
+++ b/Engine/Plugins/Experimental/TextureGraph/Source/TextureGraph/Private/Expressions/TG_Expression_Variant.cpp
@@ -93,6 +93,8 @@ void UTG_Expression_Variant::Evaluate(FTG_EvaluationContext* InContext)
 {
 	Super::Evaluate(InContext);
 
+	CommonVariantType = EvalExpressionCommonVariantType();
+
 	switch (CommonVariantType)
 	{
 	case FTG_Variant::EType::Scalar:
@@ -113,5 +115,10 @@ void UTG_Expression_Variant::Evaluate(FTG_EvaluationContext* InContext)
 		Output.EditTexture() = EvaluateTexture(InContext).RasterBlob;
 		break;
 	}
+	default:
+	{
+		Output.EditTexture() = TextureHelper::GetMagenta();
+		break;
+	}
 	}
 }
diff --git a/Engine/Plugins/Experimental/TextureGraph/Source/TextureGraph/Private/TG_Graph.cpp b/Engine/Plugins/Experimental/TextureGraph/Source/TextureGraph/Private/TG_Graph.cpp
--- a/Engine/Plugins/Experimental/TextureGraph/Source/TextureGraph/Private/TG_Graph.cpp
+++ b/Engine/Plugins/Experimental/TextureGraph/Source/TextureGraph/Private/TG_Graph.cpp
@@ -581,22 +581,26 @@ void UTG_Graph::RemoveEdge(UTG_Node& NodeFrom, FTG_Name& PinFromName, UTG_Node&
 void UTG_Graph::AppendParamsSignature(FTG_Arguments& InOutArguments, TArray<FTG_Id>& InParams,
                                       TArray<FTG_Id>& OutParams) const
 {
-	for (auto pid : Params)
+	for (auto PinId : Params)
 	{
-		const UTG_Pin* ParamPin = GetPin(pid.Value);
+		const UTG_Pin* ParamPin = GetPin(PinId.Value);
 
-		FTG_Argument Argument = {
-			ParamPin->GetAliasName(), // Use the alias name to export the param as the graph interface
-			ParamPin->GetArgumentCPPTypeName(), // Same CPP type name
-			ParamPin->GetArgumentType().Unparamed() // Remove the param tag since it is no longer a param
-		}; 
+		/// The pin has to be marked connectable to show up
+		if (!ParamPin->IsNotConnectable())
+		{
+			FTG_Argument Argument = {
+				ParamPin->GetAliasName(), // Use the alias name to export the param as the graph interface
+				ParamPin->GetArgumentCPPTypeName(), // Same CPP type name
+				ParamPin->GetArgumentType().Unparamed() // Remove the param tag since it is no longer a param
+			}; 
 
-		InOutArguments.Emplace(Argument);
+			InOutArguments.Emplace(Argument);
 
-		if (ParamPin->IsInput())
-			InParams.Emplace(pid.Value);
-		else
-			OutParams.Emplace(pid.Value);
+			if (ParamPin->IsInput())
+				InParams.Emplace(PinId.Value);
+			else
+				OutParams.Emplace(PinId.Value);
+		}
 	}
 }
 
@@ -768,11 +772,14 @@ void UTG_Graph::ForEachEdges(std::function<void(const UTG_Pin* /*pinFrom*/, cons
 				if (APinId < BPinId)
 				{
 					const UTG_Pin* BPin = GetPin(BPinId);
-					check(BPin);
-					if (APinIsOutputAkaSource)
-						visitor(APin, BPin);
-					else
-						visitor(BPin, APin);
+
+					if (BPin)
+					{
+						if (APinIsOutputAkaSource)
+							visitor(APin, BPin);
+						else
+							visitor(BPin, APin);
+					}
 				}
 			}
 		}
@@ -849,34 +856,38 @@ int UTG_Graph::GetAllOutputParamValues(TArray<FTG_Variant>& OutVariants, TArray<
 
 FTG_Ids UTG_Graph::GatherSourceNodes(const UTG_Node* InNode) const
 {
-	FTG_Ids sourceNodes; // the array of source nodes that will be return, empty for now
+	FTG_Ids SourceNodes; // the array of source nodes that will be return, empty for now
 
 	if (InNode) // Only work for a valid node
 	{
-		bool allInConnected = true;
-		bool allInDisconnected = true;
+		bool AllInConnected = true;
+		bool AllInDisconnected = true;
 		auto InPinIds = InNode->GetInputPinIds();
 		for (auto opi : InPinIds) // Over all the Input pins, find the node feeding them
 		{
 			const UTG_Pin* Pin = InNode->Pins[opi.PinIdx()];
 			if (Pin->GetEdges().IsEmpty())
 			{
-				allInConnected = false;
+				AllInConnected = false;
 			}
 			else
 			{
-				allInDisconnected = false;
+				AllInDisconnected = false;
 				auto EdgePinId = Pin->GetEdges()[0];
 				const UTG_Pin* SourcePin = GetPin(EdgePinId);
-				check(SourcePin);
 
-				// only add the source node if not already there
-				if (sourceNodes.Find(SourcePin->GetNodeId()) == INDEX_NONE)
-					sourceNodes.Emplace(SourcePin->GetNodeId());
+				if (SourcePin)
+				{
+					check(SourcePin);
+
+					// only add the source node if not already there
+					if (SourceNodes.Find(SourcePin->GetNodeId()) == INDEX_NONE)
+						SourceNodes.Emplace(SourcePin->GetNodeId());
+				}
 			}
 		}
 	}
-	return sourceNodes;
+	return SourceNodes;
 }
 
 FTG_Ids UTG_Graph::GatherAllSourceNodes(const UTG_Node* Node) const
diff --git a/Engine/Plugins/Experimental/TextureGraph/Source/TextureGraph/Private/TG_GraphEvaluation.cpp b/Engine/Plugins/Experimental/TextureGraph/Source/TextureGraph/Private/TG_GraphEvaluation.cpp
--- a/Engine/Plugins/Experimental/TextureGraph/Source/TextureGraph/Private/TG_GraphEvaluation.cpp
+++ b/Engine/Plugins/Experimental/TextureGraph/Source/TextureGraph/Private/TG_GraphEvaluation.cpp
@@ -13,6 +13,10 @@ void FTG_Evaluation::TransferVarToPin(UTG_Pin* InPin, FTG_EvaluationContext* Con
 
 	const auto VarId = InPin->GetVarId();
 	FTG_Var* Var = Context->Graph->GetVar(VarId);
+
+	if (!Var)
+		return;
+
 	check(Var);
 
 	if (Arg.IsInput())
diff --git a/Engine/Plugins/Experimental/TextureGraph/Source/TextureGraph/Private/TG_Node.cpp b/Engine/Plugins/Experimental/TextureGraph/Source/TextureGraph/Private/TG_Node.cpp
--- a/Engine/Plugins/Experimental/TextureGraph/Source/TextureGraph/Private/TG_Node.cpp
+++ b/Engine/Plugins/Experimental/TextureGraph/Source/TextureGraph/Private/TG_Node.cpp
@@ -498,8 +498,14 @@ FTG_Variant::EType UTG_Node::EvalExpressionCommonVariantType() const
 			{
 				// now we know the thing connected to the pin variant exist, what is its type?
 				UTG_Pin* OtherPin = GetGraph()->GetPin(Pin->GetVarId());
-				FTG_Variant::EType OtherType = FTG_Variant::GetTypeFromName(OtherPin->GetArgumentCPPTypeName());
-				Type = FTG_Variant::WhichCommonType(Type, OtherType);
+
+				if (OtherPin)
+				{
+					FTG_Variant::EType OtherType = FTG_Variant::GetTypeFromName(OtherPin->GetArgumentCPPTypeName());
+					Type = FTG_Variant::WhichCommonType(Type, OtherType);
+				}
+				else
+					return FTG_Variant::EType::Invalid;
 			}
 			// removed the else part here since the selfvar/default will be reset to the common type on the next evaluation
 			/*else
diff --git a/Engine/Plugins/Experimental/TextureGraph/Source/TextureGraph/Private/TG_Pin.cpp b/Engine/Plugins/Experimental/TextureGraph/Source/TextureGraph/Private/TG_Pin.cpp
--- a/Engine/Plugins/Experimental/TextureGraph/Source/TextureGraph/Private/TG_Pin.cpp
+++ b/Engine/Plugins/Experimental/TextureGraph/Source/TextureGraph/Private/TG_Pin.cpp
@@ -242,6 +242,9 @@ FString UTG_Pin::GetEvaluatedVarValue() const
 		}
 	}
 
+	if (!CurrentVar)
+		return FString();
+
 	FString CurrentValue = CurrentVar->LogValue();
 
 	// Enum requires special-case handling
diff --git a/Engine/Plugins/Experimental/TextureGraph/Source/TextureGraph/Public/Expressions/Utilities/TG_Expression_MaterialID.h b/Engine/Plugins/Experimental/TextureGraph/Source/TextureGraph/Public/Expressions/Utilities/TG_Expression_MaterialID.h
--- a/Engine/Plugins/Experimental/TextureGraph/Source/TextureGraph/Public/Expressions/Utilities/TG_Expression_MaterialID.h
+++ b/Engine/Plugins/Experimental/TextureGraph/Source/TextureGraph/Public/Expressions/Utilities/TG_Expression_MaterialID.h
@@ -22,7 +22,7 @@ public:
 	FTG_Texture MaterialIDMap;
 
 	// List of available colors extracted from the input texture. Selected colors will be shown as white in the output texture. 
-	UPROPERTY(EditAnywhere, EditFixedSize, Category = NoCategory, meta = (TGType = "TG_InputParam", NoResetToDefault))
+	UPROPERTY(EditAnywhere, EditFixedSize, Category = NoCategory, meta = (TGType = "TG_InputParam", NoResetToDefault, TGPinNotConnectable = "true"))
 	TArray<FMaterialIDMaskInfo> MaterialIDMaskInfos;
 
 	UPROPERTY()
diff --git a/Engine/Plugins/Experimental/TextureGraph/Source/TextureGraph/Public/TG_Variant.h b/Engine/Plugins/Experimental/TextureGraph/Source/TextureGraph/Public/TG_Variant.h
--- a/Engine/Plugins/Experimental/TextureGraph/Source/TextureGraph/Public/TG_Variant.h
+++ b/Engine/Plugins/Experimental/TextureGraph/Source/TextureGraph/Public/TG_Variant.h
@@ -16,6 +16,7 @@ struct FTG_EvaluationContext;
 UENUM()
 enum class ETG_VariantType : uint8
 {
+	Invalid = 0xFF												UMETA(DisplayName = "Invalid"),
 	Scalar = FTG_VariantInnerData::IndexOfType<float>()			UMETA(DisplayName = "Scalar"),
 	Color = FTG_VariantInnerData::IndexOfType<FLinearColor>()	UMETA(DisplayName = "Color"),
 	Vector = FTG_VariantInnerData::IndexOfType<FVector4f>()		UMETA(DisplayName = "Vector"),
