Converting all the hsv tweaking nodes to use variants and support colors in addition to textures

diff --git a/Engine/Plugins/Experimental/TextureGraph/Source/TextureGraph/Private/Expressions/Color/TG_Expression_HSV.cpp b/Engine/Plugins/Experimental/TextureGraph/Source/TextureGraph/Private/Expressions/Color/TG_Expression_HSV.cpp
--- a/Engine/Plugins/Experimental/TextureGraph/Source/TextureGraph/Private/Expressions/Color/TG_Expression_HSV.cpp
+++ b/Engine/Plugins/Experimental/TextureGraph/Source/TextureGraph/Private/Expressions/Color/TG_Expression_HSV.cpp
@@ -4,86 +4,117 @@
 #include "FxMat/MaterialManager.h"
 #include "Job/JobBatch.h"
 #include "Transform/Expressions/T_Color.h"
+#include "TextureGraphEngine/Helper/ColorUtil.h"
 
-void UTG_Expression_HSV::Evaluate(FTG_EvaluationContext* InContext)
+//////////////////////////////////////////////////////////////////////////
+/// RGB2HSV Correction
+//////////////////////////////////////////////////////////////////////////
+FTG_Texture UTG_Expression_RGB2HSV::EvaluateTexture(FTG_EvaluationContext* InContext)
 {
-	Super::Evaluate(InContext);
-
-	const RenderMaterial_FXPtr RenderMaterial = TextureGraphEngine::GetMaterialManager()->CreateMaterial_FX<VSH_Simple, FSH_HSV>(TEXT("T_HSV"));
+	const RenderMaterial_FXPtr RenderMaterial = TextureGraphEngine::GetMaterialManager()->CreateMaterial_FX<VSH_Simple, FSH_RGB2HSV>(TEXT("T_RGB2HSV"));
 
 	check(RenderMaterial);
 
 	if (!Input)
 	{
-		Output = FTG_Texture::GetBlack();
-		return;
+		return FTG_Texture::GetBlack();
 	}
-	
+
 	JobUPtr RenderJob = std::make_unique<Job>(InContext->Cycle->GetMix(), InContext->TargetId, std::static_pointer_cast<BlobTransform>(RenderMaterial));
 
 	RenderJob
-		->AddArg(ARG_BLOB(Input, "Input"))
-		->AddArg(ARG_FLOAT(Hue, "Hue"))
-		->AddArg(ARG_FLOAT(Saturation, "Saturation"))
-		->AddArg(ARG_FLOAT(Value, "Value"))
+		->AddArg(ARG_BLOB(Input.GetTexture(InContext), "Input"))
 		;
 
-	const FString Name = TEXT("HSV"); // FString::Printf(TEXT("Grayscale.[%s].[%d].[%llu]"), *gmask->ID(), InContext->TargetId, InContext->Cycle->Batch()->BatchId()); 
-	BufferDescriptor Desc = Output.GetBufferDescriptor();
+	const FString Name = TEXT("RGB2HSV"); // FString::Printf(TEXT("Grayscale.[%s].[%D].[%llu]"), *gmask->ID(), InContext->TargetId, InContext->Cycle->Batch()->BatchId());
+	BufferDescriptor Desc = Output.EditTexture().GetBufferDescriptor();
 
-	Output = RenderJob->InitResult(Name, &Desc);
+	TiledBlobPtr OutputResult = RenderJob->InitResult(Name, &Desc);
 	InContext->Cycle->AddJob(InContext->TargetId, std::move(RenderJob));
+
+	return OutputResult;
 }
 
-void UTG_Expression_RGB2HSV::Evaluate(FTG_EvaluationContext* InContext)
+FVector4f UTG_Expression_RGB2HSV::EvaluateVector_WithValue(FTG_EvaluationContext* InContext, const FVector4f* const ValuePtr, size_t Count)
 {
-	Super::Evaluate(InContext);
+	check(ValuePtr && Count == 1); 
+	const FVector4f& C = ValuePtr[0];
+	FLinearColor Result = ColorUtil::RGB2HSV(FLinearColor(C.X, C.Y, C.Z, 1.0f));
+	return FVector4f(Result.R, Result.G, Result.B, 1.0f);
+}
 
-	const RenderMaterial_FXPtr RenderMaterial = TextureGraphEngine::GetMaterialManager()->CreateMaterial_FX<VSH_Simple, FSH_RGB2HSV>(TEXT("T_RGB2HSV"));
+//////////////////////////////////////////////////////////////////////////
+/// HSV2RGB Correction
+//////////////////////////////////////////////////////////////////////////
+FTG_Texture UTG_Expression_HSV2RGB::EvaluateTexture(FTG_EvaluationContext* InContext)
+{
+	const RenderMaterial_FXPtr RenderMaterial = TextureGraphEngine::GetMaterialManager()->CreateMaterial_FX<VSH_Simple, FSH_HSV2RGB>(TEXT("T_HSV2RGB"));
 
 	check(RenderMaterial);
 
 	if (!Input)
 	{
-		Output = FTG_Texture::GetBlack();
-		return;
+		return FTG_Texture::GetBlack();
 	}
-	
+
 	JobUPtr RenderJob = std::make_unique<Job>(InContext->Cycle->GetMix(), InContext->TargetId, std::static_pointer_cast<BlobTransform>(RenderMaterial));
 
 	RenderJob
-		->AddArg(ARG_BLOB(Input, "Input"))
+		->AddArg(ARG_BLOB(Input.GetTexture(InContext), "Input"))
 		;
 
-	const FString Name = TEXT("RGB2HSV"); // FString::Printf(TEXT("Grayscale.[%s].[%d].[%llu]"), *gmask->ID(), InContext->TargetId, InContext->Cycle->Batch()->BatchId());
-	BufferDescriptor Desc = Output.GetBufferDescriptor();
-
-	Output = RenderJob->InitResult(Name, &Desc);
+	const FString Name = TEXT("HSV2RGB"); // FString::Printf(TEXT("Grayscale.[%s].[%D].[%llu]"), *gmask->ID(), InContext->TargetId, InContext->Cycle->Batch()->BatchId());
+	BufferDescriptor Desc = Output.EditTexture().GetBufferDescriptor();
+	TiledBlobPtr OutputResult = RenderJob->InitResult(Name, &Desc);
 	InContext->Cycle->AddJob(InContext->TargetId, std::move(RenderJob));
+
+	return OutputResult;
 }
 
-void UTG_Expression_HSV2RGB::Evaluate(FTG_EvaluationContext* InContext)
+FVector4f UTG_Expression_HSV2RGB::EvaluateVector_WithValue(FTG_EvaluationContext* InContext, const FVector4f* const ValuePtr, size_t Count)
 {
-	Super::Evaluate(InContext);
+	check(ValuePtr && Count == 1); 
+	const FVector4f& C = ValuePtr[0];
+	FLinearColor Result = ColorUtil::HSV2RGB(FLinearColor(C.X, C.Y, C.Z, 1.0f));
+	return FVector4f(Result.R, Result.G, Result.B, 1.0f);
+}
 
-	const RenderMaterial_FXPtr RenderMaterial = TextureGraphEngine::GetMaterialManager()->CreateMaterial_FX<VSH_Simple, FSH_HSV2RGB>(TEXT("T_HSV2RGB"));
+//////////////////////////////////////////////////////////////////////////
+/// HSV Correction
+//////////////////////////////////////////////////////////////////////////
+FTG_Texture UTG_Expression_HSV::EvaluateTexture(FTG_EvaluationContext* InContext)
+{
+	const RenderMaterial_FXPtr RenderMaterial = TextureGraphEngine::GetMaterialManager()->CreateMaterial_FX<VSH_Simple, FSH_HSV>(TEXT("T_HSV"));
 
 	check(RenderMaterial);
 
 	if (!Input)
 	{
-		Output = FTG_Texture::GetBlack();
-		return;
+		return FTG_Texture::GetBlack();
 	}
-	
+
 	JobUPtr RenderJob = std::make_unique<Job>(InContext->Cycle->GetMix(), InContext->TargetId, std::static_pointer_cast<BlobTransform>(RenderMaterial));
 
 	RenderJob
-		->AddArg(ARG_BLOB(Input, "Input"))
+		->AddArg(ARG_BLOB(Input.GetTexture(InContext), "Input"))
+		->AddArg(ARG_FLOAT(Hue, "Hue"))
+		->AddArg(ARG_FLOAT(Saturation, "Saturation"))
+		->AddArg(ARG_FLOAT(Value, "Value"))
 		;
 
-	const FString Name = TEXT("HSV2RGB"); // FString::Printf(TEXT("Grayscale.[%s].[%d].[%llu]"), *gmask->ID(), InContext->TargetId, InContext->Cycle->Batch()->BatchId());
-	BufferDescriptor Desc = Output.GetBufferDescriptor();
-	Output = RenderJob->InitResult(Name, &Desc);
+	const FString Name = TEXT("HSV"); // FString::Printf(TEXT("Grayscale.[%s].[%D].[%llu]"), *gmask->ID(), InContext->TargetId, InContext->Cycle->Batch()->BatchId()); 
+	BufferDescriptor Desc = Output.EditTexture().GetBufferDescriptor();
+
+	TiledBlobPtr OutputResult = RenderJob->InitResult(Name, &Desc);
 	InContext->Cycle->AddJob(InContext->TargetId, std::move(RenderJob));
+
+	return OutputResult;
+}
+
+FVector4f UTG_Expression_HSV::EvaluateVector_WithValue(FTG_EvaluationContext* InContext, const FVector4f* const ValuePtr, size_t Count)
+{
+	check(ValuePtr && Count == 1); 
+	const FVector4f& C = ValuePtr[0];
+	FLinearColor Result = ColorUtil::HSVTweak(FLinearColor(C.X, C.Y, C.Z, 1.0f), Hue, Saturation, Value);
+	return FVector4f(Result.R, Result.G, Result.B, 1.0f);
 }
diff --git a/Engine/Plugins/Experimental/TextureGraph/Source/TextureGraph/Private/Expressions/Maths/TG_Expression_Maths_TwoInputs.cpp b/Engine/Plugins/Experimental/TextureGraph/Source/TextureGraph/Private/Expressions/Maths/TG_Expression_Maths_TwoInputs.cpp
--- a/Engine/Plugins/Experimental/TextureGraph/Source/TextureGraph/Private/Expressions/Maths/TG_Expression_Maths_TwoInputs.cpp
+++ b/Engine/Plugins/Experimental/TextureGraph/Source/TextureGraph/Private/Expressions/Maths/TG_Expression_Maths_TwoInputs.cpp
@@ -4,6 +4,7 @@
 #include "Transform/Expressions/T_Maths_TwoInputs.h"
 #include "FxMat/MaterialManager.h"
 #include "Job/JobBatch.h"
+#include "TextureGraphEngine/Helper/MathUtils.h"
 
 typedef std::function<TiledBlobPtr(MixUpdateCyclePtr /* Cycle */, BufferDescriptor DesiredOutputDesc, int32 /*TargetId*/, TiledBlobPtr /*Operand1*/, TiledBlobPtr /*Operand2*/)> MathOpFunc;
 
@@ -109,7 +110,7 @@ float UTG_Expression_Step::EvaluateScalar_WithValue(FTG_EvaluationContext* InCon
 	check(ValuePtr && Count == 2);
 	float Y = ValuePtr[0];
 	float X = ValuePtr[1];
-	return X >= Y ? 1 : 0;
+	return MathUtils::Step(Y, X);
 }
 
 FTG_Texture	UTG_Expression_Step::EvaluateTexture(FTG_EvaluationContext* InContext)
diff --git a/Engine/Plugins/Experimental/TextureGraph/Source/TextureGraph/Private/Expressions/TG_Expression_Variant.cpp b/Engine/Plugins/Experimental/TextureGraph/Source/TextureGraph/Private/Expressions/TG_Expression_Variant.cpp
--- a/Engine/Plugins/Experimental/TextureGraph/Source/TextureGraph/Private/Expressions/TG_Expression_Variant.cpp
+++ b/Engine/Plugins/Experimental/TextureGraph/Source/TextureGraph/Private/Expressions/TG_Expression_Variant.cpp
@@ -93,7 +93,7 @@ void UTG_Expression_Variant::Evaluate(FTG_EvaluationContext* InContext)
 {
 	Super::Evaluate(InContext);
 
-	CommonVariantType = EvalExpressionCommonVariantType();
+	//CommonVariantType = EvalExpressionCommonVariantType();
 
 	switch (CommonVariantType)
 	{
@@ -117,7 +117,7 @@ void UTG_Expression_Variant::Evaluate(FTG_EvaluationContext* InContext)
 	}
 	default:
 	{
-		Output.EditTexture() = TextureHelper::GetMagenta();
+		Output.EditScalar() = 0.0f;
 		break;
 	}
 	}
diff --git a/Engine/Plugins/Experimental/TextureGraph/Source/TextureGraph/Private/TG_Node.cpp b/Engine/Plugins/Experimental/TextureGraph/Source/TextureGraph/Private/TG_Node.cpp
--- a/Engine/Plugins/Experimental/TextureGraph/Source/TextureGraph/Private/TG_Node.cpp
+++ b/Engine/Plugins/Experimental/TextureGraph/Source/TextureGraph/Private/TG_Node.cpp
@@ -488,7 +488,7 @@ FTG_Variant::EType UTG_Node::GetExpressionCommonVariantType() const
 
 FTG_Variant::EType UTG_Node::EvalExpressionCommonVariantType() const
 {
-	FTG_Variant::EType Type = FTG_Variant::EType::Scalar;
+	FTG_Variant::EType Type = FTG_Variant::EType::Invalid;
 
 	for (auto Pin : Pins)
 	{
diff --git a/Engine/Plugins/Experimental/TextureGraph/Source/TextureGraph/Public/Expressions/Color/TG_Expression_HSV.h b/Engine/Plugins/Experimental/TextureGraph/Source/TextureGraph/Public/Expressions/Color/TG_Expression_HSV.h
--- a/Engine/Plugins/Experimental/TextureGraph/Source/TextureGraph/Public/Expressions/Color/TG_Expression_HSV.h
+++ b/Engine/Plugins/Experimental/TextureGraph/Source/TextureGraph/Public/Expressions/Color/TG_Expression_HSV.h
@@ -6,6 +6,7 @@
 
 #include "Expressions/TG_Expression.h"
 #include "TG_Texture.h"
+#include "Expressions/Maths/TG_Expression_Maths_OneInput.h"
 
 #include "TG_Expression_HSV.generated.h"
 
@@ -13,68 +14,51 @@
 /// RGB2HSV Correction
 //////////////////////////////////////////////////////////////////////////
 UCLASS()
-class TEXTUREGRAPH_API UTG_Expression_RGB2HSV : public UTG_Expression
+class TEXTUREGRAPH_API UTG_Expression_RGB2HSV : public UTG_Expression_OneInput
 {
 	GENERATED_BODY()
 
 public:
-	TG_DECLARE_EXPRESSION(TG_Category::Adjustment);
-	virtual void						Evaluate(FTG_EvaluationContext* InContext) override;
-
-	// The input image in RGB to be converted to in HSV color space
-	UPROPERTY(meta = (TGType = "TG_Input", PinDisplayName = ""))
-	FTG_Texture							Input;
-
-	// The output image which will be in HSV color space 
-	UPROPERTY(EditAnywhere, Category = NoCategory, meta = (TGType = "TG_Output", PinDisplayName = ""))
-	FTG_Texture							Output;
-
 	virtual FTG_Name					GetDefaultName() const override { return TEXT("RGBtoHSV");}
 	virtual FText						GetTooltipText() const override { return FText::FromString(TEXT("Interprets the input as RGB and converts it to HSV.")); } 
+	virtual FName						GetCategory() const override { return TG_Category::Adjustment; } 
 
+protected:
+	/// This is unsupported because it only makes sense on a vector input
+	virtual float						EvaluateScalar_WithValue(FTG_EvaluationContext* InContext, const float* const Value, size_t Count) override { return 0.0f; }
+	virtual FVector4f					EvaluateVector_WithValue(FTG_EvaluationContext* InContext, const FVector4f* const Value, size_t Count);
+	virtual FTG_Texture					EvaluateTexture(FTG_EvaluationContext* InContext) override;
 };
 
 //////////////////////////////////////////////////////////////////////////
 /// HSV2RGB Correction
 //////////////////////////////////////////////////////////////////////////
 UCLASS()
-class TEXTUREGRAPH_API UTG_Expression_HSV2RGB : public UTG_Expression
+class TEXTUREGRAPH_API UTG_Expression_HSV2RGB : public UTG_Expression_OneInput
 {
 	GENERATED_BODY()
 
 public:
-	TG_DECLARE_EXPRESSION(TG_Category::Adjustment);
-	virtual void						Evaluate(FTG_EvaluationContext* InContext) override;
-
-	// The input image in HSV to be converted to in RGB color space
-	UPROPERTY(meta = (TGType = "TG_Input", PinDisplayName = ""))
-	FTG_Texture							Input;
-
-	// The output image which will be in RGB color space 
-	UPROPERTY(EditAnywhere, Category = NoCategory, meta = (TGType = "TG_Output", PinDisplayName = ""))
-	FTG_Texture							Output;
-
 	virtual FTG_Name					GetDefaultName() const override { return TEXT("HSVtoRGB");}
 	virtual FText						GetTooltipText() const override { return FText::FromString(TEXT("Interprets the input as HSV and converts it to RGB.")); } 
+	virtual FName						GetCategory() const override { return TG_Category::Adjustment; } 
 
+protected:
+	/// This is unsupported because it only makes sense on a vector input
+	virtual float						EvaluateScalar_WithValue(FTG_EvaluationContext* InContext, const float* const Value, size_t Count) override { return 0.0f; }
+	virtual FVector4f					EvaluateVector_WithValue(FTG_EvaluationContext* InContext, const FVector4f* const Value, size_t Count);
+	virtual FTG_Texture					EvaluateTexture(FTG_EvaluationContext* InContext) override;
 };
 
 //////////////////////////////////////////////////////////////////////////
 /// HSV Correction
 //////////////////////////////////////////////////////////////////////////
 UCLASS()
-class TEXTUREGRAPH_API UTG_Expression_HSV : public UTG_Expression
+class TEXTUREGRAPH_API UTG_Expression_HSV : public UTG_Expression_OneInput
 {
 	GENERATED_BODY()
 
 public:
-	TG_DECLARE_EXPRESSION(TG_Category::Adjustment);
-	virtual void						Evaluate(FTG_EvaluationContext* InContext) override;
-
-	// The input image to be adjusted in HSV space
-	UPROPERTY(meta = (TGType = "TG_Input", PinDisplayName = ""))
-	FTG_Texture							Input;
-
 	// Defines the basic color tone, such as red, green, or blue, without considering brightness or intensity. Adjusting the hue changes the overall color appearance while maintaining its saturation and brightness.
 	// The normalized hue. Please divide your [0, 359] hue values by 359 for this input
 	UPROPERTY(EditAnywhere, Category = NoCategory, meta = (TGType = "TG_Setting", UIMin = "0", ClampMin = "0", UIMax = "1", ClampMax = "1", PinDisplayName = "Hue (Normalized)", DisplayName = "Hue (Normalized)"))
@@ -88,11 +72,13 @@ public:
 	UPROPERTY(EditAnywhere, Category = NoCategory, meta = (TGType = "TG_Setting", UIMin = "0", ClampMin = "0", UIMax = "1", ClampMax = "1"))
 	float								Value = 1.0f;
 
-	// The output image which is a single channel grayscale version of the input image 
-	UPROPERTY(EditAnywhere, Category = NoCategory, meta = (TGType = "TG_Output", PinDisplayName = ""))
-	FTG_Texture							Output;
-
 	virtual FText						GetTooltipText() const override { return FText::FromString(TEXT("Converts Hue, Saturation and Value inputs to an RGB Image.")); } 
+	virtual FName						GetCategory() const override { return TG_Category::Adjustment; } 
 
+protected:
+	/// This is unsupported because it only makes sense on a vector input
+	virtual float						EvaluateScalar_WithValue(FTG_EvaluationContext* InContext, const float* const ValuePtr, size_t Count) override { return 0.0f; }
+	virtual FVector4f					EvaluateVector_WithValue(FTG_EvaluationContext* InContext, const FVector4f* const ValuePtr, size_t Count);
+	virtual FTG_Texture					EvaluateTexture(FTG_EvaluationContext* InContext) override;
 };
 
diff --git a/Engine/Plugins/Experimental/TextureGraph/Source/TextureGraph/Public/Expressions/Maths/TG_Expression_Maths_OneInput.h b/Engine/Plugins/Experimental/TextureGraph/Source/TextureGraph/Public/Expressions/Maths/TG_Expression_Maths_OneInput.h
--- a/Engine/Plugins/Experimental/TextureGraph/Source/TextureGraph/Public/Expressions/Maths/TG_Expression_Maths_OneInput.h
+++ b/Engine/Plugins/Experimental/TextureGraph/Source/TextureGraph/Public/Expressions/Maths/TG_Expression_Maths_OneInput.h
@@ -47,7 +47,7 @@ protected:
 	virtual std::vector<FTG_Variant>	GetEvaluateArgs() override { return std::vector<FTG_Variant>({ Input }); }
 
 public:
-	// The angle in radians to run the trigonometric function on
+	// The input to the node
 	UPROPERTY(meta = (TGType = "TG_Input"))
 	FTG_Variant							Input;
 
diff --git a/Engine/Plugins/Experimental/TextureGraph/Source/TextureGraph/Public/TG_Variant.h b/Engine/Plugins/Experimental/TextureGraph/Source/TextureGraph/Public/TG_Variant.h
--- a/Engine/Plugins/Experimental/TextureGraph/Source/TextureGraph/Public/TG_Variant.h
+++ b/Engine/Plugins/Experimental/TextureGraph/Source/TextureGraph/Public/TG_Variant.h
@@ -14,9 +14,9 @@ struct FTG_EvaluationContext;
 
 // Types of the variant are organized in increasing complexity, so a compatible type for 2 variants is always the highest 
 UENUM()
-enum class ETG_VariantType : uint8
+enum class ETG_VariantType : int8
 {
-	Invalid = 0xFF												UMETA(DisplayName = "Invalid"),
+	Invalid = -1												UMETA(DisplayName = "Invalid"),
 	Scalar = FTG_VariantInnerData::IndexOfType<float>()			UMETA(DisplayName = "Scalar"),
 	Color = FTG_VariantInnerData::IndexOfType<FLinearColor>()	UMETA(DisplayName = "Color"),
 	Vector = FTG_VariantInnerData::IndexOfType<FVector4f>()		UMETA(DisplayName = "Vector"),
@@ -50,6 +50,10 @@ public:
 	// Retrieve the FName corresponding to a variant type
 	static FName GetNameFromType(ETG_VariantType InType)
 	{
+		/// Invalid defaults to type scalar
+		if (InType == ETG_VariantType::Invalid)
+			return TEXT("Scalar");
+
 		// We could use this: 
 		// return StaticEnum<ETG_VariantType>()->GetName(static_cast<uint32>(InType));
 		// But we want simpler names
@@ -65,6 +69,10 @@ public:
 	// Retrieve the FName associated to a varaint type used for Arg cpoptypename
 	static FName GetArgNameFromType(ETG_VariantType InType)
 	{
+		/// Invalid defaults to type scalar
+		if (InType == ETG_VariantType::Invalid)
+			return TEXT("FTG_Variant.Scalar");
+
 		static FName ArgVariantType_Names[] = {
 			TEXT("FTG_Variant.Scalar"),
 			TEXT("FTG_Variant.Color"),
@@ -100,6 +108,7 @@ public:
 	static bool IsColor(EType InType) { return (InType == EType::Color); }
 	static bool IsVector(EType InType) { return (InType == EType::Vector); }
 	static bool IsTexture(EType InType) { return (InType == EType::Texture); }
+	static bool IsInvalid(EType InType) { return (InType == EType::Invalid); }
 
 	// Find the common type between 2
 	static EType WhichCommonType(EType T0, EType T1)
diff --git a/Engine/Plugins/Experimental/TextureGraph/Source/TextureGraphEngine/Helper/ColorUtil.cpp b/Engine/Plugins/Experimental/TextureGraph/Source/TextureGraphEngine/Helper/ColorUtil.cpp
--- a/Engine/Plugins/Experimental/TextureGraph/Source/TextureGraphEngine/Helper/ColorUtil.cpp
+++ b/Engine/Plugins/Experimental/TextureGraph/Source/TextureGraphEngine/Helper/ColorUtil.cpp
@@ -1,5 +1,6 @@
 // Copyright Epic Games, Inc. All Rights Reserved.
 #include "ColorUtil.h"
+#include "TextureGraphEngine/Helper/MathUtils.h"
 
 const FLinearColor ColorUtil::DefaultColor(TextureType Type)
 {
@@ -161,6 +162,47 @@ bool ColorUtil::IsColorNear(const FLinearColor& Color, const FLinearColor& Ref,
 	return IsColorBlack(Diff, IgnoreAlpha);
 }
 
+FLinearColor ColorUtil::HSV2RGB(FLinearColor C)
+{
+	FVector4f K = FVector4f(1.0, 2.0 / 3.0, 1.0 / 3.0, 3.0);
+	FVector3f Kx(K.X, K.X, K.X);
+	FVector3f P;
+	P.X = FMath::Abs(FMath::Frac(C.R + K.X) * 6.0 - K.W);
+	P.Y = FMath::Abs(FMath::Frac(C.R + K.Y) * 6.0 - K.W);
+	P.Z = FMath::Abs(FMath::Frac(C.R + K.Z) * 6.0 - K.W);
+
+	FVector3f D = P - Kx;
+	D.X = FMath::Clamp(D.X, 0.0f, 1.0f);
+	D.Y = FMath::Clamp(D.Y, 0.0f, 1.0f);
+	D.Z = FMath::Clamp(D.Z, 0.0f, 1.0f);
+
+	FVector3f RGB = C.B * FMath::Lerp(Kx, D, C.G);
+
+	return FLinearColor(RGB.X, RGB.Y, RGB.Z, 1.0f);
+}
+
+FLinearColor ColorUtil::RGB2HSV(FLinearColor C)
+{
+	FVector4f K = FVector4f(0.0, -1.0 / 3.0, 2.0 / 3.0, -1.0);
+	FLinearColor P = FMath::Lerp(FLinearColor(C.B, C.G, K.W, K.Z), FLinearColor(C.G, C.B, K.X, K.Y), MathUtils::Step(C.B, C.G));
+	FLinearColor Q = FMath::Lerp(FLinearColor(P.R, P.G, P.B, C.R), FLinearColor(C.R, P.G, P.B, P.R), MathUtils::Step(P.R, C.R));
+	float D = Q.R - FMath::Min(Q.A, Q.G);
+	float E = 1.0e-10;
+	FVector3f HSV(FMath::Abs(Q.B + (Q.A - Q.G) / (6.0 * D + E)), D / (Q.R + E), Q.R);
+
+	return FLinearColor(HSV.X, HSV.Y, HSV.Z, 1.0f);
+}
+
+FLinearColor ColorUtil::HSVTweak(FLinearColor C, float H, float S, float V)
+{
+	FLinearColor ClrHSV = RGB2HSV(C);
+	ClrHSV *= FLinearColor(H, S, V, 1.0f);
+	FLinearColor Result = HSV2RGB(ClrHSV);
+	Result.A = C.A;
+
+	return Result;
+}
+
 bool ColorUtil::IsColorWhite(const FLinearColor& Color, bool IgnoreAlpha) { return IsColorNear(Color, FLinearColor::White, IgnoreAlpha); }
 bool ColorUtil::IsColorGray(const FLinearColor& Color, bool IgnoreAlpha) { return IsColorNear(Color, FLinearColor::Gray, IgnoreAlpha); }
 bool ColorUtil::IsColorRed(const FLinearColor& Color, bool IgnoreAlpha) { return IsColorNear(Color, FLinearColor::Red, IgnoreAlpha); }
diff --git a/Engine/Plugins/Experimental/TextureGraph/Source/TextureGraphEngine/Helper/ColorUtil.h b/Engine/Plugins/Experimental/TextureGraph/Source/TextureGraphEngine/Helper/ColorUtil.h
--- a/Engine/Plugins/Experimental/TextureGraph/Source/TextureGraphEngine/Helper/ColorUtil.h
+++ b/Engine/Plugins/Experimental/TextureGraph/Source/TextureGraphEngine/Helper/ColorUtil.h
@@ -25,6 +25,10 @@ public:
 	static float						GetSquaredDistance(FLinearColor Current, FLinearColor Match);
 	static FString						GetColorName(FLinearColor Color);
 
+	static FLinearColor					HSV2RGB(FLinearColor C);
+	static FLinearColor					RGB2HSV(FLinearColor C);
+	static FLinearColor					HSVTweak(FLinearColor C, float H, float S, float V);
+
 	static bool							IsColorBlack(const FLinearColor& Color, bool IgnoreAlpha = true);
 	static bool							IsColorWhite(const FLinearColor& Color, bool IgnoreAlpha = true);
 	static bool							IsColorGray(const FLinearColor& Color, bool IgnoreAlpha = true);
diff --git a/Engine/Plugins/Experimental/TextureGraph/Source/TextureGraphEngine/Helper/MathUtils.cpp b/Engine/Plugins/Experimental/TextureGraph/Source/TextureGraphEngine/Helper/MathUtils.cpp
--- a/Engine/Plugins/Experimental/TextureGraph/Source/TextureGraphEngine/Helper/MathUtils.cpp
+++ b/Engine/Plugins/Experimental/TextureGraph/Source/TextureGraphEngine/Helper/MathUtils.cpp
@@ -102,3 +102,8 @@ FBox MathUtils::GetCombinedBounds(TArray<FBox> InputBounds)
 	return CombinedBounds;
 }
 
+float MathUtils::Step(float Y, float X)
+{
+	return X >= Y ? 1 : 0;
+}
+
diff --git a/Engine/Plugins/Experimental/TextureGraph/Source/TextureGraphEngine/Helper/MathUtils.h b/Engine/Plugins/Experimental/TextureGraph/Source/TextureGraphEngine/Helper/MathUtils.h
--- a/Engine/Plugins/Experimental/TextureGraph/Source/TextureGraphEngine/Helper/MathUtils.h
+++ b/Engine/Plugins/Experimental/TextureGraph/Source/TextureGraphEngine/Helper/MathUtils.h
@@ -19,8 +19,9 @@ public:
 	static FVector						MaxFVector();
 	static FVector2f					MinFVector2();
 	static FVector2f					MaxFVector2();
-	static void						UpdateBounds(FBox& bounds, const FVector& point);
-	static void						EncapsulateBound(FBox& bounds, FBox& otherBounds);
+	static void							UpdateBounds(FBox& bounds, const FVector& point);
+	static void							EncapsulateBound(FBox& bounds, FBox& otherBounds);
 	static FVector						GetDirection(float yzAngle, float xAngle, int xSign = 1);
-	static FBox						GetCombinedBounds(TArray<FBox> inputBounds);
+	static FBox							GetCombinedBounds(TArray<FBox> inputBounds);
+	static float						Step(float Y, float X);
 };
