Fixed crash with TextureGraphExporter on Editor Exit
Fixed Output nodes connected to color pins not using Width/Height from OutputSettings.

diff --git a/Engine/Plugins/Experimental/TextureGraph/Source/TextureGraph/Private/TG_HelperFunctions.cpp b/Engine/Plugins/Experimental/TextureGraph/Source/TextureGraph/Private/TG_HelperFunctions.cpp
--- a/Engine/Plugins/Experimental/TextureGraph/Source/TextureGraph/Private/TG_HelperFunctions.cpp
+++ b/Engine/Plugins/Experimental/TextureGraph/Source/TextureGraph/Private/TG_HelperFunctions.cpp
@@ -43,9 +43,11 @@ TArray<BlobPtr> FTG_HelperFunctions::GetTexturedOutputs(const UTG_Node* Node, FT
 					BufferDescriptor DesiredDesc = OutputExpression->Output.EditTexture().GetBufferDescriptor();
 					if (DesiredDesc.Width == 0 || DesiredDesc.Height == 0)
 					{
+						FTG_OutputSettings& OutputSetting = OutputExpression->OutputSettings;
+				
 						DesiredDesc = T_FlatColorTexture::GetFlatColorDesc("Output");
-						DesiredDesc.Width = (uint32)EResolution::Resolution256;
-						DesiredDesc.Height = (uint32)EResolution::Resolution256;
+						DesiredDesc.Width = (uint32)OutputSetting.Width;
+						DesiredDesc.Height = (uint32)OutputSetting.Height;
 					}
 
 					// OutputExpression->Output.EditTexture() =
@@ -141,6 +143,8 @@ AsyncBool FTG_HelperFunctions::ExportAsync(UTextureGraph* InTextureGraph, FStrin
 					MapSettings.LODGroup = OutputSetting.LODGroup;
 					MapSettings.Compression = OutputSetting.Compression;
 					MapSettings.IsSRGB = OutputSetting.bSRGB;
+					MapSettings.Width = (int32)OutputSetting.Width;
+					MapSettings.Height = (int32)OutputSetting.Height;
 					TargetExportSettings.ExportPreset.push_back(std::pair<FName, FExportMapSettings>{ MapSettings.Name, MapSettings });
 				}
 				else
diff --git a/Engine/Plugins/Experimental/TextureGraph/Source/TextureGraphEditor/Private/STG_EditorViewport.cpp b/Engine/Plugins/Experimental/TextureGraph/Source/TextureGraphEditor/Private/STG_EditorViewport.cpp
--- a/Engine/Plugins/Experimental/TextureGraph/Source/TextureGraphEditor/Private/STG_EditorViewport.cpp
+++ b/Engine/Plugins/Experimental/TextureGraph/Source/TextureGraphEditor/Private/STG_EditorViewport.cpp
@@ -283,8 +283,13 @@ STG_EditorViewport::STG_EditorViewport()
 
 STG_EditorViewport::~STG_EditorViewport()
 {
-	CastChecked<UEditorEngine>(GEngine)->OnPreviewFeatureLevelChanged().Remove(PreviewFeatureLevelChangedHandle);
-
+	if (PreviewFeatureLevelChangedHandle.IsValid())
+	{
+		if (UEditorEngine* EditorEngine = Cast<UEditorEngine>(GEngine))
+		{
+			EditorEngine->OnPreviewFeatureLevelChanged().Remove(PreviewFeatureLevelChangedHandle);
+		}
+	}
 	FCoreUObjectDelegates::OnObjectPropertyChanged.RemoveAll(this);
 	if (EditorViewportClient.IsValid())
 	{
diff --git a/Engine/Plugins/Experimental/TextureGraph/Source/TextureGraphEditor/Private/TG_Exporter.cpp b/Engine/Plugins/Experimental/TextureGraph/Source/TextureGraphEditor/Private/TG_Exporter.cpp
--- a/Engine/Plugins/Experimental/TextureGraph/Source/TextureGraphEditor/Private/TG_Exporter.cpp
+++ b/Engine/Plugins/Experimental/TextureGraph/Source/TextureGraphEditor/Private/TG_Exporter.cpp
@@ -89,7 +89,7 @@ struct FTG_ExporterImpl: public FTickableGameObject, public FGCObject
 {
 public:
 	FTG_ExporterImpl();
-	~FTG_ExporterImpl();
+	virtual ~FTG_ExporterImpl() override;
 
 	/** Function registered with tab manager to create the Texture Graph Exporter */
 	TSharedRef<SDockTab> CreateTGExporterTab(const FSpawnTabArgs& Args);
@@ -103,13 +103,13 @@ public:
 	TSharedPtr<FTabManager::FLayout>				TGExporterLayout;
 
 	TObjectPtr<UTextureGraph>						TextureGraphPtr;
-	TSharedPtr<STG_NodePreviewWidget>				NodePreview;
+	TWeakPtr<STG_NodePreviewWidget>					NodePreviewPtr;
 	
 	TSharedPtr<class IDetailsView>					ParametersView;
 	TSharedPtr<class IDetailsView>					ExportSettingsView;
 	TSharedPtr<class IDetailsView>					PreviewSettingsView;
 // Tracking the active viewports in this editor.
-	TSharedPtr<class FEditorViewportTabContent>		ViewportTabContent;
+	TWeakPtr<class FEditorViewportTabContent>		ViewportTabContentPtr;
 
 	TObjectPtr<UTG_Parameters>						Parameters;
 	TObjectPtr<UTG_ExportSettings>					ExportSettings;
@@ -197,7 +197,8 @@ FTG_ExporterImpl::~FTG_ExporterImpl()
 		FGlobalTabmanager::Get()->UnregisterTabSpawner(FTG_EditorTabs::TextureExporterTabId);
 		TGExporterLayout = TSharedPtr<FTabManager::FLayout>();
 		TGExporterTabManager = TSharedPtr<FTabManager>();
-
+		
+		
 		Cleanup();
 		
 		ExportSettingsView.Reset();
@@ -205,12 +206,12 @@ FTG_ExporterImpl::~FTG_ExporterImpl()
 		ParametersView.Reset();
 		
 		// cleanup UI
-		if (Parameters)
+		if (Parameters->IsValidLowLevelFast())
 		{
 			Parameters->Parameters.Empty();
 			Parameters = nullptr;
 		}
-		if (ExportSettings)
+		if (ExportSettings->IsValidLowLevelFast())
 		{
 			ExportSettings->OutputExpressionsInfos.Empty();
 			ExportSettings = nullptr;
@@ -222,7 +223,7 @@ FTG_ExporterImpl::~FTG_ExporterImpl()
 
 void FTG_ExporterImpl::Cleanup()
 {
-	if (TextureGraphPtr)
+	if (TextureGraphPtr->IsValidLowLevelFast())
 	{
 		// cleanup events
 		TextureGraphPtr->GetSettings()->GetViewportSettings().OnViewportMaterialChangedEvent.RemoveAll(this);
@@ -233,7 +234,6 @@ void FTG_ExporterImpl::Cleanup()
 		TextureGraphPtr = nullptr;
 	}
 
-	
 	OutputNodesList.Empty();
 	SelectedNode = nullptr;
 }
@@ -274,10 +274,7 @@ TSharedRef<SDockTab> FTG_ExporterImpl::CreateTGExporterTab(const FSpawnTabArgs&
 		}
 		, TGExporterTabManagerWeak
 	));
-
-
-	NodePreview = SNew(STG_NodePreviewWidget);
-
+	
 	FPropertyEditorModule& PropertyEditorModule = FModuleManager::GetModuleChecked<FPropertyEditorModule>("PropertyEditor");
 	FDetailsViewArgs ParameterViewArgs;
 	ParameterViewArgs.NameAreaSettings = FDetailsViewArgs::HideNameArea;
@@ -378,7 +375,7 @@ TSharedRef<SDockTab> FTG_ExporterImpl::CreateTGExporterTab(const FSpawnTabArgs&
 		return false;
 	};
 	
-	bool bViewportIsOff = !ViewportTabContent.IsValid();
+	bool bViewportIsOff = !ViewportTabContentPtr.IsValid();
 	TSharedPtr<SDockTab> ViewportTab; 
 	// check here if 3d viewport is turned off, we need to turn it on temporarily to initialize our systems correctly
 	if (bViewportIsOff)
@@ -618,12 +615,14 @@ TSharedRef<SDockTab> FTG_ExporterImpl::SpawnTab_Viewport(const FSpawnTabArgs& Ar
 	};
 
 	// Create a new tab
-	ViewportTabContent = MakeShareable(new FEditorViewportTabContent());
+	TSharedRef<FEditorViewportTabContent> ViewportTabContent = MakeShared<FEditorViewportTabContent>();
 	//ViewportTabContent->OnViewportTabContentLayoutChanged().AddRaw(this, &FTG_ExporterImpl::OnEditorLayoutChanged);
 
 	const FString LayoutId = FString("TG_EditorViewport");
 	ViewportTabContent->Initialize(MakeViewportFunc, DockableTab, LayoutId);
 
+	ViewportTabContentPtr = ViewportTabContent;
+	
 	// This call must occur after the toolbar is initialized.
 	SetViewportPreviewMesh();
 	
@@ -648,6 +647,9 @@ TSharedRef<SDockTab> FTG_ExporterImpl::SpawnTab_NodePreview(const FSpawnTabArgs&
 {
 	check(Args.GetTabId() == FTG_EditorTabs::NodePreviewTabId);
 
+	TSharedRef<STG_NodePreviewWidget> NodePreview = SNew(STG_NodePreviewWidget);
+	NodePreviewPtr = NodePreview;
+	
 	return SNew(SDockTab)
 		[
 			SNew(SVerticalBox)
@@ -679,7 +681,7 @@ TSharedRef<SDockTab> FTG_ExporterImpl::SpawnTab_NodePreview(const FSpawnTabArgs&
 			
 			+ SVerticalBox::Slot()
 			[
-		 		NodePreview.ToSharedRef()
+				NodePreview
 			]
 		];
 }
@@ -698,10 +700,10 @@ void FTG_ExporterImpl::SetViewportPreviewMesh()
 }
 TSharedPtr<STG_EditorViewport> FTG_ExporterImpl::GetEditorViewport() const
 {
-	if (ViewportTabContent.IsValid())
+	if (ViewportTabContentPtr.IsValid())
 	{
 		// we can use static cast here b/c we know in this editor we will have a static mesh viewport 
-		return StaticCastSharedPtr<STG_EditorViewport>(ViewportTabContent->GetFirstViewport());
+		return StaticCastSharedPtr<STG_EditorViewport>(ViewportTabContentPtr.Pin()->GetFirstViewport());
 	}
 
 	return nullptr;
@@ -745,7 +747,10 @@ void FTG_ExporterImpl::OnOutputSelectionChanged(TSharedPtr<FName> SelectedItem,
 		});
 		
 		SelectedNode = TextureGraphPtr->Graph()->GetNode(SelectedNodeId);
-		NodePreview->SelectionChanged(SelectedNode);
+		if (NodePreviewPtr.IsValid())
+		{
+			NodePreviewPtr.Pin()->SelectionChanged(SelectedNode);
+		}
 	
 	}
 }
@@ -829,10 +834,10 @@ void FTG_ExporterImpl::OnGraphChanged(UTG_Graph* InGraph, UTG_Node* InNode, bool
 
 void FTG_ExporterImpl::OnRenderingDone(UMixInterface* TextureGraph, const FInvalidationDetails* Details)
 {
-	if (TextureGraph != nullptr && TextureGraph == TextureGraphPtr)
+	if (TextureGraph != nullptr && TextureGraph == TextureGraphPtr && NodePreviewPtr.IsValid())
 	{
 		// refresh node preview
-		NodePreview->Update();
+		NodePreviewPtr.Pin()->Update();
 	}
 }
 void FTG_ExporterImpl::OnViewportMaterialChanged()
diff --git a/Engine/Plugins/Experimental/TextureGraph/Source/TextureGraphEditor/Private/TextureGraphEditorModule.cpp b/Engine/Plugins/Experimental/TextureGraph/Source/TextureGraphEditor/Private/TextureGraphEditorModule.cpp
--- a/Engine/Plugins/Experimental/TextureGraph/Source/TextureGraphEditor/Private/TextureGraphEditorModule.cpp
+++ b/Engine/Plugins/Experimental/TextureGraph/Source/TextureGraphEditor/Private/TextureGraphEditorModule.cpp
@@ -86,6 +86,12 @@ void FTextureGraphEditorModule::ShutdownModule()
 	
 	ShutdownTextureGraphEngine();
 }
+
+void FTextureGraphEditorModule::PreUnloadCallback()
+{
+	TG_Exporter.Reset();
+}
+
 void FTextureGraphEditorModule::StartTextureGraphEngine()
 {
 	if (!TextureGraphEngine::GetInstance())
diff --git a/Engine/Plugins/Experimental/TextureGraph/Source/TextureGraphEditor/Public/TextureGraphEditorModule.h b/Engine/Plugins/Experimental/TextureGraph/Source/TextureGraphEditor/Public/TextureGraphEditorModule.h
--- a/Engine/Plugins/Experimental/TextureGraph/Source/TextureGraphEditor/Public/TextureGraphEditorModule.h
+++ b/Engine/Plugins/Experimental/TextureGraph/Source/TextureGraphEditor/Public/TextureGraphEditorModule.h
@@ -32,6 +32,7 @@ public:
 	/** IModuleInterface implementation */
 	virtual void					StartupModule() override;
 	virtual void					ShutdownModule() override;
+	virtual void					PreUnloadCallback() override;
 
 	
 	virtual void					StartTextureGraphEngine();
