Fixed Grayscale output appears Red in 3D Preview
Fixed MaterialMaskInfo pin appears on the texture graph node when selecting a texture graph and that texture graph contains the material node

#virtualized

diff --git a/Engine/Plugins/Experimental/TextureGraph/Source/TextureGraph/Public/Expressions/Utilities/TG_Expression_MaterialID.h b/Engine/Plugins/Experimental/TextureGraph/Source/TextureGraph/Public/Expressions/Utilities/TG_Expression_MaterialID.h
--- a/Engine/Plugins/Experimental/TextureGraph/Source/TextureGraph/Public/Expressions/Utilities/TG_Expression_MaterialID.h
+++ b/Engine/Plugins/Experimental/TextureGraph/Source/TextureGraph/Public/Expressions/Utilities/TG_Expression_MaterialID.h
@@ -22,7 +22,7 @@ public:
 	FTG_Texture MaterialIDMap;
 
 	// List of available colors extracted from the input texture. Selected colors will be shown as white in the output texture. 
-	UPROPERTY(EditAnywhere, EditFixedSize, Category = NoCategory, meta = (TGType = "TG_InputParam", NoResetToDefault, TGPinNotConnectable = "true"))
+	UPROPERTY(EditAnywhere, EditFixedSize, Category = NoCategory, meta = (TGType = "TG_Input", NoResetToDefault, TGPinNotConnectable = "true"))
 	TArray<FMaterialIDMaskInfo> MaterialIDMaskInfos;
 
 	UPROPERTY()
diff --git a/Engine/Plugins/Experimental/TextureGraph/Source/TextureGraphEngine/2D/TargetTextureSet.cpp b/Engine/Plugins/Experimental/TextureGraph/Source/TextureGraphEngine/2D/TargetTextureSet.cpp
--- a/Engine/Plugins/Experimental/TextureGraph/Source/TextureGraphEngine/2D/TargetTextureSet.cpp
+++ b/Engine/Plugins/Experimental/TextureGraph/Source/TextureGraphEngine/2D/TargetTextureSet.cpp
@@ -140,8 +140,9 @@ void TargetTextureSet::BindOnTextureUpdate(RenderMaterial_BPPtr InMaterial, FMat
 					// Linear textures are now being displayed in linear gamma.
 					// So turning off the conversion to match UE's convention.
 					// InMaterial->SetInt(TEXT("sRGB"), texture.get()->GetDescriptor().bIsSRGB ? 1 : 0);
-
-					InMaterial->SetInt(TEXT("IsGrayscale"), texture.get()->GetDescriptor().ItemsPerPoint == 1 ? 1 : 0);
+					FString GrayScaleParam = "Is" + BindInfo.Target + "GrayScale";
+					
+					InMaterial->SetInt(FName(GrayScaleParam), texture.get()->GetDescriptor().ItemsPerPoint == 1 ? 1 : 0);
 
 					texture->OnFinalise().then([=]()
 					{
diff --git a/Engine/Plugins/Experimental/TextureGraph/Source/TextureGraphEngine/Model/Mix/ViewportSettings.cpp b/Engine/Plugins/Experimental/TextureGraph/Source/TextureGraphEngine/Model/Mix/ViewportSettings.cpp
--- a/Engine/Plugins/Experimental/TextureGraph/Source/TextureGraphEngine/Model/Mix/ViewportSettings.cpp
+++ b/Engine/Plugins/Experimental/TextureGraph/Source/TextureGraphEngine/Model/Mix/ViewportSettings.cpp
@@ -112,6 +112,15 @@ void FViewportSettings::OnMaterialUpdate()
 			}
 		}
 
+		// In case of multiple inputs, sort the inputs in ascending order.
+		if(MaterialMappingInfos.Num() > 1)
+		{
+			MaterialMappingInfos.Sort([](const FMaterialMappingInfo& A, const FMaterialMappingInfo& B)
+			{
+				return A.MaterialInput.LexicalLess(B.MaterialInput);
+			});
+		}
+
 		OnViewportMaterialChangedEvent.Broadcast();
 	}
 }
