Levels out implementation

diff --git a/Engine/Plugins/Experimental/TextureGraph/Shaders/Expressions/Expression_Levels.usf b/Engine/Plugins/Experimental/TextureGraph/Shaders/Expressions/Expression_Levels.usf
--- a/Engine/Plugins/Experimental/TextureGraph/Shaders/Expressions/Expression_Levels.usf
+++ b/Engine/Plugins/Experimental/TextureGraph/Shaders/Expressions/Expression_Levels.usf
@@ -12,11 +12,19 @@
 #define AUTO_LEVELS 0
 #endif
 
+#ifndef OUT_LEVELS
+#define OUT_LEVELS 0
+#endif
+
 Texture2D SourceTexture;
 float LowValue;
 float HighValue;
 float MidValue;
 
+float OutLow;
+float OutHigh;
+float OutputRange;
+
 float DoAutoLevel;
 float MidPercentage;
 
@@ -38,8 +46,21 @@ float LevelsControlInputCurve(float Value, float LowInput, float HighInput, floa
 	MidRanged = max(0.001, min(0.9999, MidRanged));
 
 	float InvGammaValue = log(MidRanged) / log(0.5);
+	float NewValue = pow(ValueRanged, InvGammaValue);
+
+#if OUT_LEVELS
+	float4 HistoMin			= Histogram.Load(int3(2, 1, 0));
+	float4 HistoMax			= Histogram.Load(int3(3, 1, 0));
+		
+	LowInput = 	HistoMin.w;
+	HighInput = HistoMax.w;
+	
+	float InputRange = clamp(HighInput - LowInput, 0, 1);
+	
+	NewValue = OutLow + ((NewValue - LowInput) / InputRange) * OutputRange;
+#endif /// OUT_LEVELS
 
-	return pow(ValueRanged, InvGammaValue);
+	return NewValue;
 }
 
 float4 FSH_Levels(float2 uv : TEXCOORD0) : SV_Target0
@@ -69,7 +90,7 @@ float4 FSH_Levels(float2 uv : TEXCOORD0) : SV_Target0
 		
 	LowInput = 	HistoMin.w;
 	HighInput = HistoMax.w;
-#endif
+#endif /// AUTO_LEVELS
 
 	float3 LeveledColor = LevelsControlInputCurve(SourceValue, LowInput, HighInput, MidValue);
 	
diff --git a/Engine/Plugins/Experimental/TextureGraph/Source/TextureGraph/Private/Expressions/FIlter/TG_Expression_Levels.cpp b/Engine/Plugins/Experimental/TextureGraph/Source/TextureGraph/Private/Expressions/FIlter/TG_Expression_Levels.cpp
--- a/Engine/Plugins/Experimental/TextureGraph/Source/TextureGraph/Private/Expressions/FIlter/TG_Expression_Levels.cpp
+++ b/Engine/Plugins/Experimental/TextureGraph/Source/TextureGraph/Private/Expressions/FIlter/TG_Expression_Levels.cpp
@@ -109,16 +109,25 @@ void UTG_Expression_Levels::PostEditChangeProperty(FPropertyChangedEvent& Proper
 		SetLowValue(LowValue);
 	}
 	// High
-	if (PropertyChangedEvent.GetPropertyName() == GET_MEMBER_NAME_CHECKED(UTG_Expression_Levels, HighValue))
+	else if (PropertyChangedEvent.GetPropertyName() == GET_MEMBER_NAME_CHECKED(UTG_Expression_Levels, HighValue))
 	{
 		SetHighValue(HighValue);
 	}
 	// Mid
-	if (PropertyChangedEvent.GetPropertyName() == GET_MEMBER_NAME_CHECKED(UTG_Expression_Levels, MidValue))
+	else if (PropertyChangedEvent.GetPropertyName() == GET_MEMBER_NAME_CHECKED(UTG_Expression_Levels, MidValue))
 	{
 		SetMidValue(MidValue);
 	}
 
+	else if (PropertyChangedEvent.GetPropertyName() == GET_MEMBER_NAME_CHECKED(UTG_Expression_Levels, OutLowValue))
+	{
+		OutLowValue = FMath::Clamp(OutLowValue, 0, OutHighValue);
+	}
+	else if (PropertyChangedEvent.GetPropertyName() == GET_MEMBER_NAME_CHECKED(UTG_Expression_Levels, OutHighValue))
+	{
+		OutHighValue = FMath::Clamp(OutHighValue, OutLowValue, 1.0f);
+	}
+
 	Super::PostEditChangeProperty(PropertyChangedEvent);
 }
 
@@ -205,7 +214,7 @@ void UTG_Expression_Levels::Evaluate(FTG_EvaluationContext* InContext)
 	switch (LevelsExpressionType)
 	{
 	case ELevelsExpressionType::LowMidHigh:
-		LevelsControl->InitFromLowMidHigh(LowValue, MidValue, HighValue);
+		LevelsControl->InitFromLowMidHigh(LowValue, MidValue, HighValue, OutLowValue, OutHighValue);
 		break;
 	case ELevelsExpressionType::AutoLowHigh:
 		LevelsControl->InitFromAutoLevels(MidAutoLevels);
@@ -225,6 +234,18 @@ void UTG_Expression_HistogramScan::Evaluate(FTG_EvaluationContext* InContext)
 		Output = FTG_Texture::GetBlack();
 		return;
 	}
+
+	if (Position < 0.00001f)
+	{
+		Output = FTG_Texture::GetBlack();
+		return;
+	}
+	else if (Position >= 0.99999f)
+	{
+		Output = FTG_Texture::GetWhite();
+		return;
+	}
+
 	LevelsControl = MakeShared<FLevels>();
 	LevelsControl->InitFromPositionContrast(Position, Contrast);
 
diff --git a/Engine/Plugins/Experimental/TextureGraph/Source/TextureGraph/Public/Expressions/Filter/TG_Expression_Levels.h b/Engine/Plugins/Experimental/TextureGraph/Source/TextureGraph/Public/Expressions/Filter/TG_Expression_Levels.h
--- a/Engine/Plugins/Experimental/TextureGraph/Source/TextureGraph/Public/Expressions/Filter/TG_Expression_Levels.h
+++ b/Engine/Plugins/Experimental/TextureGraph/Source/TextureGraph/Public/Expressions/Filter/TG_Expression_Levels.h
@@ -116,6 +116,14 @@ public:
 	float								HighValue = 1;
 	void SetHighValue(float InValue);
 
+	// The black point of the output. Moving this will remap the dark point to this value. Default is 0
+	UPROPERTY(EditAnywhere, Category = NoCategory, meta = (TGType = "TG_Setting", PinDisplayName = "Out Low", UIMin = "0", ClampMin = "0", UIMax = "1", ClampMax = "1", EditConditionHides, MD_ScalarEditor))
+	float								OutLowValue = 0;
+
+	// The white point of the output. Moving this will remap the white point to this value. Default is 1
+	UPROPERTY(EditAnywhere, Category = NoCategory, meta = (TGType = "TG_Setting", PinDisplayName = "Out High", UIMin = "0", ClampMin = "0", UIMax = "1", ClampMax = "1", EditConditionHides, MD_ScalarEditor))
+	float								OutHighValue = 1;
+
 
 	// The High value of the Levels adjustment, any pixel above that value is set to white. Default is 1.
 	UPROPERTY(EditAnywhere, Category = NoCategory, meta = (TGType = "TG_Setting", PinDisplayName = "Mid Auto Levels", UIMin = "0", ClampMin = "0", UIMax = "1", ClampMax = "1", EditConditionHides, MD_ScalarEditor))
@@ -148,9 +156,11 @@ public:
 	UPROPERTY(EditAnywhere, Category = NoCategory, meta = (TGType = "TG_Input", PinDisplayName = "", MD_HistogramLuminance))
 	FTG_Texture							Input;
 
+	// Drives the position of the histogram of the input image
 	UPROPERTY(EditAnywhere, Category = NoCategory, meta = (TGType = "TG_Setting", UIMin = "0", ClampMin = "0", UIMax = "1", ClampMax = "1", MD_ScalarEditor))
 	float 								Position = 0.5f;
 
+	// Drives the contrast of the histogram of the input image
 	UPROPERTY(EditAnywhere, Category = NoCategory, meta = (TGType = "TG_Setting", UIMin = "0", ClampMin = "0", UIMax = "1", ClampMax = "1", MD_ScalarEditor))
 	float 								Contrast = 0.5f; 
 
@@ -158,6 +168,6 @@ public:
 	UPROPERTY(EditAnywhere, Category = NoCategory, meta = (TGType = "TG_Output", PinDisplayName = ""))
 	FTG_Texture							Output;
 
-	virtual FText						GetTooltipText() const override { return FText::FromString(TEXT("Remaps shadows and highlights of the input. Any values less or equal to Low are mapped to black, any values, greater or equal to High are mapped to white, and any values inbetween have Gamma applied as an exponent.")); } 
+	virtual FText						GetTooltipText() const override { return FText::FromString(TEXT("Lets you drive the contrast and position of the histogram. Input must be a grayscale image.")); } 
 
 };
\ No newline at end of file
diff --git a/Engine/Plugins/Experimental/TextureGraph/Source/TextureGraphEngine/Transform/Expressions/T_Levels.cpp b/Engine/Plugins/Experimental/TextureGraph/Source/TextureGraphEngine/Transform/Expressions/T_Levels.cpp
--- a/Engine/Plugins/Experimental/TextureGraph/Source/TextureGraphEngine/Transform/Expressions/T_Levels.cpp
+++ b/Engine/Plugins/Experimental/TextureGraph/Source/TextureGraphEngine/Transform/Expressions/T_Levels.cpp
@@ -80,12 +80,15 @@ bool FLevels::SetMidFromMidExponent(float InExponent)
 	return false;
 }
 
-void FLevels::InitFromLowMidHigh(float LowValue, float MidValue, float HighValue)
+void FLevels::InitFromLowMidHigh(float LowValue, float MidValue, float HighValue, float OutLowValue, float OutHighValue)
 {
 	Low = std::max(0.0f, LowValue);
 	Mid = MidValue;
 	High = std::min(1.0f, HighValue);
 
+	OutLow = OutLowValue;
+	OutHigh = OutHighValue;
+
 	IsAutoLevels = false;
 }
 
@@ -270,35 +273,40 @@ TiledBlobPtr	T_Levels::Create(MixUpdateCyclePtr Cycle, BufferDescriptor DesiredO
 {
 	bool NeedsConvertToGrayscale = (Source->GetDescriptor().ItemsPerPoint > 1) ;
 	
-	FSH_Levels::FPermutationDomain PermutationVector;
-	PermutationVector.Set<FSH_Levels::FConvertToGrayscale>(NeedsConvertToGrayscale);
-	//PermutationVector.Set<FSH_Levels::FIsAutoLevels>(InLevels->IsAutoLevels);
-	PermutationVector.Set<FSH_Levels::FIsAutoLevels>(false);
-
-
 	if (!Source)
 	{
 		return TextureHelper::GBlack;
 	}
 
-	bool NeedHistogramRaw = false;
+	bool bNeedHistogramRaw = false;
 	TiledBlobPtr Histogram = TextureHelper::GBlack;
-	if (InLevels->IsAutoLevels)
+	bool bIsModifyingOut = InLevels->OutLow > 0 || InLevels->OutHigh < 1;
+
+	if (InLevels->IsAutoLevels || bIsModifyingOut)
 	{
-		NeedHistogramRaw = true;
+		bNeedHistogramRaw = true;
 		Histogram = T_TextureHistogram::Create(Cycle, Source, TargetId);
 	}
 
 	if (!Histogram)
 	{
 		Histogram = TextureHelper::GBlack;
-		NeedHistogramRaw = false;
+		bNeedHistogramRaw = false;
 	}
 	
-	std::shared_ptr<FxMaterial_Normal<VSH_Simple, FSH_Levels>> Mat = std::make_shared<FxMaterial_Normal<VSH_Simple, FSH_Levels>>();
+	FSH_Levels::FPermutationDomain PermutationVector;
+	PermutationVector.Set<FSH_Levels::FConvertToGrayscale>(NeedsConvertToGrayscale);
+	PermutationVector.Set<FSH_Levels::FIsAutoLevels>(InLevels->IsAutoLevels);
+	PermutationVector.Set<FSH_Levels::FIsOutLevels>(bIsModifyingOut);
+
+	std::shared_ptr<FxMaterial_Normal<VSH_Simple, FSH_Levels>> Mat = 
+		std::make_shared<FxMaterial_Normal<VSH_Simple, FSH_Levels>>(typename VSH_Simple::FPermutationDomain(), PermutationVector);
+
 	const RenderMaterial_FXPtr RenderMaterial = std::make_shared<RenderMaterial_FX_Levels>(TEXT("T_Levels"), Mat, InLevels);
 	check(RenderMaterial);
 
+	float OutputRange = FMath::Clamp(InLevels->OutHigh - InLevels->OutLow, 0.0f, 1.0f);
+
 	JobUPtr RenderJob = std::make_unique<Job_Levels>(Cycle->GetMix(), TargetId, std::static_pointer_cast<BlobTransform>(RenderMaterial), Histogram, InLevels);
 	RenderJob
 		->AddArg(ARG_BLOB(Source, "SourceTexture"))
@@ -307,6 +315,9 @@ TiledBlobPtr	T_Levels::Create(MixUpdateCyclePtr Cycle, BufferDescriptor DesiredO
 		->AddArg(ARG_FLOAT(InLevels->Mid, "MidValue"))
 		->AddArg(ARG_FLOAT(float(InLevels->IsAutoLevels), "DoAutoLevel"))
 		->AddArg(ARG_FLOAT(InLevels->MidPercentage, "MidPercentage"))
+		->AddArg(ARG_FLOAT(InLevels->OutLow, "OutLow"))
+		->AddArg(ARG_FLOAT(InLevels->OutHigh, "OutHigh"))
+		->AddArg(ARG_FLOAT(OutputRange, "OutputRange"))
 		->AddArg(std::make_shared<JobArg_Blob>(JobArg_Blob(Histogram, "Histogram").WithNotHandleTiles()))
 		;
 
diff --git a/Engine/Plugins/Experimental/TextureGraph/Source/TextureGraphEngine/Transform/Expressions/T_Levels.h b/Engine/Plugins/Experimental/TextureGraph/Source/TextureGraphEngine/Transform/Expressions/T_Levels.h
--- a/Engine/Plugins/Experimental/TextureGraph/Source/TextureGraphEngine/Transform/Expressions/T_Levels.h
+++ b/Engine/Plugins/Experimental/TextureGraph/Source/TextureGraphEngine/Transform/Expressions/T_Levels.h
@@ -25,12 +25,16 @@ public:
 		SHADER_PARAMETER(float, MidValue)
 		SHADER_PARAMETER(float, DoAutoLevel)
 		SHADER_PARAMETER(float, MidPercentage)
+		SHADER_PARAMETER(float, OutLow)
+		SHADER_PARAMETER(float, OutHigh)
+		SHADER_PARAMETER(float, OutputRange)
 		SHADER_PARAMETER_TEXTURE(Texture2D, Histogram)
 	END_SHADER_PARAMETER_STRUCT()
 
 	class FConvertToGrayscale : SHADER_PERMUTATION_BOOL("CONVERT_TO_GRAYSCALE");
 	class FIsAutoLevels : SHADER_PERMUTATION_BOOL("AUTO_LEVELS");
-	using FPermutationDomain = TShaderPermutationDomain<FConvertToGrayscale, FIsAutoLevels>;
+	class FIsOutLevels : SHADER_PERMUTATION_BOOL("OUT_LEVELS");
+	using FPermutationDomain = TShaderPermutationDomain<FConvertToGrayscale, FIsAutoLevels, FIsOutLevels>;
 
 	static bool						ShouldCompilePermutation(const FGlobalShaderPermutationParameters& Parameters) 
 	{
@@ -62,6 +66,12 @@ struct TEXTUREGRAPHENGINE_API FLevels
 	bool 								IsAutoLevels = false;
 	float 								MidPercentage = 0.5f;
 
+	// The black point of the output. Moving this will remap the dark point to this value. Default is 0
+	float								OutLow = 0;
+
+	// The white point of the output. Moving this will remap the white point to this value. Default is 1
+	float								OutHigh = 1;
+
 	bool SetLow(float InValue);
 	bool SetMid(float InValue);
 	bool SetHigh(float InValue);
@@ -75,7 +85,7 @@ struct TEXTUREGRAPHENGINE_API FLevels
 	float EvalMidExponent() const;
 	bool SetMidFromMidExponent(float InExponent);
 
-	void InitFromLowMidHigh(float LowValue, float MidValue, float HighValue);
+	void InitFromLowMidHigh(float LowValue, float MidValue, float HighValue, float OutLow, float OutHigh);
 	void InitFromAutoLevels(float MidPercentage);
 	void InitFromPositionContrast(float MidValue, float Contrast);
 
