Some minor bug fixes.

diff --git a/Engine/Plugins/Experimental/TextureGraph/Source/TextureGraph/Public/Expressions/Procedural/TG_Expression_Gradient.h b/Engine/Plugins/Experimental/TextureGraph/Source/TextureGraph/Public/Expressions/Procedural/TG_Expression_Gradient.h
--- a/Engine/Plugins/Experimental/TextureGraph/Source/TextureGraph/Public/Expressions/Procedural/TG_Expression_Gradient.h
+++ b/Engine/Plugins/Experimental/TextureGraph/Source/TextureGraph/Public/Expressions/Procedural/TG_Expression_Gradient.h
@@ -35,22 +35,27 @@ public:
 	UPROPERTY(EditAnywhere, BlueprintReadWrite, Category = NoCategory, meta = (TGType = "TG_Setting", EditConditionHides))
 	EGradientInterpolation				Interpolation = EGradientInterpolation::GTI_Linear;
 
-	// Type of interpolation to use
+	// Rotation of the gradient
 	UPROPERTY(EditAnywhere, BlueprintReadWrite, Category = NoCategory, meta = (TGType = "TG_Setting", EditConditionHides))
 	EGradientRotation					Rotation = EGradientRotation::GTR_0;
 
+	// Rotation of the gradient
 	UPROPERTY(EditAnywhere, BlueprintReadWrite, Category = NoCategory, meta = (TGType = "TG_Setting", DisplayName = "Rotation", EditConditionHides, PinDisplayName = "Rotation"))
 	EGradientRotationLimited			RotationLimited = EGradientRotationLimited::GTRL_0;
 
+	// The center of the radial gradient
 	UPROPERTY(EditAnywhere, BlueprintReadWrite, Category = NoCategory, meta = (TGType = "TG_Input", EditConditionHides, UIMin = "0", ClampMin = "0", UIMax = "1", ClampMax = "1"))
 	FVector2f							Center = { 0.5f, 0.5f };
 
+	// The center of the radial gradient
 	UPROPERTY(EditAnywhere, BlueprintReadWrite, Category = NoCategory, meta = (TGType = "TG_Input", EditConditionHides, UIMin = "0.001", ClampMin = "0.001", UIMax = "1", ClampMax = "1"))
 	float								Radius = 0.25f;
 
+	// First point of the line for axial gradients
 	UPROPERTY(EditAnywhere, BlueprintReadWrite, Category = NoCategory, meta = (TGType = "TG_Input", EditConditionHides, UIMin = "0", ClampMin = "0", UIMax = "1", ClampMax = "1"))
 	FVector2f							Point1 = { 0.25f, 0.25f };
 
+	// Second point of the line for axial gradients
 	UPROPERTY(EditAnywhere, BlueprintReadWrite, Category = NoCategory, meta = (TGType = "TG_Input", EditConditionHides, UIMin = "0", ClampMin = "0", UIMax = "1", ClampMax = "1"))
 	FVector2f							Point2 = { 0.75f, 0.75f };
 
diff --git a/Engine/Plugins/Experimental/TextureGraph/Source/TextureGraphEngine/Transform/Expressions/T_Maths_OneInput.cpp b/Engine/Plugins/Experimental/TextureGraph/Source/TextureGraphEngine/Transform/Expressions/T_Maths_OneInput.cpp
--- a/Engine/Plugins/Experimental/TextureGraph/Source/TextureGraphEngine/Transform/Expressions/T_Maths_OneInput.cpp
+++ b/Engine/Plugins/Experimental/TextureGraph/Source/TextureGraphEngine/Transform/Expressions/T_Maths_OneInput.cpp
@@ -79,7 +79,7 @@ TiledBlobPtr T_Maths_OneInput::CreateTrigonometry(MixUpdateCyclePtr Cycle, Buffe
 #define IMPL_SINGLE_ARG_SHADER(Name) TiledBlobPtr T_Maths_OneInput::Create##Name(MixUpdateCyclePtr Cycle, BufferDescriptor DesiredOutputDesc, int32 TargetId, TiledBlobPtr Input) \
 { \
 	FString TransformName = FString(TEXT("T_")) + FString(#Name); \
-	return CreateGenericOneInput<FSH_Sin>(Cycle, DesiredOutputDesc, TargetId, Input, TransformName); \
+	return CreateGenericOneInput<FSH_##Name>(Cycle, DesiredOutputDesc, TargetId, Input, TransformName); \
 }
 
 IMPL_SINGLE_ARG_SHADER(Abs);
