FIxing bug triggered in PostLoad
also introducing Setters for OutLow and OutHigh

diff --git a/Engine/Plugins/Experimental/TextureGraph/Source/TextureGraph/Private/Expressions/FIlter/TG_Expression_Levels.cpp b/Engine/Plugins/Experimental/TextureGraph/Source/TextureGraph/Private/Expressions/FIlter/TG_Expression_Levels.cpp
--- a/Engine/Plugins/Experimental/TextureGraph/Source/TextureGraph/Private/Expressions/FIlter/TG_Expression_Levels.cpp
+++ b/Engine/Plugins/Experimental/TextureGraph/Source/TextureGraph/Private/Expressions/FIlter/TG_Expression_Levels.cpp
@@ -92,6 +92,8 @@ bool FTG_LevelsSettings::SetMidFromMidExponent(float InExponent)
 
 void UTG_Expression_Levels::PostLoad()
 {
+	Super::PostLoad();
+
 	// Restore LevelsSettings inner struct from saved values
 	Levels.Low = (LowValue);
 	Levels.High = (HighValue);
@@ -118,14 +120,15 @@ void UTG_Expression_Levels::PostEditChangeProperty(FPropertyChangedEvent& Proper
 	{
 		SetMidValue(MidValue);
 	}
-
+	// Out Low
 	else if (PropertyChangedEvent.GetPropertyName() == GET_MEMBER_NAME_CHECKED(UTG_Expression_Levels, OutLowValue))
 	{
-		OutLowValue = FMath::Clamp(OutLowValue, 0, OutHighValue);
+		SetOutLowValue(OutLowValue);
 	}
+	// Out High
 	else if (PropertyChangedEvent.GetPropertyName() == GET_MEMBER_NAME_CHECKED(UTG_Expression_Levels, OutHighValue))
 	{
-		OutHighValue = FMath::Clamp(OutHighValue, OutLowValue, 1.0f);
+		SetOutHighValue(OutHighValue);
 	}
 
 	Super::PostEditChangeProperty(PropertyChangedEvent);
@@ -199,6 +202,17 @@ void UTG_Expression_Levels::SetHighValue(float InValue)
 	}
 }
 
+
+void UTG_Expression_Levels::SetOutLowValue(float InValue)
+{
+	OutLowValue = FMath::Clamp(InValue, 0, OutHighValue);
+}
+
+void UTG_Expression_Levels::SetOutHighValue(float InValue)
+{
+	OutHighValue = FMath::Clamp(OutHighValue, OutLowValue, 1.0f);
+}
+
 void UTG_Expression_Levels::Evaluate(FTG_EvaluationContext* InContext)
 {
 	Super::Evaluate(InContext);
diff --git a/Engine/Plugins/Experimental/TextureGraph/Source/TextureGraph/Public/Expressions/Filter/TG_Expression_Levels.h b/Engine/Plugins/Experimental/TextureGraph/Source/TextureGraph/Public/Expressions/Filter/TG_Expression_Levels.h
--- a/Engine/Plugins/Experimental/TextureGraph/Source/TextureGraph/Public/Expressions/Filter/TG_Expression_Levels.h
+++ b/Engine/Plugins/Experimental/TextureGraph/Source/TextureGraph/Public/Expressions/Filter/TG_Expression_Levels.h
@@ -117,13 +117,14 @@ public:
 	void SetHighValue(float InValue);
 
 	// The black point of the output. Moving this will remap the dark point to this value. Default is 0
-	UPROPERTY(EditAnywhere, Category = NoCategory, meta = (TGType = "TG_Setting", PinDisplayName = "Out Low", UIMin = "0", ClampMin = "0", UIMax = "1", ClampMax = "1", EditConditionHides, MD_ScalarEditor))
+	UPROPERTY(EditAnywhere, Setter, Category = NoCategory, meta = (TGType = "TG_Setting", PinDisplayName = "Out Low", UIMin = "0", ClampMin = "0", UIMax = "1", ClampMax = "1", EditConditionHides, MD_ScalarEditor))
 	float								OutLowValue = 0;
+	void SetOutLowValue(float InValue);
 
 	// The white point of the output. Moving this will remap the white point to this value. Default is 1
-	UPROPERTY(EditAnywhere, Category = NoCategory, meta = (TGType = "TG_Setting", PinDisplayName = "Out High", UIMin = "0", ClampMin = "0", UIMax = "1", ClampMax = "1", EditConditionHides, MD_ScalarEditor))
+	UPROPERTY(EditAnywhere, Setter, Category = NoCategory, meta = (TGType = "TG_Setting", PinDisplayName = "Out High", UIMin = "0", ClampMin = "0", UIMax = "1", ClampMax = "1", EditConditionHides, MD_ScalarEditor))
 	float								OutHighValue = 1;
-
+	void SetOutHighValue(float InValue);
 
 	// The High value of the Levels adjustment, any pixel above that value is set to white. Default is 1.
 	UPROPERTY(EditAnywhere, Category = NoCategory, meta = (TGType = "TG_Setting", PinDisplayName = "Mid Auto Levels", UIMin = "0", ClampMin = "0", UIMax = "1", ClampMax = "1", EditConditionHides, MD_ScalarEditor))
