Prevent cook crash in the absence of an editor

diff --git a/Engine/Plugins/Experimental/TextureGraph/Source/TextureGraphEngine/2D/TextureHelper.cpp b/Engine/Plugins/Experimental/TextureGraph/Source/TextureGraphEngine/2D/TextureHelper.cpp
--- a/Engine/Plugins/Experimental/TextureGraph/Source/TextureGraphEngine/2D/TextureHelper.cpp
+++ b/Engine/Plugins/Experimental/TextureGraph/Source/TextureGraphEngine/2D/TextureHelper.cpp
@@ -106,8 +106,11 @@ cti::continuable<bool> TextureHelper::InitSolidTexture(TiledBlobPtr* BlobObj, FL
 void TextureHelper::InitStockTextures()
 {
 #if WITH_EDITOR
+	// No need to run during command-let execution
 	if (!GEditor)
+	{
 		return;
+	}
 
 	check(IsInGameThread());
 
@@ -115,7 +118,9 @@ void TextureHelper::InitStockTextures()
 
 	/// If we've already go the stuff, then we don't need to worry about it
 	if (GBlack)
+	{
 		return;
+	}
 
 	TextureType Type = TextureHelper::TextureContentToTextureType(TextureContent::Albedo);
 	TexDescriptor AlbedoDesc(TextureSet::GDesc[(int32)Type]);
diff --git a/Engine/Plugins/Experimental/TextureGraph/Source/TextureGraphEngine/TextureGraphEngine.cpp b/Engine/Plugins/Experimental/TextureGraph/Source/TextureGraphEngine/TextureGraphEngine.cpp
--- a/Engine/Plugins/Experimental/TextureGraph/Source/TextureGraphEngine/TextureGraphEngine.cpp
+++ b/Engine/Plugins/Experimental/TextureGraph/Source/TextureGraphEngine/TextureGraphEngine.cpp
@@ -294,9 +294,17 @@ void TextureGraphEngine::FirstRunInit()
 	
 	MaterialMgrObj = TStrongObjectPtr<UMaterialManager>(UMaterialManager::CreateNew<UMaterialManager>());
 	
+#if WITH_EDITOR
+	// No need to run during command-let execution
+	if (!GEditor || !FSlateApplication::IsInitialized() || !FApp::CanEverRender())
+	{
+		return;
+	}
+
 	// If the engine is not set to run, then we don't do anything here. Not even initialise it (to void any 
 	// initialisation issues that we might get with people who don't really use the TextureGraph)
 	TextureHelper::InitStockTextures();
+#endif /// WITH_EDITOR
 }
 
 
diff --git a/Engine/Plugins/Experimental/TextureGraph/Source/TextureGraphEngine/TextureGraphEngine.h b/Engine/Plugins/Experimental/TextureGraph/Source/TextureGraphEngine/TextureGraphEngine.h
--- a/Engine/Plugins/Experimental/TextureGraph/Source/TextureGraphEngine/TextureGraphEngine.h
+++ b/Engine/Plugins/Experimental/TextureGraph/Source/TextureGraphEngine/TextureGraphEngine.h
@@ -148,20 +148,22 @@ public:
 	/// Inline functions
 	//////////////////////////////////////////////////////////////////////////
 	/// The if check is temporary over here since no one is initialising this right now
-	FORCEINLINE static TextureGraphEngine*						GetInstance() { return GInstance; }
+	FORCEINLINE static TextureGraphEngine*				GetInstance() { return GInstance; }
 	FORCEINLINE static MixManager*						GetMixManager() { return GInstance ? GInstance->MixMgrObj.get() : nullptr; }
 	FORCEINLINE static UMaterialManager*				GetMaterialManager() { return GInstance ? GInstance->MaterialMgrObj.Get() : nullptr; }
 	FORCEINLINE static Scheduler*						GetScheduler() { return GInstance ? GInstance->SchedulerObj.get() : nullptr; }
 	FORCEINLINE static DeviceManager*					GetDeviceManager() { return GInstance ? GInstance->DeviceManagerObj.get() : nullptr; }
 	FORCEINLINE static Blobber*							GetBlobber() { return GInstance ? GInstance->BlobberObj.get() : nullptr; }
 	FORCEINLINE static int64							GetFrameId() { return GInstance ? GInstance->FrameId : 0; }
-	FORCEINLINE static TextureGraphEditor::RenderDocManager*	GetRenderDocManager() { return GInstance ? GInstance->RenderDocMgrObj.get() : nullptr; }
+	FORCEINLINE static TextureGraphEditor::RenderDocManager* GetRenderDocManager() { return GInstance ? GInstance->RenderDocMgrObj.get() : nullptr; }
 
 	FORCEINLINE static bool								IsTestMode() { return GInstance ? GInstance->bIsTestMode : false; }
 
 	FORCEINLINE static bool								IsDestroying() { return bGIsEngineDestroying; }
 	FORCEINLINE static EngineObserverSourcePtr			GetObserverSource() { return GObserverSource; }
-	FORCEINLINE static FTextureGraphErrorReporter*				GetErrorReporter(UMixInterface* MixKey) { return GInstance ? GInstance->ErrorReporters.Find(MixKey)->get() : nullptr; }
+	FORCEINLINE static FTextureGraphErrorReporter*		GetErrorReporter(UMixInterface* MixKey) { return GInstance ? GInstance->ErrorReporters.Find(MixKey)->get() : nullptr; }
+
+	FORCEINLINE static bool								IsRunning() { return GInstance ? GInstance->bRunEngine.load() : false; }
 };
 
 
