Fixed some minor issues and crashes

#virtualized

diff --git a/Engine/Plugins/Experimental/TextureGraph/Source/TextureGraphEditor/Private/EdGraph/STG_NodeThumbnail.cpp b/Engine/Plugins/Experimental/TextureGraph/Source/TextureGraphEditor/Private/EdGraph/STG_NodeThumbnail.cpp
--- a/Engine/Plugins/Experimental/TextureGraph/Source/TextureGraphEditor/Private/EdGraph/STG_NodeThumbnail.cpp
+++ b/Engine/Plugins/Experimental/TextureGraph/Source/TextureGraphEditor/Private/EdGraph/STG_NodeThumbnail.cpp
@@ -92,10 +92,21 @@ void STG_NodeThumbnail::UpdateParams(TiledBlobPtr InBlob)
 		UTextureRenderTarget2D* TextureRT2D = Cast<UTextureRenderTarget2D>(BlobTexture);
 	
 		float ShowChecker = 1.0;
+		
+		FLinearColor NormalizedDimensions(0.9f, 0.9f, 0.0f);
+		
 		if (BlobTexture)
 		{
 			ShowChecker = 0.0;
 			BrushMaterial->SetTextureParameterValue("ThumbTex", BlobTexture);
+
+			int32 Width = BlobTexture->GetSurfaceWidth();
+			int32 Height = BlobTexture->GetSurfaceHeight();
+
+			float MaxDimension = FMath::Max(Width, Height);
+
+			NormalizedDimensions.R *= Width / MaxDimension;
+			NormalizedDimensions.G *= Height / MaxDimension;
 		}
 
 		float SingleChannel = 0.0;
@@ -112,6 +123,8 @@ void STG_NodeThumbnail::UpdateParams(TiledBlobPtr InBlob)
 
 		BrushMaterial->SetScalarParameterValue("ShowChecker", ShowChecker);
 		BrushMaterial->SetScalarParameterValue("SingleChannel", SingleChannel);
+		
+		BrushMaterial->SetVectorParameterValue("NormalizedDimensions", NormalizedDimensions);
 		Brush->SetResourceObject(BrushMaterial);
 	}
 }
diff --git a/Engine/Plugins/Experimental/TextureGraph/Source/TextureGraphEditor/Private/EdGraph/TG_EdGraphNode.cpp b/Engine/Plugins/Experimental/TextureGraph/Source/TextureGraphEditor/Private/EdGraph/TG_EdGraphNode.cpp
--- a/Engine/Plugins/Experimental/TextureGraph/Source/TextureGraphEditor/Private/EdGraph/TG_EdGraphNode.cpp
+++ b/Engine/Plugins/Experimental/TextureGraph/Source/TextureGraphEditor/Private/EdGraph/TG_EdGraphNode.cpp
@@ -483,7 +483,7 @@ void UTG_EdGraphNode::UpdatePinVisibility(UEdGraphPin* Pin, UTG_Pin* TGPin) cons
 	}
 
 	Pin->bHidden = !bCanEditChange && bEditConditionHides;
-	Pin->bDefaultValueIsReadOnly = bCanEditChange;
+	Pin->bDefaultValueIsReadOnly = !bCanEditChange;
 	
 	//Adding GIsTransacting check here as functions that are creating the Transaction
 	//should not be called from here 
diff --git a/Engine/Plugins/Experimental/TextureGraph/Source/TextureGraphEditor/Private/STG_NodePreview.cpp b/Engine/Plugins/Experimental/TextureGraph/Source/TextureGraphEditor/Private/STG_NodePreview.cpp
--- a/Engine/Plugins/Experimental/TextureGraph/Source/TextureGraphEditor/Private/STG_NodePreview.cpp
+++ b/Engine/Plugins/Experimental/TextureGraph/Source/TextureGraphEditor/Private/STG_NodePreview.cpp
@@ -44,8 +44,7 @@ void FNodeViewer::DrawCurrentImage(FViewport*, FCanvas* Canvas, const FDrawPrope
 TOptional<TVariant<FColor, FLinearColor>> FNodeViewer::GetCurrentImagePixelColor(FIntPoint PixelCoords, int32 MipIndex) const
 {
 	const uint64 PixelIndex = PixelCoords.Y * NodeDescriptor.Width + PixelCoords.X;
-
-	if (NodeTexture && CurrentBlob)
+	if (NodeTexture && CurrentBlob && CurrentBlob->IsValid())
 	{
 		if (!CurrentBlob->GetBufferRef()->HasRaw())
 		{
diff --git a/Engine/Plugins/Experimental/TextureGraph/Source/TextureGraphEditor/Private/TG_Exporter.cpp b/Engine/Plugins/Experimental/TextureGraph/Source/TextureGraphEditor/Private/TG_Exporter.cpp
--- a/Engine/Plugins/Experimental/TextureGraph/Source/TextureGraphEditor/Private/TG_Exporter.cpp
+++ b/Engine/Plugins/Experimental/TextureGraph/Source/TextureGraphEditor/Private/TG_Exporter.cpp
@@ -767,9 +767,9 @@ void FTG_ExporterImpl::SetTextureGraphToExport(UTextureGraph* InTextureGraph)
 {
 	// clear out previous handles
 	Cleanup();
-	
+
 	// overwrite the original TextureGraph in place by constructing a new one with the same name
-	FObjectDuplicationParameters Params = InitStaticDuplicateObjectParams(InTextureGraph, InTextureGraph->GetOuter(), NAME_None,
+	FObjectDuplicationParameters Params = InitStaticDuplicateObjectParams(InTextureGraph, GetTransientPackage(), NAME_None,
 	~RF_Standalone, UTextureGraph::StaticClass(), EDuplicateMode::Normal, EInternalObjectFlags::None);
 	
 	TextureGraphPtr = Cast<UTextureGraph>(StaticDuplicateObjectEx(Params));
diff --git a/Engine/Plugins/Experimental/TextureGraph/Source/TextureGraphEngine/Transform/Layer/T_Thumbnail.cpp b/Engine/Plugins/Experimental/TextureGraph/Source/TextureGraphEngine/Transform/Layer/T_Thumbnail.cpp
--- a/Engine/Plugins/Experimental/TextureGraph/Source/TextureGraphEngine/Transform/Layer/T_Thumbnail.cpp
+++ b/Engine/Plugins/Experimental/TextureGraph/Source/TextureGraphEngine/Transform/Layer/T_Thumbnail.cpp
@@ -42,8 +42,11 @@ TiledBlobRef T_Thumbnail::Bind(UMixInterface* Mix, UObject* Model, TiledBlobPtr
 		;
 
 	BufferDescriptor Desc = InBlobToBind->GetDescriptor();
-	Desc.Width = RenderMaterial_Thumbnail::GThumbWidth;
-	Desc.Height = RenderMaterial_Thumbnail::GThumbHeight;
+
+	// Normalize dimensions and remap to expected thumbnail size
+	const float MaxDimension = FMath::Max(Desc.Width, Desc.Height);
+	Desc.Width = RenderMaterial_Thumbnail::GThumbWidth * Desc.Width / MaxDimension;
+	Desc.Height = RenderMaterial_Thumbnail::GThumbHeight * Desc.Height / MaxDimension;
 
 	FString Name = FString::Printf(TEXT("T_Thumbnail [%s]"), *InBlobToBind->GetDescriptor().Name);
 	
