Histogram range

diff --git a/Engine/Plugins/Experimental/TextureGraph/Shaders/Expressions/Expression_Levels.usf b/Engine/Plugins/Experimental/TextureGraph/Shaders/Expressions/Expression_Levels.usf
--- a/Engine/Plugins/Experimental/TextureGraph/Shaders/Expressions/Expression_Levels.usf
+++ b/Engine/Plugins/Experimental/TextureGraph/Shaders/Expressions/Expression_Levels.usf
@@ -46,8 +46,7 @@ float LevelsControlInputCurve(float Value, float LowInput, float HighInput, floa
 	MidRanged = max(0.001, min(0.9999, MidRanged));
 
 	float InvGammaValue = log(MidRanged) / log(0.5);
-	float NewValue = pow(ValueRanged, InvGammaValue);
-
+	
 #if OUT_LEVELS
 	float4 HistoMin			= Histogram.Load(int3(2, 1, 0));
 	float4 HistoMax			= Histogram.Load(int3(3, 1, 0));
@@ -57,10 +56,12 @@ float LevelsControlInputCurve(float Value, float LowInput, float HighInput, floa
 	
 	float InputRange = clamp(HighInput - LowInput, 0, 1);
 	
-	NewValue = OutLow + ((NewValue - LowInput) / InputRange) * OutputRange;
+	ValueRanged = OutLow + ((ValueRanged - LowInput) / InputRange) * OutputRange;
 #endif /// OUT_LEVELS
 
-	return NewValue;
+
+
+	return pow(ValueRanged, InvGammaValue);
 }
 
 float4 FSH_Levels(float2 uv : TEXCOORD0) : SV_Target0
diff --git a/Engine/Plugins/Experimental/TextureGraph/Source/TextureGraph/Private/Expressions/FIlter/TG_Expression_Levels.cpp b/Engine/Plugins/Experimental/TextureGraph/Source/TextureGraph/Private/Expressions/FIlter/TG_Expression_Levels.cpp
--- a/Engine/Plugins/Experimental/TextureGraph/Source/TextureGraph/Private/Expressions/FIlter/TG_Expression_Levels.cpp
+++ b/Engine/Plugins/Experimental/TextureGraph/Source/TextureGraph/Private/Expressions/FIlter/TG_Expression_Levels.cpp
@@ -239,6 +239,9 @@ void UTG_Expression_Levels::Evaluate(FTG_EvaluationContext* InContext)
 	Output = T_Levels::Create(InContext->Cycle, Desc, Input.RasterBlob, LevelsControl, InContext->TargetId);
 }
 
+/////////////////////////////////////////////////////////////////
+/// Histogram scan
+/////////////////////////////////////////////////////////////////
 void UTG_Expression_HistogramScan::Evaluate(FTG_EvaluationContext* InContext)
 {
 	Super::Evaluate(InContext);
@@ -267,3 +270,47 @@ void UTG_Expression_HistogramScan::Evaluate(FTG_EvaluationContext* InContext)
 	BufferDescriptor Desc = Output.GetBufferDescriptor();
 	Output = T_Levels::Create(InContext->Cycle, Desc, Input.RasterBlob, LevelsControl, InContext->TargetId);
 }
+
+/////////////////////////////////////////////////////////////////
+/// Histogram range
+/////////////////////////////////////////////////////////////////
+void UTG_Expression_HistogramRange::Evaluate(FTG_EvaluationContext* InContext)
+{
+	Super::Evaluate(InContext);
+
+	if (!Input) // No Input, black Output
+	{
+		Output = FTG_Texture::GetBlack();
+		return;
+	}
+
+	LevelsControl = MakeShared<FLevels>();
+	LevelsControl->InitFromRange(Range, Position);
+
+	BufferDescriptor Desc = Output.GetBufferDescriptor();
+	Output = T_Levels::Create(InContext->Cycle, Desc, Input.RasterBlob, LevelsControl, InContext->TargetId);
+}
+
+/////////////////////////////////////////////////////////////////
+/// Histogram select
+/////////////////////////////////////////////////////////////////
+#if 0 /// TODO
+void UTG_Expression_HistogramSelect::Evaluate(FTG_EvaluationContext* InContext)
+{
+	Super::Evaluate(InContext);
+
+	if (!Input) // No Input, black Output
+	{
+		Output = FTG_Texture::GetBlack();
+		return;
+	}
+
+	LevelsControl = MakeShared<FLevels>();
+	LevelsControl->InitFromPositionContrast(Position, Contrast);
+	LevelsControl->InitFromRange(Range, Position);
+
+	BufferDescriptor Desc = Output.GetBufferDescriptor();
+	Output = T_Levels::Create(InContext->Cycle, Desc, Input.RasterBlob, LevelsControl, InContext->TargetId);
+}
+
+#endif 
\ No newline at end of file
diff --git a/Engine/Plugins/Experimental/TextureGraph/Source/TextureGraph/Public/Expressions/Filter/TG_Expression_Levels.h b/Engine/Plugins/Experimental/TextureGraph/Source/TextureGraph/Public/Expressions/Filter/TG_Expression_Levels.h
--- a/Engine/Plugins/Experimental/TextureGraph/Source/TextureGraph/Public/Expressions/Filter/TG_Expression_Levels.h
+++ b/Engine/Plugins/Experimental/TextureGraph/Source/TextureGraph/Public/Expressions/Filter/TG_Expression_Levels.h
@@ -140,7 +140,9 @@ public:
 	bool IsAutoLevel() const { return LevelsExpressionType == ELevelsExpressionType::AutoLowHigh; }
 };
 
-
+/////////////////////////////////////////////////////////////////
+/// Histogram scan
+/////////////////////////////////////////////////////////////////
 UCLASS()
 class TEXTUREGRAPH_API UTG_Expression_HistogramScan : public UTG_Expression
 {
@@ -171,4 +173,81 @@ public:
 
 	virtual FText						GetTooltipText() const override { return FText::FromString(TEXT("Lets you drive the contrast and position of the histogram. Input must be a grayscale image.")); } 
 
-};
\ No newline at end of file
+};
+
+/////////////////////////////////////////////////////////////////
+/// Histogram range
+/////////////////////////////////////////////////////////////////
+UCLASS()
+class TEXTUREGRAPH_API UTG_Expression_HistogramRange : public UTG_Expression
+{
+	GENERATED_BODY()
+
+	TSharedPtr<FLevels> 				LevelsControl;
+
+public:
+
+	TG_DECLARE_EXPRESSION(TG_Category::Adjustment);
+	virtual void						Evaluate(FTG_EvaluationContext* InContext) override;
+
+	// The input image to adjust the levels for
+	UPROPERTY(EditAnywhere, Category = NoCategory, meta = (TGType = "TG_Input", PinDisplayName = "", MD_HistogramLuminance))
+	FTG_Texture							Input;
+
+	// Drives the range of the levels out
+	UPROPERTY(EditAnywhere, Category = NoCategory, meta = (TGType = "TG_Setting", UIMin = "0", ClampMin = "0", UIMax = "1", ClampMax = "1", MD_ScalarEditor))
+	float 								Range = 1.0f; 
+
+
+	// Drives the position of the histogram of the input image
+	UPROPERTY(EditAnywhere, Category = NoCategory, meta = (TGType = "TG_Setting", UIMin = "0", ClampMin = "0", UIMax = "1", ClampMax = "1", MD_ScalarEditor))
+	float 								Position = 0.5f;
+
+	// The output image filtered result of the Levels operator 
+	UPROPERTY(EditAnywhere, Category = NoCategory, meta = (TGType = "TG_Output", PinDisplayName = ""))
+	FTG_Texture							Output;
+
+	virtual FText						GetTooltipText() const override { return FText::FromString(TEXT("Reduce and/or move the range of a grayscale input. Can be used to remap grayscale images.")); } 
+
+};
+
+/////////////////////////////////////////////////////////////////
+/// Histogram select
+/////////////////////////////////////////////////////////////////
+#if 0 /// TODO
+UCLASS()
+class TEXTUREGRAPH_API UTG_Expression_HistogramSelect : public UTG_Expression
+{
+	GENERATED_BODY()
+
+	TSharedPtr<FLevels> 				LevelsControl;
+
+public:
+
+	TG_DECLARE_EXPRESSION(TG_Category::Adjustment);
+	virtual void						Evaluate(FTG_EvaluationContext* InContext) override;
+
+	// The input image to adjust the levels for
+	UPROPERTY(EditAnywhere, Category = NoCategory, meta = (TGType = "TG_Input", PinDisplayName = "", MD_HistogramLuminance))
+	FTG_Texture							Input;
+
+	// Drives the position of the histogram of the input image
+	UPROPERTY(EditAnywhere, Category = NoCategory, meta = (TGType = "TG_Setting", UIMin = "0", ClampMin = "0", UIMax = "1", ClampMax = "1", MD_ScalarEditor))
+	float 								Position = 0.5f;
+
+	// Drives the range of the levels out
+	UPROPERTY(EditAnywhere, Category = NoCategory, meta = (TGType = "TG_Setting", UIMin = "0", ClampMin = "0", UIMax = "1", ClampMax = "1", MD_ScalarEditor))
+	float 								Range = 1.0f; 
+
+	// Drives the contrast of the histogram of the input image
+	UPROPERTY(EditAnywhere, Category = NoCategory, meta = (TGType = "TG_Setting", UIMin = "0", ClampMin = "0", UIMax = "1", ClampMax = "1", MD_ScalarEditor))
+	float 								Contrast = 0.5f; 
+
+	// The output image filtered result of the Levels operator 
+	UPROPERTY(EditAnywhere, Category = NoCategory, meta = (TGType = "TG_Output", PinDisplayName = ""))
+	FTG_Texture							Output;
+
+	virtual FText						GetTooltipText() const override { return FText::FromString(TEXT("Reduce and/or move the range of a grayscale input. Can be used to remap grayscale images.")); } 
+
+};
+#endif 
diff --git a/Engine/Plugins/Experimental/TextureGraph/Source/TextureGraphEngine/Transform/Expressions/T_Levels.cpp b/Engine/Plugins/Experimental/TextureGraph/Source/TextureGraphEngine/Transform/Expressions/T_Levels.cpp
--- a/Engine/Plugins/Experimental/TextureGraph/Source/TextureGraphEngine/Transform/Expressions/T_Levels.cpp
+++ b/Engine/Plugins/Experimental/TextureGraph/Source/TextureGraphEngine/Transform/Expressions/T_Levels.cpp
@@ -112,8 +112,19 @@ void FLevels::InitFromPositionContrast(float InPosition, float InContrast)
 	IsAutoLevels = false;
 }
 
+void FLevels::InitFromRange(float InRange_, float InPosition_)
+{
+	float InRange = FMath::Clamp(1.0f - InRange_, 0.0f, 1.0f) * 0.5f;
+	float InPosition = FMath::Clamp(1.0f - InPosition_, 0.0f, 1.0f) * 0.5f;
+	float C = InRange; /// *0.5f;
+	float P = 1.0f - FMath::Clamp(InPosition, 0, 1);
+	float P1 = (FMath::Max(P, 0.5f) - 0.5f) * 2.0f;
+	float P2 = FMath::Min(P * 2.0f, 1.0f);
+	OutLow = FMath::Lerp(P1, P2, C);
+	OutHigh = FMath::Lerp(P2, P1, C);
 
-
+	IsAutoLevels = false;
+}
 
 class TEXTUREGRAPHENGINE_API RenderMaterial_FX_Levels : public RenderMaterial_FX
 {
diff --git a/Engine/Plugins/Experimental/TextureGraph/Source/TextureGraphEngine/Transform/Expressions/T_Levels.h b/Engine/Plugins/Experimental/TextureGraph/Source/TextureGraphEngine/Transform/Expressions/T_Levels.h
--- a/Engine/Plugins/Experimental/TextureGraph/Source/TextureGraphEngine/Transform/Expressions/T_Levels.h
+++ b/Engine/Plugins/Experimental/TextureGraph/Source/TextureGraphEngine/Transform/Expressions/T_Levels.h
@@ -88,6 +88,7 @@ struct TEXTUREGRAPHENGINE_API FLevels
 	void InitFromLowMidHigh(float LowValue, float MidValue, float HighValue, float OutLow, float OutHigh);
 	void InitFromAutoLevels(float MidPercentage);
 	void InitFromPositionContrast(float MidValue, float Contrast);
+	void InitFromRange(float InRange, float InPosition);
 
 	FLevelsHistogramData HistogramData;
 };
