Dangling texture pointer issues fixed when a texture graph is called through a blueprint witha  texture parameter.

diff --git a/Engine/Plugins/Experimental/TextureGraph/Source/TextureGraph/Private/Blueprint/TG_AsyncExportTask.cpp b/Engine/Plugins/Experimental/TextureGraph/Source/TextureGraph/Private/Blueprint/TG_AsyncExportTask.cpp
--- a/Engine/Plugins/Experimental/TextureGraph/Source/TextureGraph/Private/Blueprint/TG_AsyncExportTask.cpp
+++ b/Engine/Plugins/Experimental/TextureGraph/Source/TextureGraph/Private/Blueprint/TG_AsyncExportTask.cpp
@@ -44,7 +44,12 @@ void UTG_AsyncExportTask::Activate()
 	TargetExportSettings = FExportSettings();
 	TargetExportSettings.OnDone.BindUFunction(this, "OnExportDone");
 
-	FTG_HelperFunctions::ExportAsync(TextureGraphPtr, "", "", TargetExportSettings, false, OverwriteTextures, true);
+	FTG_HelperFunctions::ExportAsync(TextureGraphPtr, "", "", TargetExportSettings, false, OverwriteTextures, true)
+		.then([this](bool)
+		{
+			OnExportDone();
+		});
+
 }
 
 void UTG_AsyncExportTask::OnExportDone()
@@ -64,4 +69,4 @@ void UTG_AsyncExportTask::FinishDestroy()
 		OrignalTextureGraphPtr = nullptr;
 	}
 	Super::FinishDestroy();
-}
\ No newline at end of file
+}
diff --git a/Engine/Plugins/Experimental/TextureGraph/Source/TextureGraph/Private/Blueprint/TG_BlueprintFunctionLibrary.cpp b/Engine/Plugins/Experimental/TextureGraph/Source/TextureGraph/Private/Blueprint/TG_BlueprintFunctionLibrary.cpp
--- a/Engine/Plugins/Experimental/TextureGraph/Source/TextureGraph/Private/Blueprint/TG_BlueprintFunctionLibrary.cpp
+++ b/Engine/Plugins/Experimental/TextureGraph/Source/TextureGraph/Private/Blueprint/TG_BlueprintFunctionLibrary.cpp
@@ -13,6 +13,7 @@
 #include "Expressions/Input/TG_Expression_Color.h"
 #include "Expressions/Input/TG_Expression_OutputSettings.h"
 #include "TG_HelperFunctions.h"
+#include "Engine/Texture2D.h"
 
 #include UE_INLINE_GENERATED_CPP_BY_NAME(TG_BlueprintFunctionLibrary)
 
@@ -35,7 +36,8 @@ void UTG_BlueprintFunctionLibrary::SetTextureParameterValue(UObject* WorldContex
 				if (ExpressionPtr)
 				{
 					bFoundParameter = true;
-					ExpressionPtr->SetAsset(ParameterValue);
+					UTexture2D* DupTexture = (UTexture2D*)StaticDuplicateObject(ParameterValue, GetTransientPackage(), NAME_None, ~RF_Standalone, UTexture2D::StaticClass());
+					ExpressionPtr->SetAsset(DupTexture);
 				}
 			}
 		}
diff --git a/Engine/Plugins/Experimental/TextureGraph/Source/TextureGraph/Private/TG_HelperFunctions.cpp b/Engine/Plugins/Experimental/TextureGraph/Source/TextureGraph/Private/TG_HelperFunctions.cpp
--- a/Engine/Plugins/Experimental/TextureGraph/Source/TextureGraph/Private/TG_HelperFunctions.cpp
+++ b/Engine/Plugins/Experimental/TextureGraph/Source/TextureGraph/Private/TG_HelperFunctions.cpp
@@ -173,9 +173,7 @@ AsyncBool FTG_HelperFunctions::ExportAsync(UTextureGraph* InTextureGraph, FStrin
 	.then([InTextureGraph, &TargetExportSettings]()
 	{
 		InTextureGraph->InvalidateAll();
-	})
-	.then([InTextureGraph, &TargetExportSettings]()
-	{
+
 		// Add analytics tag
 		if (FEngineAnalytics::IsAvailable())
 		{
@@ -185,7 +183,8 @@ AsyncBool FTG_HelperFunctions::ExportAsync(UTextureGraph* InTextureGraph, FStrin
 			// Send Analytics event 
 			FEngineAnalytics::GetProvider().RecordEvent(TEXT("Editor.Usage.TextureGraph.Export"), Attributes);
 		}
-		return cti::make_ready_continuable(true);
+
+		return true;
 	});
 }
 
diff --git a/Engine/Plugins/Experimental/TextureGraph/Source/TextureGraphEngine/Data/Blobber.cpp b/Engine/Plugins/Experimental/TextureGraph/Source/TextureGraphEngine/Data/Blobber.cpp
--- a/Engine/Plugins/Experimental/TextureGraph/Source/TextureGraphEngine/Data/Blobber.cpp
+++ b/Engine/Plugins/Experimental/TextureGraph/Source/TextureGraphEngine/Data/Blobber.cpp
@@ -242,7 +242,7 @@ void Blobber::UpdateBlobCache()
 		// if the blob has a single ref count left then we can safely de-cache it
 		if (Entry->BlobObj.use_count() == 1)
 		{
-			UE_LOG(LogBlob, Log, TEXT("Removing permanent blob %s (Hash: %llu)"), *Entry->BlobObj->DisplayName(), Entry->BlobObj->Hash()->Value());
+			UE_LOG(LogBlob, VeryVerbose, TEXT("Removing permanent blob %s (Hash: %llu)"), *Entry->BlobObj->DisplayName(), Entry->BlobObj->Hash()->Value());
 			HashesToRemove.push_back(Iter.Key());
 
 			HashTypeVec IntermediateHashes = Entry->BlobObj->Hash()->GetIntermediateHashes();
@@ -270,7 +270,7 @@ void Blobber::UpdateBlobCache()
 
 	if (!HashesToRemove.empty())
 	{
-		UE_LOG(LogBlob, Log, TEXT("Removing num items from the cache: %llu"), HashesToRemove.size());
+		UE_LOG(LogBlob, VeryVerbose, TEXT("Removing num items from the cache: %llu"), HashesToRemove.size());
 		for (size_t RemoveIndex = 0; RemoveIndex < HashesToRemove.size(); RemoveIndex++)
 		{
 			HashType HashToRemove = HashesToRemove[RemoveIndex];
