Fix issues with some EXRs

diff --git a/Engine/Plugins/Experimental/TextureGraph/Source/TextureGraphEngine/2D/TextureHelper.cpp b/Engine/Plugins/Experimental/TextureGraph/Source/TextureGraphEngine/2D/TextureHelper.cpp
--- a/Engine/Plugins/Experimental/TextureGraph/Source/TextureGraphEngine/2D/TextureHelper.cpp
+++ b/Engine/Plugins/Experimental/TextureGraph/Source/TextureGraphEngine/2D/TextureHelper.cpp
@@ -358,6 +358,7 @@ ETextureRenderTargetFormat TextureHelper::GetRenderTargetFormatFromPixelFormat(E
 		return ETextureRenderTargetFormat::RTF_RGB10A2; 
 	case PF_R32G32B32F:
 		return ETextureRenderTargetFormat::RTF_RGBA32f;
+	case PF_A16B16G16R16:
 	case PF_FloatRGBA:
 		return ETextureRenderTargetFormat::RTF_RGBA16f;
 	default:
@@ -1478,8 +1479,12 @@ bool TextureHelper::GetPixelFormatFromTextureSourceFormat(ETextureSourceFormat S
 	switch (SourceFormat)
 	{
 	// Currently supported formats : 
-	case ETextureSourceFormat::TSF_RGBA32F:
 	case ETextureSourceFormat::TSF_RGBA16F:
+		OutPixelFormat = PF_FloatRGBA;
+		OutNumChannels = 4;
+		break;
+
+	case ETextureSourceFormat::TSF_RGBA32F:
 	case ETextureSourceFormat::TSF_RGBA16:
 	case ETextureSourceFormat::TSF_BGRA8:
 	case ETextureSourceFormat::TSF_BGRE8:
