Fix sporadic crashes with texture graph using blueprints

diff --git a/Engine/Plugins/Experimental/TextureGraph/Source/TextureGraphEngine/Data/Blobber.cpp b/Engine/Plugins/Experimental/TextureGraph/Source/TextureGraphEngine/Data/Blobber.cpp
--- a/Engine/Plugins/Experimental/TextureGraph/Source/TextureGraphEngine/Data/Blobber.cpp
+++ b/Engine/Plugins/Experimental/TextureGraph/Source/TextureGraphEngine/Data/Blobber.cpp
@@ -873,6 +873,30 @@ void Blobber::AddBlobEntry(HashType Hash, BlobPtr BlobObj, BlobCacheOptions Opti
 	InvalidateTimestamp = Util::Time();
 }
 
+void Blobber::RemoveHashMapping(HashType Hash)
+{
+	UE_LOG(LogBlob, Warning, TEXT("Removing incorrect hash mapping: %#016lx"), Hash);
+
+	FScopeLock HashLock(&HashMutex);
+
+	auto HashIter = HashMappings.find(Hash);
+
+	/// if no mapping was found 
+	/// then just ignore it ...
+	if (HashIter == HashMappings.end())
+		return;
+
+	/// there's a cyclical Hash link, then ignore it as well
+	if (HashIter->second->Value() == Hash)
+		return;
+
+	/// Otherwise, we need to use the RHS of the mapping
+	CHashPtr MappedHash = HashIter->second;
+
+	HashMappings.erase(HashIter);
+	Hash = MappedHash->Value();
+}
+
 BlobRef Blobber::AddInternal(BlobPtr BlobObj, BlobCacheOptions Options)
 {
 	CHashPtr Hash = BlobObj->Hash();
@@ -908,8 +932,10 @@ BlobRef Blobber::AddInternal(BlobPtr BlobObj, BlobCacheOptions Options)
 					return BlobRef(std::static_pointer_cast<Blob>(TiledBlobObj), true, false);
 				}
 
-				/// 2. TODO: This is a tricky one
-				check(false);
+				/// We need to remove this hash mapping
+				RemoveHashMapping(Hash->Value());
+
+				return BlobRef(nullptr);
 			}
 
 			/// 2. If the incoming BlobObj is un-tiled but the Existing one IS tiled
@@ -920,10 +946,17 @@ BlobRef Blobber::AddInternal(BlobPtr BlobObj, BlobCacheOptions Options)
 				TiledBlobPtr ExistingTiled = std::static_pointer_cast<TiledBlob>(Existing.lock());
 
 				/// The hashes can only match up if the Existing (tiled) BlobObj is 1x1
-				check(ExistingTiled->Rows() == 1 && ExistingTiled->Cols() == 1);
-				check(*ExistingTiled->GetTile(0, 0)->Hash() == *Hash);
+				if (ExistingTiled->Rows() == 1 && ExistingTiled->Cols() == 1)
+				{
+					//check(ExistingTiled->Rows() == 1 && ExistingTiled->Cols() == 1);
+					//check(*ExistingTiled->GetTile(0, 0)->Hash() == *Hash);
+					return ExistingTiled->GetTile(0, 0).lock();
+				}
+
+				/// We need to remove this hash mapping
+				RemoveHashMapping(Hash->Value());
 
-				return ExistingTiled->GetTile(0, 0).lock();
+				return BlobRef(nullptr);
 			}
 		}
 
diff --git a/Engine/Plugins/Experimental/TextureGraph/Source/TextureGraphEngine/Data/Blobber.h b/Engine/Plugins/Experimental/TextureGraph/Source/TextureGraphEngine/Data/Blobber.h
--- a/Engine/Plugins/Experimental/TextureGraph/Source/TextureGraphEngine/Data/Blobber.h
+++ b/Engine/Plugins/Experimental/TextureGraph/Source/TextureGraphEngine/Data/Blobber.h
@@ -165,6 +165,7 @@ private:
 	void							AddHashMapping(HashType LHS, CHashPtr RHS);
 	void							AddBlobEntry(HashType Hash, BlobPtr BlobObj, BlobCacheOptions Options);
 	void							AddBlobEntryThreadSafe(HashType Hash, BlobPtr BlobObj, BlobCacheOptions Options);
+	void							RemoveHashMapping(HashType Hash);
 
 	void							PrintStats();
 	BlobRef							AddInternal(BlobPtr BlobObj, BlobCacheOptions Options);
diff --git a/Engine/Plugins/Experimental/TextureGraph/Source/TextureGraphEngine/Data/TiledBlob.cpp b/Engine/Plugins/Experimental/TextureGraph/Source/TextureGraphEngine/Data/TiledBlob.cpp
--- a/Engine/Plugins/Experimental/TextureGraph/Source/TextureGraphEngine/Data/TiledBlob.cpp
+++ b/Engine/Plugins/Experimental/TextureGraph/Source/TextureGraphEngine/Data/TiledBlob.cpp
@@ -1059,9 +1059,10 @@ void TiledBlob_Promise::AddLinkedBlob(BlobPtr LinkedBlob)
 	check(LinkedBlob->IsTiled());
 
 	TiledBlobPtr TiledLinkedBlob = std::static_pointer_cast<TiledBlob>(LinkedBlob);
-	check(TiledLinkedBlob->IsPromise());
 
-	TiledBlob_PromisePtr TiledPromiseLinkedBlob = std::static_pointer_cast<TiledBlob_Promise>(LinkedBlob);
+	/// If the number of rows and columns don't match up then we don't do anything
+	if (TiledLinkedBlob->Rows() != Rows() || TiledLinkedBlob->Cols() != Cols())
+		return;
 
 	/// If we got a different pointer back from the blobber then that means that this result was already cached into 
 	/// the system. We need to copy it over to the original pointer if it's a finalised blob
@@ -1069,11 +1070,14 @@ void TiledBlob_Promise::AddLinkedBlob(BlobPtr LinkedBlob)
 	/// and we don't have do any awkward pointer adjustment shenanigans to make things work
 	if (IsFinalised())
 	{
-		/// Copy the contents over
-		TiledPromiseLinkedBlob->FinaliseFrom(this);
+		TiledLinkedBlob->FinaliseFrom(this);
 	}
 
-	TiledPromiseLinkedBlob->CachedBlob = std::static_pointer_cast<TiledBlob_Promise>(shared_from_this());
+	if (TiledLinkedBlob->IsPromise())
+	{
+		TiledBlob_PromisePtr TiledPromiseLinkedBlob = std::static_pointer_cast<TiledBlob_Promise>(LinkedBlob);
+		TiledPromiseLinkedBlob->CachedBlob = std::static_pointer_cast<TiledBlob_Promise>(shared_from_this());
+	}
 
 	TiledBlob::AddLinkedBlob(LinkedBlob);
 }
diff --git a/Engine/Plugins/Experimental/TextureGraph/Source/TextureGraphEngine/Data/TiledBlob.h b/Engine/Plugins/Experimental/TextureGraph/Source/TextureGraphEngine/Data/TiledBlob.h
--- a/Engine/Plugins/Experimental/TextureGraph/Source/TextureGraphEngine/Data/TiledBlob.h
+++ b/Engine/Plugins/Experimental/TextureGraph/Source/TextureGraphEngine/Data/TiledBlob.h
@@ -17,6 +17,7 @@ typedef T_Tiles<BlobPtr>			BlobPtrTiles;
 class TEXTUREGRAPHENGINE_API TiledBlob : public Blob
 {
 	friend class Blobber;
+	friend class TiledBlob_Promise;
 
 protected:
 	mutable BlobPtrTiles			Tiles;					/// The Tiles that make up the larger BlobObj
diff --git a/Engine/Plugins/Experimental/TextureGraph/Source/TextureGraphEngine/Job/Job.cpp b/Engine/Plugins/Experimental/TextureGraph/Source/TextureGraphEngine/Job/Job.cpp
--- a/Engine/Plugins/Experimental/TextureGraph/Source/TextureGraphEngine/Job/Job.cpp
+++ b/Engine/Plugins/Experimental/TextureGraph/Source/TextureGraphEngine/Job/Job.cpp
@@ -801,7 +801,7 @@ AsyncPrepareResult Job::PrepareTargets(JobBatch* Batch)
 			}
 			else
 			{
-				UE_LOG(LogJob, Log, TEXT("Transform: %s (Tile: %d, %d) => %s"), *Transform->GetName(), (int32)TileX, (int32)TileY, *TileBlob->Name());
+				UE_LOG(LogJob, VeryVerbose, TEXT("Transform: %s (Tile: %d, %d) => %s"), *Transform->GetName(), (int32)TileX, (int32)TileY, *TileBlob->Name());
 				TileInvalidationMatrix[TileX][TileY] = false;
 			}
 
