Adding save option to TG Export Async blueprint node

diff --git a/Engine/Plugins/Experimental/TextureGraph/Source/TextureGraph/Private/Blueprint/TG_AsyncExportTask.cpp b/Engine/Plugins/Experimental/TextureGraph/Source/TextureGraph/Private/Blueprint/TG_AsyncExportTask.cpp
--- a/Engine/Plugins/Experimental/TextureGraph/Source/TextureGraph/Private/Blueprint/TG_AsyncExportTask.cpp
+++ b/Engine/Plugins/Experimental/TextureGraph/Source/TextureGraph/Private/Blueprint/TG_AsyncExportTask.cpp
@@ -14,10 +14,12 @@ UTG_AsyncExportTask::UTG_AsyncExportTask(const FObjectInitializer& ObjectInitial
 	
 }
 
-UTG_AsyncExportTask* UTG_AsyncExportTask::TG_AsyncExportTask( UTextureGraph* InTextureGraph, const bool OverwriteTextures)
+UTG_AsyncExportTask* UTG_AsyncExportTask::TG_AsyncExportTask( UTextureGraph* InTextureGraph, const bool OverwriteTextures, const bool bSave)
 {
 	UTG_AsyncExportTask* Task = NewObject<UTG_AsyncExportTask>();
 	Task->OverwriteTextures = OverwriteTextures;
+	Task->bSave = bSave;
+
 	if (InTextureGraph != nullptr)
 	{
 		Task->OrignalTextureGraphPtr = InTextureGraph;
@@ -44,7 +46,7 @@ void UTG_AsyncExportTask::Activate()
 	TargetExportSettings = FExportSettings();
 	TargetExportSettings.OnDone.BindUFunction(this, "OnExportDone");
 
-	FTG_HelperFunctions::ExportAsync(TextureGraphPtr, "", "", TargetExportSettings, false, OverwriteTextures, true)
+	FTG_HelperFunctions::ExportAsync(TextureGraphPtr, "", "", TargetExportSettings, false, OverwriteTextures, true,bSave)
 		.then([this](bool)
 		{
 			OnExportDone();
diff --git a/Engine/Plugins/Experimental/TextureGraph/Source/TextureGraph/Private/TG_HelperFunctions.cpp b/Engine/Plugins/Experimental/TextureGraph/Source/TextureGraph/Private/TG_HelperFunctions.cpp
--- a/Engine/Plugins/Experimental/TextureGraph/Source/TextureGraph/Private/TG_HelperFunctions.cpp
+++ b/Engine/Plugins/Experimental/TextureGraph/Source/TextureGraph/Private/TG_HelperFunctions.cpp
@@ -92,7 +92,7 @@ void FTG_HelperFunctions::EnsureOutputIsTexture(MixUpdateCyclePtr Cycle, UTG_Nod
 	}
 }
 
-AsyncBool FTG_HelperFunctions::ExportAsync(UTextureGraph* InTextureGraph, FString ExportPath, FString AssetName, FExportSettings& TargetExportSettings, bool OverrideExportPath, bool OverwriteTextures /*= true*/, bool ExportAllOutputs /*= false*/)
+AsyncBool FTG_HelperFunctions::ExportAsync(UTextureGraph* InTextureGraph, FString ExportPath, FString AssetName, FExportSettings& TargetExportSettings, bool OverrideExportPath, bool OverwriteTextures /*= true*/, bool ExportAllOutputs /*= false*/,bool bSave /*= true*/)
 {
 	FString ErrorMessage = "";
 	TargetExportSettings.Reset();
@@ -145,6 +145,7 @@ AsyncBool FTG_HelperFunctions::ExportAsync(UTextureGraph* InTextureGraph, FStrin
 					MapSettings.IsSRGB = OutputSetting.bSRGB;
 					MapSettings.Width = (int32)OutputSetting.Width;
 					MapSettings.Height = (int32)OutputSetting.Height;
+					MapSettings.bSave = bSave;
 					TargetExportSettings.ExportPreset.push_back(std::pair<FName, FExportMapSettings>{ MapSettings.Name, MapSettings });
 				}
 				else
diff --git a/Engine/Plugins/Experimental/TextureGraph/Source/TextureGraph/Public/Blueprint/TG_AsyncExportTask.h b/Engine/Plugins/Experimental/TextureGraph/Source/TextureGraph/Public/Blueprint/TG_AsyncExportTask.h
--- a/Engine/Plugins/Experimental/TextureGraph/Source/TextureGraph/Public/Blueprint/TG_AsyncExportTask.h
+++ b/Engine/Plugins/Experimental/TextureGraph/Source/TextureGraph/Public/Blueprint/TG_AsyncExportTask.h
@@ -21,7 +21,7 @@ class TEXTUREGRAPH_API UTG_AsyncExportTask : public UTG_AsyncTask
 
 public:
 	UFUNCTION(BlueprintCallable, Category = "TextureGraph" , meta=(DisplayName="Texture Graph Export (Async)", BlueprintInternalUseOnly = "true"))
-	static UTG_AsyncExportTask* TG_AsyncExportTask(UTextureGraph* TextureGraph, const bool OverwriteTextures);
+	static UTG_AsyncExportTask* TG_AsyncExportTask(UTextureGraph* TextureGraph, const bool OverwriteTextures, const bool bSave);
 
 	virtual void Activate() override;
 
@@ -36,8 +36,13 @@ private:
 	void OnExportDone();
 	
 	bool OverwriteTextures = true;
+
+	bool bSave = false;
+
 	UTextureGraph* OrignalTextureGraphPtr;
+
 	UTextureGraph* TextureGraphPtr;
+
 	FExportSettings TargetExportSettings;
 };
 
diff --git a/Engine/Plugins/Experimental/TextureGraph/Source/TextureGraph/Public/TG_HelperFunctions.h b/Engine/Plugins/Experimental/TextureGraph/Source/TextureGraph/Public/TG_HelperFunctions.h
--- a/Engine/Plugins/Experimental/TextureGraph/Source/TextureGraph/Public/TG_HelperFunctions.h
+++ b/Engine/Plugins/Experimental/TextureGraph/Source/TextureGraph/Public/TG_HelperFunctions.h
@@ -20,7 +20,7 @@ public:
 
 	static void EnsureOutputIsTexture(MixUpdateCyclePtr Cycle, UTG_Node* OutputNode);
 
-	static AsyncBool ExportAsync(UTextureGraph* InTextureGraph, FString ExportPath, FString AssetName, FExportSettings& TargetExportSettings, bool OverrideExportPath, bool OverwriteTextures = true,bool ExportAllOutputs = false);
+	static AsyncBool ExportAsync(UTextureGraph* InTextureGraph, FString ExportPath, FString AssetName, FExportSettings& TargetExportSettings, bool OverrideExportPath, bool OverwriteTextures = true,bool ExportAllOutputs = false, bool bSave = true);
 	static AsyncBool RenderAsync(UTextureGraph* InTextureGraph, JobBatchPtr ExistingBatch = nullptr);
 
 
diff --git a/Engine/Plugins/Experimental/TextureGraph/Source/TextureGraphEngine/Export/TextureExporter.cpp b/Engine/Plugins/Experimental/TextureGraph/Source/TextureGraphEngine/Export/TextureExporter.cpp
--- a/Engine/Plugins/Experimental/TextureGraph/Source/TextureGraphEngine/Export/TextureExporter.cpp
+++ b/Engine/Plugins/Experimental/TextureGraph/Source/TextureGraphEngine/Export/TextureExporter.cpp
@@ -187,8 +187,18 @@ AsyncInt TextureExporter::ExportRawAsUAsset(RawBufferPtr RawObj,const FExportMap
 		FAssetRegistryModule::AssetCreated(NewTexture);
 		Package->FullyLoad();
 		FTextureCompilingManager::Get().FinishCompilation({NewTexture});
-		bool success = UEditorLoadingAndSavingUtils::SavePackages({ Package }, false);
-		Promise.set_value(success);
+		
+		bool Success = true;
+		if (Setting.bSave)
+		{
+			Success = UEditorLoadingAndSavingUtils::SavePackages({ Package }, false);
+		}
+		else
+		{
+			Package->SetDirtyFlag(true);
+		}
+
+		Promise.set_value(Success);
 #else
 		Promise.set_value(-1);
 #endif
diff --git a/Engine/Plugins/Experimental/TextureGraph/Source/TextureGraphEngine/Export/TextureExporter.h b/Engine/Plugins/Experimental/TextureGraph/Source/TextureGraphEngine/Export/TextureExporter.h
--- a/Engine/Plugins/Experimental/TextureGraph/Source/TextureGraphEngine/Export/TextureExporter.h
+++ b/Engine/Plugins/Experimental/TextureGraph/Source/TextureGraphEngine/Export/TextureExporter.h
@@ -32,6 +32,7 @@ struct TEXTUREGRAPHENGINE_API FExportMapSettings
 	int32 Height;
 	TextureCompressionSettings Compression;
 	TextureGroup LODGroup;
+	bool bSave;
 
 	DECLARE_DELEGATE_OneParam(OnExportMapDone, FExportMapSettings MapSettings);
 	OnExportMapDone OnDone;
