Assigning the same output to two different material inputs doesnt render correctly

diff --git a/Engine/Plugins/Experimental/TextureGraph/Source/TextureGraphEditor/Private/TG_RenderModeManager.cpp b/Engine/Plugins/Experimental/TextureGraph/Source/TextureGraphEditor/Private/TG_RenderModeManager.cpp
--- a/Engine/Plugins/Experimental/TextureGraph/Source/TextureGraphEditor/Private/TG_RenderModeManager.cpp
+++ b/Engine/Plugins/Experimental/TextureGraph/Source/TextureGraphEditor/Private/TG_RenderModeManager.cpp
@@ -32,15 +32,16 @@ void TG_RenderModeManager::ChangeRenderMode(FName NewRenderMode, UTextureGraph*
 void TG_RenderModeManager::UpdateRenderMode(UTextureGraph* TextureGraph)
 {
 	UMixSettings* settings = TextureGraph->GetSettings();
-	size_t numTargets = settings->NumTargets();
+	size_t NumTargets = settings->NumTargets();
 	const FName MaterialName = settings->GetViewportSettings().GetMaterialName();
+	const int32 NumRenderModeMappings = _renderModeMaterials.Num();
 	
-	if (numTargets != _renderModeMaterials.Num())
+	if (NumTargets != NumRenderModeMappings)
 	{
-		InitializeDefaultMaterials(numTargets, TextureGraph);
+		InitializeDefaultMaterials(NumTargets, TextureGraph);
 	}
 
-	for (size_t ti = 0; ti < numTargets; ti++)
+	for (size_t ti = 0; ti < NumTargets; ti++)
 	{
 		RenderMaterial_BPPtr renderMaterial = GetTargetRenderModeMaterial(ti,_currentRenderMode);
 		renderMaterial->Instance()->ClearParameterValues();
diff --git a/Engine/Plugins/Experimental/TextureGraph/Source/TextureGraphEngine/2D/TargetTextureSet.cpp b/Engine/Plugins/Experimental/TextureGraph/Source/TextureGraphEngine/2D/TargetTextureSet.cpp
--- a/Engine/Plugins/Experimental/TextureGraph/Source/TextureGraphEngine/2D/TargetTextureSet.cpp
+++ b/Engine/Plugins/Experimental/TextureGraph/Source/TextureGraphEngine/2D/TargetTextureSet.cpp
@@ -41,7 +41,6 @@ AsyncBufferResultPtr TargetTextureSet::BindTo(RenderMaterial_BPPtr Material, TAr
 	
 	for(const FMaterialMappingInfo MappingInfo : MaterialMappingInfo)
 	{
-		
 		BindOnTextureUpdate(Material, MappingInfo);
 	}
 			
@@ -94,9 +93,12 @@ void TargetTextureSet::SetTexture(FName TextureName, TiledBlobRef InTexture)
 	}
 
 	//if there exists a callback, now is the time to update it 
-	if (Callbacks.Contains(TextureName) && Callbacks[TextureName])
+	if (Callbacks.Contains(TextureName) && Callbacks[TextureName].Num())
 	{
-		Callbacks[TextureName](InTexture);
+		for(int i = 0; i < Callbacks[TextureName].Num(); i++)
+		{
+			Callbacks[TextureName][i](InTexture);
+		}
 	}
 }
 
@@ -171,11 +173,11 @@ void TargetTextureSet::RegisterCallback(TextureReadyCallback Callback, FMaterial
 
 	if(Callbacks.Contains(MaterialMappingInfo.Target))
 	{
-		Callbacks[MaterialMappingInfo.Target] = Callback;
+		Callbacks[MaterialMappingInfo.Target].Add(Callback);
 	}
 	else
 	{
-		Callbacks.Add(MaterialMappingInfo.Target, Callback);
+		Callbacks.Add(MaterialMappingInfo.Target, TArray<TextureReadyCallback>{Callback});
 	}
 }
 
diff --git a/Engine/Plugins/Experimental/TextureGraph/Source/TextureGraphEngine/2D/TargetTextureSet.h b/Engine/Plugins/Experimental/TextureGraph/Source/TextureGraphEngine/2D/TargetTextureSet.h
--- a/Engine/Plugins/Experimental/TextureGraph/Source/TextureGraphEngine/2D/TargetTextureSet.h
+++ b/Engine/Plugins/Experimental/TextureGraph/Source/TextureGraphEngine/2D/TargetTextureSet.h
@@ -17,7 +17,7 @@ class TEXTUREGRAPHENGINE_API TargetTextureSet : public TextureSet
 {
 protected:
 
-	mutable TMap<FName, TextureReadyCallback> Callbacks;		/// The callbacks to call when texture is set
+	mutable TMap<FName, TArray<TextureReadyCallback>> Callbacks;		/// The callbacks to call when texture is set
 
 	int32							Id = -1;					/// The ID/index into the list of target texture sets
 	FString							Name;						/// Display name of this target set
