Fixed issues related to 3D preview
- Specific PBR layer is not rendered correctly in the 3D Preview
- Changing roughness does not affect the 3D preview
- 3D viewport showing artefacts when setting it to metallic and low roughness

#virtualized

diff --git a/Engine/Plugins/Experimental/TextureGraph/Source/TextureGraphEngine/Transform/Expressions/T_FlatColorTexture.cpp b/Engine/Plugins/Experimental/TextureGraph/Source/TextureGraphEngine/Transform/Expressions/T_FlatColorTexture.cpp
--- a/Engine/Plugins/Experimental/TextureGraph/Source/TextureGraphEngine/Transform/Expressions/T_FlatColorTexture.cpp
+++ b/Engine/Plugins/Experimental/TextureGraph/Source/TextureGraphEngine/Transform/Expressions/T_FlatColorTexture.cpp
@@ -21,7 +21,8 @@ BufferDescriptor T_FlatColorTexture::GetFlatColorDesc(FString name, BufferFormat
 	Desc.ItemsPerPoint = 4;
 	Desc.Width = 1;
 	Desc.Height = 1;
-
+	Desc.bMipMaps = true;
+	
 	return Desc;
 }
 
@@ -39,7 +40,8 @@ TiledBlobPtr T_FlatColorTexture::Create(MixUpdateCyclePtr InCycle, BufferDescrip
 	Desc.Height = InCycle->GetMix()->GetNumYTiles();
 	Desc.Format = BufferFormat::Byte;
 	Desc.ItemsPerPoint = 4;
-
+	Desc.bMipMaps = true;
+	
 	BufferDescriptor OutputDesc = BufferDescriptor::Combine(Desc, DesiredOutputDesc);
 	if (OutputDesc.Format == BufferFormat::Byte) // Byte size automatically enable sRGB
 	{
@@ -52,7 +54,8 @@ TiledBlobPtr T_FlatColorTexture::Create(MixUpdateCyclePtr InCycle, BufferDescrip
 		 ->AddArg(WithUnbounded(ARG_INT(OutputDesc.ItemsPerPoint,"ItemsPerPoint")))
 		 ->AddArg(WithUnbounded(ARG_INT(OutputDesc.Height, "Height")))
 		 ->AddArg(WithUnbounded(ARG_INT(OutputDesc.Width, "Width")))
-		 ->AddArg(WithUnbounded(ARG_INT((int32)OutputDesc.Format, "Format")));
+		 ->AddArg(WithUnbounded(ARG_INT((int32)OutputDesc.Format, "Format")))
+		 ->AddArg(WithUnbounded(ARG_INT((int32)OutputDesc.bMipMaps, "MipMaps")));
 
 	auto JobResult = RenderJob->InitResult(SimpleColor->GetName(), &OutputDesc);
 	
