Removed redundant log

diff --git a/Engine/Plugins/Experimental/TextureGraph/Source/TextureGraph/Private/TG_Variant.cpp b/Engine/Plugins/Experimental/TextureGraph/Source/TextureGraph/Private/TG_Variant.cpp
--- a/Engine/Plugins/Experimental/TextureGraph/Source/TextureGraph/Private/TG_Variant.cpp
+++ b/Engine/Plugins/Experimental/TextureGraph/Source/TextureGraph/Private/TG_Variant.cpp
@@ -39,7 +39,6 @@ FTG_Variant::FTG_Variant(FTG_Texture RHS)
 bool FTG_Variant::Serialize(FArchive& Ar)
 {
 	Ar << *this;
-	UE_LOG(LogTextureGraph, Log, TEXT("FTG_Variant::Serialize"));
 	return true;
 }
 
