Adding validation and error checking for File name for TG export

diff --git a/Engine/Plugins/Experimental/TextureGraph/Source/TextureGraph/Private/Expressions/Input/TG_Expression_OutputSettings.cpp b/Engine/Plugins/Experimental/TextureGraph/Source/TextureGraph/Private/Expressions/Input/TG_Expression_OutputSettings.cpp
--- a/Engine/Plugins/Experimental/TextureGraph/Source/TextureGraph/Private/Expressions/Input/TG_Expression_OutputSettings.cpp
+++ b/Engine/Plugins/Experimental/TextureGraph/Source/TextureGraph/Private/Expressions/Input/TG_Expression_OutputSettings.cpp
@@ -4,6 +4,7 @@
 
 #include "TG_Graph.h"
 #include "Model/StaticImageResource.h"
+#include "TG_HelperFunctions.h"
 
 void UTG_Expression_OutputSettings::Evaluate(FTG_EvaluationContext* InContext)
 {
@@ -62,9 +63,30 @@ void UTG_Expression_OutputSettings::Evaluate(FTG_EvaluationContext* InContext)
 
 bool UTG_Expression_OutputSettings::Validate(MixUpdateCyclePtr Cycle)
 {
-	UMixInterface* ParentMix = Cast<UMixInterface>(GetOutermostObject());
+	FText Reason;
+	FTG_HelperFunctions::IsFileNameValid(Settings.BaseName, Reason);
 	
-	return true;
+	auto InputPin = GetParentNode()->GetPin("Input");
+	if (!InputPin->IsConnected() && !Reason.IsEmpty())
+	{
+		FString Message = "Invalid file name. " + Reason.ToString();
+		UMixInterface* ParentMix = Cast<UMixInterface>(GetOutermostObject());
+		auto ErrorType = static_cast<int32>(ETextureGraphErrorType::NODE_WARNING);
+		TextureGraphEngine::GetErrorReporter(ParentMix)->ReportWarning(ErrorType, Message, GetParentNode());
+	}
+
+	FText PathReason;
+	FTG_HelperFunctions::IsFolderPathValid(Settings.FolderPath.ToString(), PathReason);
+
+	if (!InputPin->IsConnected() && !PathReason.IsEmpty())
+	{
+		FString Message = "Invalid folder path. " + PathReason.ToString();
+		UMixInterface* ParentMix = Cast<UMixInterface>(GetOutermostObject());
+		auto ErrorType = static_cast<int32>(ETextureGraphErrorType::NODE_WARNING);
+		TextureGraphEngine::GetErrorReporter(ParentMix)->ReportWarning(ErrorType, Message, GetParentNode());
+	}
+
+	return Super::Validate(Cycle);
 }
 void UTG_Expression_OutputSettings::SetTitleName(FName NewName)
 {
diff --git a/Engine/Plugins/Experimental/TextureGraph/Source/TextureGraph/Private/Expressions/Output/TG_Expression_Output.cpp b/Engine/Plugins/Experimental/TextureGraph/Source/TextureGraph/Private/Expressions/Output/TG_Expression_Output.cpp
--- a/Engine/Plugins/Experimental/TextureGraph/Source/TextureGraph/Private/Expressions/Output/TG_Expression_Output.cpp
+++ b/Engine/Plugins/Experimental/TextureGraph/Source/TextureGraph/Private/Expressions/Output/TG_Expression_Output.cpp
@@ -95,6 +95,33 @@ void UTG_Expression_Output::Evaluate(FTG_EvaluationContext* InContext)
 	}
 }
 
+bool UTG_Expression_Output::Validate(MixUpdateCyclePtr	Cycle)
+{
+	FText Reason;
+	FTG_HelperFunctions::IsFileNameValid(OutputSettings.BaseName, Reason);
+	auto OutputSettingsPin = GetParentNode()->GetPin("OutputSettings");
+	if (!OutputSettingsPin->IsConnected() && !Reason.IsEmpty())
+	{
+		FString Message = "Invalid file name. " + Reason.ToString();
+		UMixInterface* ParentMix = Cast<UMixInterface>(GetOutermostObject());
+		auto ErrorType = static_cast<int32>(ETextureGraphErrorType::NODE_WARNING);
+		TextureGraphEngine::GetErrorReporter(ParentMix)->ReportWarning(ErrorType, Message, GetParentNode());
+	}
+
+	FText PathReason;
+	FTG_HelperFunctions::IsFolderPathValid(OutputSettings.FolderPath.ToString(), PathReason);
+
+	if (!OutputSettingsPin->IsConnected() && !PathReason.IsEmpty())
+	{
+		FString Message = "Invalid folder path. " + PathReason.ToString();
+		UMixInterface* ParentMix = Cast<UMixInterface>(GetOutermostObject());
+		auto ErrorType = static_cast<int32>(ETextureGraphErrorType::NODE_WARNING);
+		TextureGraphEngine::GetErrorReporter(ParentMix)->ReportWarning(ErrorType, Message, GetParentNode());
+	}
+
+	return Super::Validate(Cycle);
+}
+
 void UTG_Expression_Output::SetTitleName(FName NewName)
 {
 	GetParentNode()->GetOutputPin(GET_MEMBER_NAME_CHECKED(UTG_Expression_Output, Output))->SetAliasName(NewName);
diff --git a/Engine/Plugins/Experimental/TextureGraph/Source/TextureGraph/Private/TG_HelperFunctions.cpp b/Engine/Plugins/Experimental/TextureGraph/Source/TextureGraph/Private/TG_HelperFunctions.cpp
--- a/Engine/Plugins/Experimental/TextureGraph/Source/TextureGraph/Private/TG_HelperFunctions.cpp
+++ b/Engine/Plugins/Experimental/TextureGraph/Source/TextureGraph/Private/TG_HelperFunctions.cpp
@@ -12,6 +12,24 @@
 #include "Model/Mix/MixManager.h"
 #include "Expressions/TG_Expression.h"
 #include "Transform/Expressions/T_FlatColorTexture.h"
+#include "Logging/MessageLog.h"
+#include "Misc/UObjectToken.h"
+#include "Misc/FileHelper.h"
+
+bool FTG_HelperFunctions::IsFileNameValid(FName FileName, FText& Reason)
+{
+	FileName.IsValidObjectName(Reason);
+	FFileHelper::IsFilenameValidForSaving(FileName.ToString(), Reason);
+
+	return Reason.IsEmpty();
+}
+
+bool FTG_HelperFunctions::IsFolderPathValid(FString FolderPath, FText& Reason)
+{
+	FName::IsValidXName(FolderPath, INVALID_OBJECTPATH_CHARACTERS INVALID_LONGPACKAGE_CHARACTERS, &Reason);
+	FPaths::ValidatePath(FolderPath, &Reason);
+	return Reason.IsEmpty();
+}
 
 TArray<BlobPtr> FTG_HelperFunctions::GetTexturedOutputs(const UTG_Node* Node, FTG_EvaluationContext* TextureConversionContext /*= nullptr*/)
 
@@ -95,6 +113,7 @@ void FTG_HelperFunctions::EnsureOutputIsTexture(MixUpdateCyclePtr Cycle, UTG_Nod
 AsyncBool FTG_HelperFunctions::ExportAsync(UTextureGraph* InTextureGraph, FString ExportPath, FString AssetName, FExportSettings& TargetExportSettings, bool OverrideExportPath, bool OverwriteTextures /*= true*/, bool ExportAllOutputs /*= false*/,bool bSave /*= true*/)
 {
 	FString ErrorMessage = "";
+	bool AnyValidExport = false;
 	TargetExportSettings.Reset();
 	
 	FInvalidationDetails Details;
@@ -106,7 +125,7 @@ AsyncBool FTG_HelperFunctions::ExportAsync(UTextureGraph* InTextureGraph, FStrin
 	MixUpdateCyclePtr Cycle = Batch->GetCycle();
 	InTextureGraph->Update(Cycle);
 	
-	InTextureGraph->Graph()->ForEachNodes([=,&TargetExportSettings](const UTG_Node* Node, uint32 Index)
+	InTextureGraph->Graph()->ForEachNodes([=,&TargetExportSettings, &ErrorMessage , &AnyValidExport](const UTG_Node* Node, uint32 Index)
 	{
 		UTG_Expression_Output* TargetExpression = Cast<UTG_Expression_Output>(Node->GetExpression());
 		if (TargetExpression)
@@ -126,17 +145,23 @@ AsyncBool FTG_HelperFunctions::ExportAsync(UTextureGraph* InTextureGraph, FStrin
 					Path = ExportPath;
 				}
 
-				FString Name = AssetName.IsEmpty() ? OutputSetting.GetFullOutputName() : AssetName;
+				FString FileName = AssetName.IsEmpty() ? OutputSetting.GetFullOutputName() : AssetName;
 
-				bool IsNameValid = TextureExporter::IsPackageNameValid(Path, Name);
+				FText Reason;
+				bool IsNameValid = IsFileNameValid(FName(*FileName), Reason);
+
+				FText PathReason;
+				bool IsPathValid = TextureExporter::IsPackageNameValid(Path, FileName);
+				IsFolderPathValid(Path, PathReason);
+				IsPathValid &= PathReason.IsEmpty();
 
 				bool HasOutputs = ExportBlobs.Num() > 0;
 
-				if (HasOutputs && IsNameValid)
+				if (HasOutputs && IsNameValid && IsPathValid)
 				{
 					TiledBlobPtr Output = std::static_pointer_cast<TiledBlob>(ExportBlobs[0]);//Dealing with one output per Node for now
-					FExportMapSettings MapSettings = TextureExporter::GetExportSettingsForTarget(TargetExportSettings, std::static_pointer_cast<TiledBlob>(Output), *Name);
-					MapSettings.Name = FName(*Name);
+					FExportMapSettings MapSettings = TextureExporter::GetExportSettingsForTarget(TargetExportSettings, std::static_pointer_cast<TiledBlob>(Output), *FileName);
+					MapSettings.Name = FName(*FileName);
 					MapSettings.Path = Path;
 					MapSettings.UseOverridePath = OverrideExportPath;
 					MapSettings.OverwriteTextures = OverwriteTextures;
@@ -147,25 +172,47 @@ AsyncBool FTG_HelperFunctions::ExportAsync(UTextureGraph* InTextureGraph, FStrin
 					MapSettings.Height = (int32)OutputSetting.Height;
 					MapSettings.bSave = bSave;
 					TargetExportSettings.ExportPreset.push_back(std::pair<FName, FExportMapSettings>{ MapSettings.Name, MapSettings });
+					AnyValidExport = true;
 				}
 				else
 				{
 					//Log Error to Error System
 					if (!HasOutputs)
 					{
-						auto ErrorType = static_cast<int32>(ETextureGraphErrorType::UNSUPPORTED_TYPE);
-						TextureGraphEngine::GetErrorReporter(InTextureGraph)->ReportError(ErrorType, FString::Format(TEXT("Texture Export Error : No valid output found for OutputSetting {0}"), { OutputSetting.OutputName.ToString() }), nullptr);
+						ErrorMessage += FString::Format(TEXT("Texture Export Error : No valid output found for OutputSetting {0}"), { OutputSetting.OutputName.ToString() });
+						ErrorMessage += "\n";
+					}
+					if (!IsPathValid)
+					{
+						ErrorMessage += FString::Format(TEXT("Texture Export Error : Invalid path set for OutputSetting {0}"), { OutputSetting.OutputName.ToString() }) + (PathReason.IsEmpty() ? "" : " Folder " + PathReason.ToString());
+						ErrorMessage += "\n";
 					}
 					if (!IsNameValid)
 					{
-						auto ErrorType = static_cast<int32>(ETextureGraphErrorType::UNSUPPORTED_TYPE);
-						TextureGraphEngine::GetErrorReporter(InTextureGraph)->ReportError(ErrorType, FString::Format(TEXT("Texture Export Error : Invalid path set for OutputSetting {0}"), { OutputSetting.OutputName.ToString() }), nullptr);
+						FString ReasonString = Reason.IsEmpty() ? "Invalid File Name" : Reason.ToString();
+						ErrorMessage += FString::Format(TEXT("Texture Export Error : {0} for OutputSetting {1}"), { ReasonString , OutputSetting.OutputName.ToString()});
+						ErrorMessage += "\n";
 					}
 				}
 			}
 		}
 	});
 
+	if (!ErrorMessage.IsEmpty())
+	{
+		ErrorMessage = ErrorMessage.LeftChop(1);
+		auto ErrorType = static_cast<int32>(ETextureGraphErrorType::UNSUPPORTED_TYPE);
+		TextureGraphEngine::GetErrorReporter(InTextureGraph)->ReportError(ErrorType, ErrorMessage, nullptr);
+
+		FMessageLog("PIE").Error()
+			->AddToken(FTextToken::Create(FText::FromString(ErrorMessage)));
+
+		if (!AnyValidExport)
+		{
+			return cti::make_ready_continuable(false);
+		}
+	}
+
 	return RenderAsync(InTextureGraph, Batch)
 	.then([InTextureGraph, &TargetExportSettings, ExportPath](auto result) 
 	{
diff --git a/Engine/Plugins/Experimental/TextureGraph/Source/TextureGraph/Public/Expressions/Output/TG_Expression_Output.h b/Engine/Plugins/Experimental/TextureGraph/Source/TextureGraph/Public/Expressions/Output/TG_Expression_Output.h
--- a/Engine/Plugins/Experimental/TextureGraph/Source/TextureGraph/Public/Expressions/Output/TG_Expression_Output.h
+++ b/Engine/Plugins/Experimental/TextureGraph/Source/TextureGraph/Public/Expressions/Output/TG_Expression_Output.h
@@ -18,6 +18,8 @@ class TEXTUREGRAPH_API UTG_Expression_Output : public UTG_Expression
 public:
 	TG_DECLARE_VARIANT_EXPRESSION(TG_Category::Output);
 	virtual void Evaluate(FTG_EvaluationContext* InContext) override;
+
+	virtual bool Validate(MixUpdateCyclePtr Cycle) override;
 	
 	// The final generated output
 	UPROPERTY(meta = (TGType = "TG_Input", PinDisplayName = ""))
@@ -30,7 +32,7 @@ public:
 	//When we will work on Node UI for FOutputSettings we will set the category as TG_Setting
 	UPROPERTY(EditAnywhere, Category = NoCategory, meta = (TGType = "TG_Input", CollapsableChildProperties,ShowOnlyInnerProperties, FullyExpand, NoResetToDefault, PinDisplayName = "Settings") )
 	FTG_OutputSettings OutputSettings;
-	
+
 	virtual FTG_Name GetDefaultName() const override { return TEXT("Output");}
 	virtual void SetTitleName(FName NewName) override;
 	virtual FName GetTitleName() const override;
diff --git a/Engine/Plugins/Experimental/TextureGraph/Source/TextureGraph/Public/TG_HelperFunctions.h b/Engine/Plugins/Experimental/TextureGraph/Source/TextureGraph/Public/TG_HelperFunctions.h
--- a/Engine/Plugins/Experimental/TextureGraph/Source/TextureGraph/Public/TG_HelperFunctions.h
+++ b/Engine/Plugins/Experimental/TextureGraph/Source/TextureGraph/Public/TG_HelperFunctions.h
@@ -15,6 +15,9 @@ class TEXTUREGRAPH_API FTG_HelperFunctions
 {
 
 public:
+	static bool IsFileNameValid(FName FileName, FText& Reason);
+	static bool IsFolderPathValid(FString FolderPath, FText& Reason);
+
 	static void	InitTargets(UTextureGraph* InTextureGraph);
 	static TArray<BlobPtr> GetTexturedOutputs(const UTG_Node* Node, FTG_EvaluationContext* TextureConversionContext = nullptr);
 
@@ -22,8 +25,6 @@ public:
 
 	static AsyncBool ExportAsync(UTextureGraph* InTextureGraph, FString ExportPath, FString AssetName, FExportSettings& TargetExportSettings, bool OverrideExportPath, bool OverwriteTextures = true,bool ExportAllOutputs = false, bool bSave = true);
 	static AsyncBool RenderAsync(UTextureGraph* InTextureGraph, JobBatchPtr ExistingBatch = nullptr);
-
-
 	
 	template <typename T_Type>
 	static TArray<T_Type> GetOutputsOfType(const UTG_Node* Node)
