Fixed issue: Moving the Blur node causes the link between the scalar node and the Radius parameter to be lost

Temp fix - Commenting this code for now. Calling UpdateInputPinsVisibility on selection is causing disconnection of the pins. Disabling this code will cause the UE-215729 bug. But that only affects the details panel so for now this is fine.

diff --git a/Engine/Plugins/Experimental/TextureGraph/Source/TextureGraphEditor/Private/TG_Editor.cpp b/Engine/Plugins/Experimental/TextureGraph/Source/TextureGraphEditor/Private/TG_Editor.cpp
--- a/Engine/Plugins/Experimental/TextureGraph/Source/TextureGraphEditor/Private/TG_Editor.cpp
+++ b/Engine/Plugins/Experimental/TextureGraph/Source/TextureGraphEditor/Private/TG_Editor.cpp
@@ -1201,12 +1201,19 @@ void FTG_Editor::OnSelectedNodesChanged(const TSet<class UObject*>& NewSelection
 		{
 			SelectedObjects.Add(GraphNode->GetDetailsObject());
 
+			// Temp fix - Commenting this code for now.
+			// Calling UpdateInputPinsVisibility on selection is causing disconnection of the pins.
+			// Disabling this code will cause UE-215729 bug.
+			// But that only effects details panel so for now this is fine.
+			// Problem: Visibility and interaction of the pin are being updated from this code which are dependent on each other
+			// Ideally they shouldnt be dependent on each other.
+			
 			// calling this here to internally set EditCondition metadata properly for Details panel
 			// this is driven from the CanEditChange() method of the expression which can be customized as required
-			GraphNode->UpdateInputPinsVisibility();
-		}
+			// GraphNode->UpdateInputPinsVisibility();
+		}  
 	}
-
+	
 	GetDetailView()->SetObjects(SelectedObjects, true);
 	FocusDetailsPanel();
 
