Added support for string parameters and loading textures from path.

diff --git a/Engine/Plugins/Experimental/TextureGraph/Source/TextureGraph/Private/Blueprint/TG_BlueprintFunctionLibrary.cpp b/Engine/Plugins/Experimental/TextureGraph/Source/TextureGraph/Private/Blueprint/TG_BlueprintFunctionLibrary.cpp
--- a/Engine/Plugins/Experimental/TextureGraph/Source/TextureGraph/Private/Blueprint/TG_BlueprintFunctionLibrary.cpp
+++ b/Engine/Plugins/Experimental/TextureGraph/Source/TextureGraph/Private/Blueprint/TG_BlueprintFunctionLibrary.cpp
@@ -11,9 +11,12 @@
 #include "Expressions/Input/TG_Expression_Scalar.h"
 #include "Expressions/Input/TG_Expression_Vector.h"
 #include "Expressions/Input/TG_Expression_Color.h"
+#include "Expressions/Input/TG_Expression_Bool.h"
+#include "Expressions/Input/TG_Expression_String.h"
 #include "Expressions/Input/TG_Expression_OutputSettings.h"
 #include "TG_HelperFunctions.h"
 #include "Engine/Texture2D.h"
+#include "Job/Scheduler.h"
 
 #include UE_INLINE_GENERATED_CPP_BY_NAME(TG_BlueprintFunctionLibrary)
 
@@ -22,217 +25,146 @@
 
 #define LOCTEXT_NAMESPACE "TG_BlueprintFunctionLibrary"
 
-void UTG_BlueprintFunctionLibrary::SetTextureParameterValue(UObject* WorldContextObject, UTextureGraph* InTextureGraph, FName ParameterName, UTexture* ParameterValue)
+UTG_Pin* GetParamPin(UObject* WorldContextObject, UTextureGraph* InTextureGraph, FName ParameterName)
 {
 	if (UWorld* World = GEngine->GetWorldFromContextObject(WorldContextObject, EGetWorldErrorMode::LogAndReturnNull))
 	{
-		auto bFoundParameter = false;
 		if (InTextureGraph)
 		{
-			auto PinParam = InTextureGraph->Graph()->FindParamPin(ParameterName);
-			if (PinParam)
-			{
-				auto ExpressionPtr = Cast<UTG_Expression_Texture>(PinParam->GetNodePtr()->GetExpression());
-				if (ExpressionPtr)
-				{
-					bFoundParameter = true;
-					UTexture2D* DupTexture = (UTexture2D*)StaticDuplicateObject(ParameterValue, GetTransientPackage(), NAME_None, ~RF_Standalone, UTexture2D::StaticClass());
-					ExpressionPtr->SetAsset(DupTexture);
-				}
-			}
-		}
-
-		if (!bFoundParameter)
-		{
-			AddParamWarning(ParameterName, InTextureGraph, "SetTextureParameterValue");
+			return InTextureGraph->Graph()->FindParamPin(ParameterName);
 		}
 	}
+
+	return nullptr;
 }
 
-UTexture* UTG_BlueprintFunctionLibrary::GetTextureParameterValue(UObject* WorldContextObject, UTextureGraph* InTextureGraph, FName ParameterName)
+template <typename T_Expr>
+T_Expr* GetParamExpression(UObject* WorldContextObject, UTextureGraph* InTextureGraph, FName ParameterName)
 {
-	UTexture* ParameterValue = nullptr;
-	
-	if (UWorld* World = GEngine->GetWorldFromContextObject(WorldContextObject, EGetWorldErrorMode::LogAndReturnNull))
+	UTG_Pin* PinParam = GetParamPin(WorldContextObject, InTextureGraph, ParameterName);
+	if (PinParam)
 	{
-		auto bFoundParameter = false;
-		if (InTextureGraph)
-		{
-			auto PinParam = InTextureGraph->Graph()->FindParamPin(ParameterName);
-			if (PinParam)
-			{
-				auto ExpressionPtr = Cast<UTG_Expression_Texture>(PinParam->GetNodePtr()->GetExpression());
-				if (ExpressionPtr)
-				{
-					bFoundParameter = true;
-					ParameterValue = ExpressionPtr->Source;
-				}
-			}
-		}
-
-		if (!bFoundParameter)
-		{
-			AddParamWarning(ParameterName, InTextureGraph,"GetTextureParameterValue");
-		}
+		return Cast<T_Expr>(PinParam->GetNodePtr()->GetExpression());
 	}
 
-	return ParameterValue;
+	return nullptr;
 }
 
-void UTG_BlueprintFunctionLibrary::SetScalarParameterValue(UObject* WorldContextObject, UTextureGraph* InTextureGraph, FName ParameterName, float ParameterValue)
+template <typename T_Expr, typename T_Expr_Value>
+void SetParameterValue_Generic(UObject* WorldContextObject, UTextureGraph* InTextureGraph, FName ParameterName, T_Expr_Value ParameterValue, FString FunctionName)
 {
-	if (UWorld* World = GEngine->GetWorldFromContextObject(WorldContextObject, EGetWorldErrorMode::LogAndReturnNull))
+	UTG_Pin* PinParam = GetParamPin(WorldContextObject, InTextureGraph, ParameterName);
+	if (PinParam)
 	{
-		auto bFoundParameter = false;
-		if (InTextureGraph)
-		{
-			auto PinParam = InTextureGraph->Graph()->FindParamPin(ParameterName);
-			if (PinParam)
-			{
-				auto ExpressionPtr = Cast<UTG_Expression_Scalar>(PinParam->GetNodePtr()->GetExpression());
-				if (ExpressionPtr)
-				{
-					ExpressionPtr->Scalar = ParameterValue;
-					PinParam->SetValue(ParameterValue);
-					bFoundParameter = true;
-				}
-			}
-		}
-
-		if (!bFoundParameter)
-		{
-			AddParamWarning(ParameterName, InTextureGraph, "SetScalarParameterValue");
-		}
+		PinParam->SetValue(ParameterValue);
+	}
+	else
+	{
+		UTG_BlueprintFunctionLibrary::AddParamWarning(ParameterName, InTextureGraph, FunctionName);
 	}
 }
 
-float UTG_BlueprintFunctionLibrary::GetScalarParameterValue(UObject* WorldContextObject, UTextureGraph* InTextureGraph, FName ParameterName)
+template <typename T_Expr_Value>
+T_Expr_Value GetParameterValue_Generic(UObject* WorldContextObject, UTextureGraph* InTextureGraph, FName ParameterName, FString FunctionName, T_Expr_Value& OutValue)
 {
-	float ParameterValue = 0.0f;
+	UTG_Pin* PinParam = GetParamPin(WorldContextObject, InTextureGraph, ParameterName);
 
-	if (UWorld* World = GEngine->GetWorldFromContextObject(WorldContextObject, EGetWorldErrorMode::LogAndReturnNull))
+	if (PinParam)
 	{
-		auto bFoundParameter = false;
-		if (InTextureGraph)
-		{
-			auto PinParam = InTextureGraph->Graph()->FindParamPin(ParameterName);
-			if (PinParam)
-			{
-				bFoundParameter = PinParam->GetValue(ParameterValue);
-			}
-		}
-
-		if (!bFoundParameter)
-		{
-			AddParamWarning(ParameterName, InTextureGraph, "GetScalarParameterValue");
-		}
+		PinParam->GetValue(OutValue);
+	}
+	else
+	{
+		UTG_BlueprintFunctionLibrary::AddParamWarning(ParameterName, InTextureGraph, FunctionName);
 	}
 
-	return ParameterValue;
+	return OutValue;
 }
 
-void UTG_BlueprintFunctionLibrary::SetVectorParameterValue(UObject* WorldContextObject, UTextureGraph* InTextureGraph, FName ParameterName, FVector4f ParameterValue)
+
+void UTG_BlueprintFunctionLibrary::SetTextureParameterValue(UObject* WorldContextObject, UTextureGraph* InTextureGraph, FName ParameterName, UTexture* ParameterValue)
 {
-	if (UWorld* World = GEngine->GetWorldFromContextObject(WorldContextObject, EGetWorldErrorMode::LogAndReturnNull))
+	UTG_Expression_Texture* ExpressionPtr = GetParamExpression<UTG_Expression_Texture>(WorldContextObject, InTextureGraph, ParameterName);
+	if (ExpressionPtr)
 	{
-		auto bFoundParameter = false;
-		if (InTextureGraph)
-		{
-			auto PinParam = InTextureGraph->Graph()->FindParamPin(ParameterName);
-			if (PinParam)
-			{
-				auto ExpressionPtr = Cast<UTG_Expression_Vector>(PinParam->GetNodePtr()->GetExpression());
-				if (ExpressionPtr)
-				{
-					ExpressionPtr->Vector = ParameterValue;
-					PinParam->SetValue(ParameterValue);
-					bFoundParameter = true;
-				}
-			}
-		}
-
-		if (!bFoundParameter)
-		{
-			AddParamWarning(ParameterName, InTextureGraph, "SetVectorParameterValue");
-		}
+		UTexture2D* DupTexture = (UTexture2D*)StaticDuplicateObject(ParameterValue, GetTransientPackage(), NAME_None, ~RF_Standalone, UTexture2D::StaticClass());
+		ExpressionPtr->SetAsset(DupTexture);
+	}
+	else
+	{
+		AddParamWarning(ParameterName, InTextureGraph, "SetTextureParameterValue");
 	}
 }
 
-FVector4f UTG_BlueprintFunctionLibrary::GetVectorParameterValue(UObject* WorldContextObject, UTextureGraph* InTextureGraph, FName ParameterName)
+UTexture* UTG_BlueprintFunctionLibrary::GetTextureParameterValue(UObject* WorldContextObject, UTextureGraph* InTextureGraph, FName ParameterName)
 {
-	FVector4f ParameterValue = FVector4f::Zero();
+	UTG_Expression_Texture* ExpressionPtr = GetParamExpression<UTG_Expression_Texture>(WorldContextObject, InTextureGraph, ParameterName);
 
-	if (UWorld* World = GEngine->GetWorldFromContextObject(WorldContextObject, EGetWorldErrorMode::LogAndReturnNull))
+	if (ExpressionPtr)
 	{
-		auto bFoundParameter = false;
-		if (InTextureGraph)
-		{
-			auto PinParam = InTextureGraph->Graph()->FindParamPin(ParameterName);
-			if (PinParam)
-			{
-				bFoundParameter = PinParam->GetValue(ParameterValue);
-			}
-		}
-
-		if (!bFoundParameter)
-		{
-			AddParamWarning(ParameterName, InTextureGraph, "GetVectorParameterValue");
-		}
+		return ExpressionPtr->Source;
+	}
+	else
+	{
+		AddParamWarning(ParameterName, InTextureGraph, "SetTextureParameterValue");
 	}
 
-	return ParameterValue;
+	return nullptr;
 }
 
-void UTG_BlueprintFunctionLibrary::SetColorParameterValue(UObject* WorldContextObject, UTextureGraph* InTextureGraph, FName ParameterName, FLinearColor ParameterValue)
+void UTG_BlueprintFunctionLibrary::SetScalarParameterValue(UObject* WorldContextObject, UTextureGraph* InTextureGraph, FName ParameterName, float ParameterValue)
 {
-	if (UWorld* World = GEngine->GetWorldFromContextObject(WorldContextObject, EGetWorldErrorMode::LogAndReturnNull))
-	{
-		auto bFoundParameter = false;
-		if (InTextureGraph)
-		{
-			auto PinParam = InTextureGraph->Graph()->FindParamPin(ParameterName);
-			if (PinParam)
-			{
-				auto ExpressionPtr = Cast<UTG_Expression_Color>(PinParam->GetNodePtr()->GetExpression());
-				if (ExpressionPtr)
-				{
-					ExpressionPtr->Color = ParameterValue;
-					PinParam->SetValue(ParameterValue);
-					bFoundParameter = true;
-				}
-			}
-		}
+	SetParameterValue_Generic<UTG_Expression_Scalar, float>(WorldContextObject, InTextureGraph, ParameterName, ParameterValue, "SetScalarParameterValue");
+}
 
-		if (!bFoundParameter)
-		{
-			AddParamWarning(ParameterName, InTextureGraph, "SetColorParameterValue");
-		}
-	}
+float UTG_BlueprintFunctionLibrary::GetScalarParameterValue(UObject* WorldContextObject, UTextureGraph* InTextureGraph, FName ParameterName)
+{
+	float ParameterValue = 0.0f;
+	return GetParameterValue_Generic<float>(WorldContextObject, InTextureGraph, ParameterName, "GetScalarParameterValue", ParameterValue);
+}
+
+void UTG_BlueprintFunctionLibrary::SetVectorParameterValue(UObject* WorldContextObject, UTextureGraph* InTextureGraph, FName ParameterName, FVector4f ParameterValue)
+{
+	SetParameterValue_Generic<UTG_Expression_Vector, FVector4f>(WorldContextObject, InTextureGraph, ParameterName, ParameterValue, "SetVectorParameterValue");
+}
+
+FVector4f UTG_BlueprintFunctionLibrary::GetVectorParameterValue(UObject* WorldContextObject, UTextureGraph* InTextureGraph, FName ParameterName)
+{
+	FVector4f ParameterValue = FVector4f::Zero();
+	return GetParameterValue_Generic<FVector4f>(WorldContextObject, InTextureGraph, ParameterName, "GetVectorParameterValue", ParameterValue);
+}
+
+void UTG_BlueprintFunctionLibrary::SetColorParameterValue(UObject* WorldContextObject, UTextureGraph* InTextureGraph, FName ParameterName, FLinearColor ParameterValue)
+{
+	SetParameterValue_Generic<UTG_Expression_Color, FLinearColor>(WorldContextObject, InTextureGraph, ParameterName, ParameterValue, "SetColorParameterValue");
 }
 
 FLinearColor UTG_BlueprintFunctionLibrary::GetColorParameterValue(UObject* WorldContextObject, UTextureGraph* InTextureGraph, FName ParameterName)
 {
 	FLinearColor ParameterValue = FLinearColor::Black;
+	return GetParameterValue_Generic<FLinearColor>(WorldContextObject, InTextureGraph, ParameterName, "GetColorParameterValue", ParameterValue);
+}
 
-	if (UWorld* World = GEngine->GetWorldFromContextObject(WorldContextObject, EGetWorldErrorMode::LogAndReturnNull))
-	{
-		auto bFoundParameter = false;
-		if (InTextureGraph)
-		{
-			auto PinParam = InTextureGraph->Graph()->FindParamPin(ParameterName);
-			if (PinParam)
-			{
-				bFoundParameter = PinParam->GetValue(ParameterValue);
-			}
-		}
+void UTG_BlueprintFunctionLibrary::SetBoolParameterValue(UObject* WorldContextObject, UTextureGraph* InTextureGraph, FName ParameterName, bool bParameterValue)
+{
+	SetParameterValue_Generic<UTG_Expression_Bool, bool>(WorldContextObject, InTextureGraph, ParameterName, bParameterValue, "SetBoolParameterValue");
+}
 
-		if (!bFoundParameter)
-		{
-			AddParamWarning(ParameterName, InTextureGraph, "GetColorParameterValue");
-		}
-	}
+bool UTG_BlueprintFunctionLibrary::GetBoolParameterValue(UObject* WorldContextObject, UTextureGraph* InTextureGraph, FName ParameterName)
+{
+	bool bParameterValue = false;
+	return GetParameterValue_Generic<bool>(WorldContextObject, InTextureGraph, ParameterName, "GetBoolParameterValue", bParameterValue);
+}
 
-	return ParameterValue;
+void UTG_BlueprintFunctionLibrary::SetStringParameterValue(UObject* WorldContextObject, UTextureGraph* InTextureGraph, FName ParameterName, FString ParameterValue)
+{
+	SetParameterValue_Generic<UTG_Expression_String, FString>(WorldContextObject, InTextureGraph, ParameterName, ParameterValue, "SetStringParameterValue");
+}
+
+FString UTG_BlueprintFunctionLibrary::GetStringParameterValue(UObject* WorldContextObject, UTextureGraph* InTextureGraph, FName ParameterName)
+{
+	FString ParameterValue;
+	return GetParameterValue_Generic<FString>(WorldContextObject, InTextureGraph, ParameterName, "GetStringParameterValue", ParameterValue);
 }
 
 void UTG_BlueprintFunctionLibrary::SetSettingsParameterValue(UObject* WorldContextObject, UTextureGraph* InTextureGraph, FName ParameterName, int Width, int Height, 
@@ -320,6 +252,41 @@ FTG_OutputSettings UTG_BlueprintFunctionLibrary::GetSettingsParameterValue(UObje
 	return ParameterValue;
 }
 
+#if 0 /// This requires a lot of additional work. Will attempt after summer break
+TArray<UTextureRenderTarget2D*> UTG_BlueprintFunctionLibrary::RenderTextureGraph(UObject* WorldContextObject, UTextureGraph* InTextureGraph)
+{
+	check(IsInGameThread());
+	TArray<UTextureRenderTarget2D*> Results;
+	UTextureGraph* TextureGraphPtr = (UTextureGraph*)StaticDuplicateObject(InTextureGraph, GetTransientPackage(), NAME_None, ~RF_Standalone, UTextureGraph::StaticClass());
+	FTG_HelperFunctions::InitTargets(TextureGraphPtr);
+
+	JobBatchPtr Batch = FTG_HelperFunctions::InitRenderBatch(TextureGraphPtr);
+	if (!Batch)
+		return Results;
+
+	/// Add to the scheduler
+	TextureGraphEngine::GetScheduler()->AddBatch(Batch);
+
+	std::atomic_bool* bIsMixRendered = new std::atomic_bool(false);
+	Batch->OnDone([=, &bIsMixRendered](JobBatch*) mutable
+	{
+		*bIsMixRendered = true;
+	});
+
+	while (!*bIsMixRendered)
+	{
+		if (!TextureGraphEngine::IsDestroying())
+			TextureGraphEngine::Update(0);
+	}
+
+	delete bIsMixRendered;
+
+	//FTG_HelperFunctions::RenderAsync(TextureGraphPtr);
+	//TextureGraphEngine::GetDeviceManager()->WaitForQueuedTasks(ENamedThreads::GameThread);
+	return Results;
+}
+#endif 
+
 void UTG_BlueprintFunctionLibrary::AddParamWarning(FName ParamName, UObject* ObjectPtr, FString FunctionName)
 {
 	FFormatNamedArguments Arguments;
diff --git a/Engine/Plugins/Experimental/TextureGraph/Source/TextureGraph/Private/Expressions/Input/TG_Expression_Bool.cpp b/Engine/Plugins/Experimental/TextureGraph/Source/TextureGraph/Private/Expressions/Input/TG_Expression_Bool.cpp
--- a/Engine/Plugins/Experimental/TextureGraph/Source/TextureGraph/Private/Expressions/Input/TG_Expression_Bool.cpp
+++ b/Engine/Plugins/Experimental/TextureGraph/Source/TextureGraph/Private/Expressions/Input/TG_Expression_Bool.cpp
@@ -3,25 +3,6 @@
 #include "Expressions/Input/TG_Expression_Bool.h"
 
 
-FTG_SignaturePtr UTG_Expression_Bool::BuildInputParameterSignature() const
-{
-	FTG_Signature::FInit SignatureInit = GetSignatureInitArgsFromClass();
-	return MakeShared<FTG_Signature>(SignatureInit);
-};
-
-FTG_SignaturePtr UTG_Expression_Bool::BuildInputConstantSignature() const
-{
-	FTG_Signature::FInit SignatureInit = GetSignatureInitArgsFromClass();
-	for (auto& arg : SignatureInit.Arguments)
-	{
-		if (arg.IsInput() && arg.IsParam())
-		{
-			arg.ArgumentType = arg.ArgumentType.Unparamed();
-			arg.ArgumentType.SetNotConnectable();
-		}
-	}
-	return MakeShared<FTG_Signature>(SignatureInit);
-};
 
 void UTG_Expression_Bool::Evaluate(FTG_EvaluationContext* InContext)
 {
diff --git a/Engine/Plugins/Experimental/TextureGraph/Source/TextureGraph/Private/Expressions/Input/TG_Expression_Color.cpp b/Engine/Plugins/Experimental/TextureGraph/Source/TextureGraph/Private/Expressions/Input/TG_Expression_Color.cpp
--- a/Engine/Plugins/Experimental/TextureGraph/Source/TextureGraph/Private/Expressions/Input/TG_Expression_Color.cpp
+++ b/Engine/Plugins/Experimental/TextureGraph/Source/TextureGraph/Private/Expressions/Input/TG_Expression_Color.cpp
@@ -2,28 +2,6 @@
 
 #include "expressions/Input/TG_Expression_Color.h"
 
-
-
-FTG_SignaturePtr UTG_Expression_Color::BuildInputParameterSignature() const
-{
-	FTG_Signature::FInit SignatureInit = GetSignatureInitArgsFromClass();
-	return MakeShared<FTG_Signature>(SignatureInit);
-};
-
-FTG_SignaturePtr UTG_Expression_Color::BuildInputConstantSignature() const
-{
-	FTG_Signature::FInit SignatureInit = GetSignatureInitArgsFromClass();
-	for (auto& arg : SignatureInit.Arguments)
-	{
-		if (arg.IsInput() && arg.IsParam())
-		{
-			arg.ArgumentType = arg.ArgumentType.Unparamed();
-			arg.ArgumentType.SetNotConnectable();
-		}
-	}
-	return MakeShared<FTG_Signature>(SignatureInit);
-};
-
 void UTG_Expression_Color::Evaluate(FTG_EvaluationContext* InContext)
 {
 	Super::Evaluate(InContext);
diff --git a/Engine/Plugins/Experimental/TextureGraph/Source/TextureGraph/Private/Expressions/Input/TG_Expression_InputParam.cpp b/Engine/Plugins/Experimental/TextureGraph/Source/TextureGraph/Private/Expressions/Input/TG_Expression_InputParam.cpp
--- a/Engine/Plugins/Experimental/TextureGraph/Source/TextureGraph/Private/Expressions/Input/TG_Expression_InputParam.cpp
+++ b/Engine/Plugins/Experimental/TextureGraph/Source/TextureGraph/Private/Expressions/Input/TG_Expression_InputParam.cpp
@@ -9,7 +9,7 @@ void UTG_Expression_InputParam::PostEditChangeProperty(FPropertyChangedEvent& Pr
 	// if Graph changes catch it first
 	if (PropertyChangedEvent.GetPropertyName() == FName(TEXT("IsConstant")))
 	{
-		UE_LOG(LogTextureGraph, Log, TEXT("InputParam  Expression Parameter/Constant PostEditChangeProperty."));
+		UE_LOG(LogTextureGraph, VeryVerbose, TEXT("InputParam  Expression Parameter/Constant PostEditChangeProperty."));
 		NotifySignatureChanged();
 	}
 
@@ -17,6 +17,26 @@ void UTG_Expression_InputParam::PostEditChangeProperty(FPropertyChangedEvent& Pr
 }
 #endif
 
+FTG_SignaturePtr UTG_Expression_InputParam::BuildInputParameterSignature() const
+{
+	FTG_Signature::FInit SignatureInit = GetSignatureInitArgsFromClass();
+	return MakeShared<FTG_Signature>(SignatureInit);
+};
+
+FTG_SignaturePtr UTG_Expression_InputParam::BuildInputConstantSignature() const
+{
+	FTG_Signature::FInit SignatureInit = GetSignatureInitArgsFromClass();
+	for (auto& Arg : SignatureInit.Arguments)
+	{
+		if (Arg.IsInput() && Arg.IsParam())
+		{
+			Arg.ArgumentType = Arg.ArgumentType.Unparamed();
+			Arg.ArgumentType.SetNotConnectable();
+		}
+	}
+	return MakeShared<FTG_Signature>(SignatureInit);
+};
+
 void UTG_Expression_InputParam::SetbIsConstant(bool InIsConstant)
 {
 	if (bIsConstant != InIsConstant)
diff --git a/Engine/Plugins/Experimental/TextureGraph/Source/TextureGraph/Private/Expressions/Input/TG_Expression_Material.cpp b/Engine/Plugins/Experimental/TextureGraph/Source/TextureGraph/Private/Expressions/Input/TG_Expression_Material.cpp
--- a/Engine/Plugins/Experimental/TextureGraph/Source/TextureGraph/Private/Expressions/Input/TG_Expression_Material.cpp
+++ b/Engine/Plugins/Experimental/TextureGraph/Source/TextureGraph/Private/Expressions/Input/TG_Expression_Material.cpp
@@ -8,14 +8,14 @@ void UTG_Expression_Material::PostEditChangeProperty(FPropertyChangedEvent& Prop
 	// First catch if Material changes
 	if (PropertyChangedEvent.GetPropertyName() == GET_MEMBER_NAME_CHECKED(UTG_Expression_Material, Material))
 	{
-		UE_LOG(LogTextureGraph, Log, TEXT("Material Expression PostEditChangeProperty."));
+		UE_LOG(LogTextureGraph, VeryVerbose, TEXT("Material Expression PostEditChangeProperty."));
 		SetMaterialInternal(Material);
 		FeedbackPinValue(GET_MEMBER_NAME_CHECKED(UTG_Expression_Material, RenderedAttribute), RenderedAttribute);
 	}
 	// Second catch if AttributeName changes
 	if (PropertyChangedEvent.GetPropertyName() == GET_MEMBER_NAME_CHECKED(UTG_Expression_Material, RenderedAttribute))
 	{
-		UE_LOG(LogTextureGraph, Log, TEXT("Material Expression PostEditChangeProperty."));
+		UE_LOG(LogTextureGraph, VeryVerbose, TEXT("Material Expression PostEditChangeProperty."));
 		SetRenderedAttribute(RenderedAttribute);
 	}
 
diff --git a/Engine/Plugins/Experimental/TextureGraph/Source/TextureGraph/Private/Expressions/Input/TG_Expression_MaterialFunction.cpp b/Engine/Plugins/Experimental/TextureGraph/Source/TextureGraph/Private/Expressions/Input/TG_Expression_MaterialFunction.cpp
--- a/Engine/Plugins/Experimental/TextureGraph/Source/TextureGraph/Private/Expressions/Input/TG_Expression_MaterialFunction.cpp
+++ b/Engine/Plugins/Experimental/TextureGraph/Source/TextureGraph/Private/Expressions/Input/TG_Expression_MaterialFunction.cpp
@@ -39,7 +39,7 @@ void UTG_Expression_MaterialFunction::PostEditChangeProperty(FPropertyChangedEve
 	// First catch if Material changes
 	if (PropertyChangedEvent.GetPropertyName() == FName(TEXT("MaterialFunction")) && !(PropertyChangedEvent.ChangeType == EPropertyChangeType::Unspecified))
 	{
-		UE_LOG(LogTextureGraph, Log, TEXT("Material Function Expression PostEditChangeProperty."));
+		UE_LOG(LogTextureGraph, VeryVerbose, TEXT("Material Function Expression PostEditChangeProperty."));
 		SetMaterialFunctionInternal(MaterialFunction);
 	}
 
diff --git a/Engine/Plugins/Experimental/TextureGraph/Source/TextureGraph/Private/Expressions/Input/TG_Expression_Scalar.cpp b/Engine/Plugins/Experimental/TextureGraph/Source/TextureGraph/Private/Expressions/Input/TG_Expression_Scalar.cpp
--- a/Engine/Plugins/Experimental/TextureGraph/Source/TextureGraph/Private/Expressions/Input/TG_Expression_Scalar.cpp
+++ b/Engine/Plugins/Experimental/TextureGraph/Source/TextureGraph/Private/Expressions/Input/TG_Expression_Scalar.cpp
@@ -3,26 +3,6 @@
 #include "expressions/Input/TG_Expression_Scalar.h"
 
 
-FTG_SignaturePtr UTG_Expression_Scalar::BuildInputParameterSignature() const
-{
-	FTG_Signature::FInit SignatureInit = GetSignatureInitArgsFromClass();
-	return MakeShared<FTG_Signature>(SignatureInit);
-};
-
-FTG_SignaturePtr UTG_Expression_Scalar::BuildInputConstantSignature() const
-{
-	FTG_Signature::FInit SignatureInit = GetSignatureInitArgsFromClass();
-	for (auto& arg : SignatureInit.Arguments)
-	{
-		if (arg.IsInput() && arg.IsParam())
-		{
-			arg.ArgumentType = arg.ArgumentType.Unparamed();
-			arg.ArgumentType.SetNotConnectable();
-		}
-	}
-	return MakeShared<FTG_Signature>(SignatureInit);
-};
-
 void UTG_Expression_Scalar::Evaluate(FTG_EvaluationContext* InContext)
 {
 	Super::Evaluate(InContext);
diff --git a/Engine/Plugins/Experimental/TextureGraph/Source/TextureGraph/Private/Expressions/Input/TG_Expression_String.cpp b/Engine/Plugins/Experimental/TextureGraph/Source/TextureGraph/Private/Expressions/Input/TG_Expression_String.cpp
new file mode 100644
--- /dev/null
+++ b/Engine/Plugins/Experimental/TextureGraph/Source/TextureGraph/Private/Expressions/Input/TG_Expression_String.cpp
@@ -0,0 +1,23 @@
+// Copyright Epic Games, Inc. All Rights Reserved.
+
+#include "expressions/Input/TG_Expression_String.h"
+
+
+void UTG_Expression_String::Evaluate(FTG_EvaluationContext* InContext)
+{
+	Super::Evaluate(InContext);
+
+	// The Value is updated either as an input or as a setting and then becomes the output for this expression
+	// The pin out is named "ValueOut"
+	ValueOut = String;
+}
+
+void UTG_Expression_String::SetTitleName(FName NewName)
+{
+	GetParentNode()->GetInputPin(GET_MEMBER_NAME_CHECKED(UTG_Expression_String, String))->SetAliasName(NewName);
+}
+
+FName UTG_Expression_String::GetTitleName() const
+{
+	return GetParentNode()->GetInputPin(GET_MEMBER_NAME_CHECKED(UTG_Expression_String, String))->GetAliasName();
+}
\ No newline at end of file
diff --git a/Engine/Plugins/Experimental/TextureGraph/Source/TextureGraph/Private/Expressions/Input/TG_Expression_Texture.cpp b/Engine/Plugins/Experimental/TextureGraph/Source/TextureGraph/Private/Expressions/Input/TG_Expression_Texture.cpp
--- a/Engine/Plugins/Experimental/TextureGraph/Source/TextureGraph/Private/Expressions/Input/TG_Expression_Texture.cpp
+++ b/Engine/Plugins/Experimental/TextureGraph/Source/TextureGraph/Private/Expressions/Input/TG_Expression_Texture.cpp
@@ -6,29 +6,6 @@
 #include "2D/TextureHelper.h"
 #include "Model/StaticImageResource.h"
 
-
-
-FTG_SignaturePtr UTG_Expression_Texture::BuildInputParameterSignature() const
-{
-	FTG_Signature::FInit SignatureInit = GetSignatureInitArgsFromClass();
-	return MakeShared<FTG_Signature>(SignatureInit);
-};
-
-FTG_SignaturePtr UTG_Expression_Texture::BuildInputConstantSignature() const
-{
-	FTG_Signature::FInit SignatureInit = GetSignatureInitArgsFromClass();
-	for (auto& arg : SignatureInit.Arguments)
-	{	
-		// Turn the Input parameter TG_Texture entry into private so it s not exposed or used
-		// in the constant case, the expression only rely on the UTexture asset
-		if (arg.IsInput() && arg.IsParam())
-		{
-			arg.ArgumentType = FTG_ArgumentType(ETG_Access::Private);
-		}
-	}
-	return MakeShared<FTG_Signature>(SignatureInit);
-};
-
 void UTG_Expression_Texture::Evaluate(FTG_EvaluationContext* InContext)
 {
 	Super::Evaluate(InContext);
diff --git a/Engine/Plugins/Experimental/TextureGraph/Source/TextureGraph/Private/Expressions/Input/TG_Expression_TexturePath.cpp b/Engine/Plugins/Experimental/TextureGraph/Source/TextureGraph/Private/Expressions/Input/TG_Expression_TexturePath.cpp
new file mode 100644
--- /dev/null
+++ b/Engine/Plugins/Experimental/TextureGraph/Source/TextureGraph/Private/Expressions/Input/TG_Expression_TexturePath.cpp
@@ -0,0 +1,65 @@
+// Copyright Epic Games, Inc. All Rights Reserved.
+
+#include "Expressions/Input/TG_Expression_TexturePath.h"
+
+#include "TG_Graph.h"
+#include "2D/TextureHelper.h"
+#include "Model/StaticImageResource.h"
+
+void UTG_Expression_TexturePath::Evaluate(FTG_EvaluationContext* InContext)
+{
+	Super::Evaluate(InContext);
+	if (Texture)
+	{
+		Output = Texture;
+	}
+	else if (!Path.IsEmpty())
+	{
+		UStaticImageResource* StaticImageResource = UStaticImageResource::CreateNew<UStaticImageResource>();
+
+		FString LocalPath = Path;
+		FPackagePath PackagePath;
+		FString PathExt = FPaths::GetExtension(Path);
+
+		if (FPackagePath::TryFromMountedName(Path, PackagePath))
+		{
+			LocalPath = PackagePath.GetLocalFullPath();
+			FString LocalPathExt = FPaths::GetExtension(LocalPath);
+
+			if (LocalPathExt != PathExt)
+			{
+				LocalPath = FPaths::ChangeExtension(LocalPath, PathExt);
+			}
+		}
+
+		StaticImageResource->SetAssetUUID(LocalPath);
+		StaticImageResource->SetIsFileSystem(true);
+
+		//Until we have srgb value exposed in the UI we need to set the Srgb of the Output Descriptor here from the source
+		//This gets updated for the late bond case but since we do not have the UI to specify the override in other nodes 
+		// the override value will always be set to false while combining the buffers
+		auto DesiredDesc = Output.GetBufferDescriptor();
+		DesiredDesc.bIsSRGB = true;
+
+		Output = StaticImageResource->GetBlob(InContext->Cycle, DesiredDesc, 0);
+	}
+	else
+	{
+		Output = FTG_Texture::GetBlack();
+	}
+}
+
+bool UTG_Expression_TexturePath::Validate(MixUpdateCyclePtr Cycle)
+{
+	return true;
+}
+void UTG_Expression_TexturePath::SetTitleName(FName NewName)
+{
+	GetParentNode()->GetPin(GET_MEMBER_NAME_CHECKED(UTG_Expression_TexturePath, Path))->SetAliasName(NewName);
+}
+
+FName UTG_Expression_TexturePath::GetTitleName() const
+{
+	return GetParentNode()->GetPin(GET_MEMBER_NAME_CHECKED(UTG_Expression_TexturePath, Path))->GetAliasName();
+}
+
diff --git a/Engine/Plugins/Experimental/TextureGraph/Source/TextureGraph/Private/Expressions/Input/TG_Expression_Vector.cpp b/Engine/Plugins/Experimental/TextureGraph/Source/TextureGraph/Private/Expressions/Input/TG_Expression_Vector.cpp
--- a/Engine/Plugins/Experimental/TextureGraph/Source/TextureGraph/Private/Expressions/Input/TG_Expression_Vector.cpp
+++ b/Engine/Plugins/Experimental/TextureGraph/Source/TextureGraph/Private/Expressions/Input/TG_Expression_Vector.cpp
@@ -2,26 +2,6 @@
 
 #include "Expressions/Input/TG_Expression_Vector.h"
 
-FTG_SignaturePtr UTG_Expression_Vector::BuildInputParameterSignature() const
-{
-	FTG_Signature::FInit SignatureInit = GetSignatureInitArgsFromClass();
-	return MakeShared<FTG_Signature>(SignatureInit);
-};
-
-FTG_SignaturePtr UTG_Expression_Vector::BuildInputConstantSignature() const
-{
-	FTG_Signature::FInit SignatureInit = GetSignatureInitArgsFromClass();
-	for (auto& arg : SignatureInit.Arguments)
-	{
-		if (arg.IsInput() && arg.IsParam())
-		{
-			arg.ArgumentType = arg.ArgumentType.Unparamed();
-			arg.ArgumentType.SetNotConnectable();
-		}
-	}
-	return MakeShared<FTG_Signature>(SignatureInit);
-};
-
 void UTG_Expression_Vector::Evaluate(FTG_EvaluationContext* InContext)
 {
 	Super::Evaluate(InContext);
diff --git a/Engine/Plugins/Experimental/TextureGraph/Source/TextureGraph/Private/Expressions/Output/TG_Expression_Output.cpp b/Engine/Plugins/Experimental/TextureGraph/Source/TextureGraph/Private/Expressions/Output/TG_Expression_Output.cpp
--- a/Engine/Plugins/Experimental/TextureGraph/Source/TextureGraph/Private/Expressions/Output/TG_Expression_Output.cpp
+++ b/Engine/Plugins/Experimental/TextureGraph/Source/TextureGraph/Private/Expressions/Output/TG_Expression_Output.cpp
@@ -148,7 +148,7 @@ void UTG_Expression_Output::InitializeOutputSettings()
 	OutputSettings.Initialize(GetParentNode()->GetGraph()->GetPathName(),GetTitleName());
 
 	UTG_Pin* Settings = GetParentNode()->GetPin(GET_MEMBER_NAME_CHECKED(UTG_Expression_Output, OutputSettings));
-	Settings->SetValue(OutputSettings.ToString());
+	Settings->FromString(OutputSettings.ToString());
 }
 
 void UTG_Expression_Output::SetExport(bool bExport)
diff --git a/Engine/Plugins/Experimental/TextureGraph/Source/TextureGraph/Private/Expressions/TG_Expression.cpp b/Engine/Plugins/Experimental/TextureGraph/Source/TextureGraph/Private/Expressions/TG_Expression.cpp
--- a/Engine/Plugins/Experimental/TextureGraph/Source/TextureGraph/Private/Expressions/TG_Expression.cpp
+++ b/Engine/Plugins/Experimental/TextureGraph/Source/TextureGraph/Private/Expressions/TG_Expression.cpp
@@ -24,7 +24,7 @@ void UTG_Expression::Serialize(FArchive& Ar)
 	// in the case of Loading, let s check the expression version retreived and log if we detect an older version
 	if (Ar.IsLoading() && (InstanceExpressionClassVersion < this->GetExpressionClassVersion()))
 	{
-		UE_LOG(LogTextureGraph, Log, TEXT("Detected an expression in a previous version"));
+		UE_LOG(LogTextureGraph, VeryVerbose, TEXT("Detected an expression in a previous version"));
 	}
 }
 
@@ -133,7 +133,7 @@ void UTG_Expression::PostEditChangeProperty(FPropertyChangedEvent& PropertyChang
 {
 	Super::PostEditChangeProperty(PropertyChangedEvent);
 	
-	UE_LOG(LogTextureGraph, Log, TEXT("Expression PostEditChangeProperty. ChangeType: %d"), (int32)PropertyChangedEvent.ChangeType);
+	UE_LOG(LogTextureGraph, VeryVerbose, TEXT("Expression PostEditChangeProperty. ChangeType: %d"), (int32)PropertyChangedEvent.ChangeType);
 
 	NotifyExpressionChanged(PropertyChangedEvent);
 }
@@ -322,5 +322,5 @@ void UTG_Expression::LogEvaluation(FTG_EvaluationContext* InContext)
 	for (auto& VarMapEntry : InContext->Outputs.VarArguments)
 		OutVarIds.Emplace(VarMapEntry.Value.Var->GetId());
 
-	UE_LOG(LogTextureGraph, Log, TEXT("%-*s %s"), UTG_Graph::LogHeaderWidth, *GetDefaultName().ToString(), *UTG_Graph::LogCall(InVarIds, OutVarIds));
+	UE_LOG(LogTextureGraph, VeryVerbose, TEXT("%-*s %s"), UTG_Graph::LogHeaderWidth, *GetDefaultName().ToString(), *UTG_Graph::LogCall(InVarIds, OutVarIds));
 }
diff --git a/Engine/Plugins/Experimental/TextureGraph/Source/TextureGraph/Private/TG_Graph.cpp b/Engine/Plugins/Experimental/TextureGraph/Source/TextureGraph/Private/TG_Graph.cpp
--- a/Engine/Plugins/Experimental/TextureGraph/Source/TextureGraph/Private/TG_Graph.cpp
+++ b/Engine/Plugins/Experimental/TextureGraph/Source/TextureGraph/Private/TG_Graph.cpp
@@ -12,7 +12,7 @@
 #if WITH_EDITOR
 void UTG_Graph::PostEditChangeProperty(FPropertyChangedEvent& PropertyChangedEvent)
 {
-	UE_LOG(LogTextureGraph, Log, TEXT("UTG_Graph::PostEditChangeProperty."));
+	UE_LOG(LogTextureGraph, VeryVerbose, TEXT("UTG_Graph::PostEditChangeProperty."));
 }
 
 bool UTG_Graph::Modify(bool bAlwaysMarkDirty)
@@ -20,7 +20,7 @@ bool UTG_Graph::Modify(bool bAlwaysMarkDirty)
 	// Runtime state is about to get dirty
 	// And remember it as such for undo redo
 	bIsGraphTraversalDirty = true;
-	UE_LOG(LogTextureGraph, Log, TEXT("UTG_Graph::Modify: Graph Modified."));
+	UE_LOG(LogTextureGraph, VeryVerbose, TEXT("UTG_Graph::Modify: Graph Modified."));
 
 	return Super::Modify(bAlwaysMarkDirty);
 }
@@ -28,7 +28,7 @@ bool UTG_Graph::Modify(bool bAlwaysMarkDirty)
 
 void UTG_Graph::PostEditUndo()
 {
-	UE_LOG(LogTextureGraph, Log, TEXT("UTG_Graph::PostEditUndo."));
+	UE_LOG(LogTextureGraph, VeryVerbose, TEXT("UTG_Graph::PostEditUndo."));
 	UObject::PostEditUndo();
 }
 
@@ -56,7 +56,7 @@ void UTG_Graph::Serialize(FArchive& Ar)
 
 	int32 Version = Ar.CustomVer(FTG_CustomVersion::GUID);
 
-	UE_LOG(LogTextureGraph, Log, TEXT("  %s Graph: %s"),
+	UE_LOG(LogTextureGraph, VeryVerbose, TEXT("  %s Graph: %s"),
 		(Ar.IsSaving() ? TEXT("Saved") : TEXT("Loaded")),
 		*Name);
 }
@@ -65,7 +65,7 @@ void UTG_Graph::PostLoad()
 {
 	Super::PostLoad();
 
-	UE_LOG(LogTextureGraph, Log, TEXT("  PostLoad Graph: %s"), *Name);
+	UE_LOG(LogTextureGraph, VeryVerbose, TEXT("  PostLoad Graph: %s"), *Name);
 
 	// We have to reset the Transactional flag for the UTG_Graph and we do not know why...
 	// For the other UObject in the data structure (Script, Node, Expression) it is set
@@ -86,7 +86,7 @@ void UTG_Graph::PostLoad()
 			if (!PinsMatchSignature)
 			{
 				NodeRequiringRemap.Add(i);
-				UE_LOG(LogTextureGraph, Log, TEXT("Node %s serialized signature is different from expression"), *Node->GetId().ToString());
+				UE_LOG(LogTextureGraph, VeryVerbose, TEXT("Node %s serialized signature is different from expression"), *Node->GetId().ToString());
 				Node->WarningStack.Add(FName(FString::Printf(TEXT("Node %s serialized signature is different from expression's signature, pins are regenerated"), *Node->GetId().ToString())));
 			}
 
@@ -126,7 +126,7 @@ void UTG_Graph::PostLoad()
 void UTG_Graph::PreSave(FObjectPreSaveContext SaveContext)
 {
 	Super::PreSave(SaveContext);
-	UE_LOG(LogTextureGraph, Log, TEXT("  PreSave Graph: %s"), *Name);
+	UE_LOG(LogTextureGraph, VeryVerbose, TEXT("  PreSave Graph: %s"), *Name);
 }
 
 // Inner setup node calling allocation of pins and edges and var in cascade
@@ -655,13 +655,13 @@ bool UTG_Graph::RenameParam(FName OldName, FName NewName)
 
 void UTG_Graph::OnNodeChanged(UTG_Node* InNode, bool bIsTweaking)
 {
-	UE_LOG(LogTextureGraph, Log, TEXT("UTG_Graph::Node Changed"))
+	UE_LOG(LogTextureGraph, VeryVerbose, TEXT("UTG_Graph::Node Changed"))
 	NotifyGraphChanged(InNode, bIsTweaking);
 }
 
 void UTG_Graph::OnNodeSignatureChanged(UTG_Node* InNode)
 {
-	UE_LOG(LogTextureGraph, Log, TEXT("UTG_Graph::Node Recreate"));
+	UE_LOG(LogTextureGraph, VeryVerbose, TEXT("UTG_Graph::Node Recreate"));
 	// recreate node with expression
 	RegenerateNode(InNode);
 
@@ -675,7 +675,7 @@ void UTG_Graph::OnNodeSignatureChanged(UTG_Node* InNode)
 
 void UTG_Graph::OnNodePinChanged(FTG_Id InPinId, UTG_Node* InNode)
 {
-	UE_LOG(LogTextureGraph, Log, TEXT("UTG_Graph::Pin Changed"))
+	UE_LOG(LogTextureGraph, VeryVerbose, TEXT("UTG_Graph::Pin Changed"))
 
 	// Commented this code
 	// Right now pin changed is being called only from renaming pin.
@@ -846,7 +846,7 @@ int UTG_Graph::GetAllOutputParamValues(TArray<FTG_Variant>& OutVariants, TArray<
 				}
 				else
 				{
-					UE_LOG(LogTextureGraph, Log, TEXT("Output {} variant failed to access"), *(Pin->GetAliasName().ToString()));
+					UE_LOG(LogTextureGraph, VeryVerbose, TEXT("Output {} variant failed to access"), *(Pin->GetAliasName().ToString()));
 				}
 			}
 		}
@@ -1034,7 +1034,7 @@ void UTG_Graph::SetExtraEditorNodes(const TArray<TObjectPtr<const UObject>>& InN
 #endif
 const FTG_GraphTraversal& UTG_Graph::GetTraversal() const
 {
-	UE_LOG(LogTextureGraph, Log, TEXT("UTG_Graph::GetTraversal: Dirty = %s"),
+	UE_LOG(LogTextureGraph, VeryVerbose, TEXT("UTG_Graph::GetTraversal: Dirty = %s"),
 	       bIsGraphTraversalDirty ? TEXT("true") : TEXT("false"));
 	if (bIsGraphTraversalDirty)
 	{
@@ -1081,7 +1081,7 @@ void UTG_Graph::Log()
 	LogMessage.ParseIntoArray(Lines, TEXT("\n"));
 	for (const FString& Line : Lines)
 	{
-		UE_LOG(LogTextureGraph, Log, TEXT("%s"), *Line);
+		UE_LOG(LogTextureGraph, VeryVerbose, TEXT("%s"), *Line);
 	}
 }
 
diff --git a/Engine/Plugins/Experimental/TextureGraph/Source/TextureGraph/Private/TG_HelperFunctions.cpp b/Engine/Plugins/Experimental/TextureGraph/Source/TextureGraph/Private/TG_HelperFunctions.cpp
--- a/Engine/Plugins/Experimental/TextureGraph/Source/TextureGraph/Private/TG_HelperFunctions.cpp
+++ b/Engine/Plugins/Experimental/TextureGraph/Source/TextureGraph/Private/TG_HelperFunctions.cpp
@@ -110,12 +110,13 @@ void FTG_HelperFunctions::EnsureOutputIsTexture(MixUpdateCyclePtr Cycle, UTG_Nod
 	}
 }
 
-AsyncBool FTG_HelperFunctions::ExportAsync(UTextureGraph* InTextureGraph, FString ExportPath, FString AssetName, FExportSettings& TargetExportSettings, bool OverrideExportPath, bool OverwriteTextures /*= true*/, bool ExportAllOutputs /*= false*/,bool bSave /*= true*/)
+JobBatchPtr FTG_HelperFunctions::InitExportBatch(UTextureGraph* InTextureGraph, FString ExportPath, FString AssetName, FExportSettings& TargetExportSettings,
+	bool OverrideExportPath, bool OverwriteTextures, bool ExportAllOutputs, bool bSave)
 {
 	FString ErrorMessage = "";
 	bool AnyValidExport = false;
 	TargetExportSettings.Reset();
-	
+
 	FInvalidationDetails Details;
 	Details.All();
 	Details.Mix = InTextureGraph;
@@ -124,79 +125,79 @@ AsyncBool FTG_HelperFunctions::ExportAsync(UTextureGraph* InTextureGraph, FStrin
 	/// Update the mix so that the rendering Cycle gets populated
 	MixUpdateCyclePtr Cycle = Batch->GetCycle();
 	InTextureGraph->Update(Cycle);
-	
+
 	InTextureGraph->Graph()->ForEachNodes([=,&TargetExportSettings, &ErrorMessage , &AnyValidExport](const UTG_Node* Node, uint32 Index)
-	{
-		UTG_Expression_Output* TargetExpression = Cast<UTG_Expression_Output>(Node->GetExpression());
-		if (TargetExpression)
 		{
-			FTG_OutputSettings& OutputSetting = TargetExpression->OutputSettings;
-
-			if (OutputSetting.bExport || ExportAllOutputs)
+			UTG_Expression_Output* TargetExpression = Cast<UTG_Expression_Output>(Node->GetExpression());
+			if (TargetExpression)
 			{
-				FTG_EvaluationContext EvaluationContext;
-				EvaluationContext.Cycle = Batch->GetCycle();
-
-				auto ExportBlobs = FTG_HelperFunctions::GetTexturedOutputs(Node, &EvaluationContext);
+				FTG_OutputSettings& OutputSetting = TargetExpression->OutputSettings;
 
-				FString Path = OutputSetting.FolderPath.ToString();
-				if (OverrideExportPath)
+				if (OutputSetting.bExport || ExportAllOutputs)
 				{
-					Path = ExportPath;
-				}
+					FTG_EvaluationContext EvaluationContext;
+					EvaluationContext.Cycle = Batch->GetCycle();
 
-				FString FileName = AssetName.IsEmpty() ? OutputSetting.GetFullOutputName() : AssetName;
+					auto ExportBlobs = FTG_HelperFunctions::GetTexturedOutputs(Node, &EvaluationContext);
 
-				FText Reason;
-				bool IsNameValid = IsFileNameValid(FName(*FileName), Reason);
+					FString Path = OutputSetting.FolderPath.ToString();
+					if (OverrideExportPath)
+					{
+						Path = ExportPath;
+					}
 
-				FText PathReason;
-				bool IsPathValid = TextureExporter::IsPackageNameValid(Path, FileName);
-				IsFolderPathValid(Path, PathReason);
-				IsPathValid &= PathReason.IsEmpty();
+					FString FileName = AssetName.IsEmpty() ? OutputSetting.GetFullOutputName() : AssetName;
 
-				bool HasOutputs = ExportBlobs.Num() > 0;
+					FText Reason;
+					bool IsNameValid = IsFileNameValid(FName(*FileName), Reason);
 
-				if (HasOutputs && IsNameValid && IsPathValid)
-				{
-					TiledBlobPtr Output = std::static_pointer_cast<TiledBlob>(ExportBlobs[0]);//Dealing with one output per Node for now
-					FExportMapSettings MapSettings = TextureExporter::GetExportSettingsForTarget(TargetExportSettings, std::static_pointer_cast<TiledBlob>(Output), *FileName);
-					MapSettings.Name = FName(*FileName);
-					MapSettings.Path = Path;
-					MapSettings.UseOverridePath = OverrideExportPath;
-					MapSettings.OverwriteTextures = OverwriteTextures;
-					MapSettings.LODGroup = OutputSetting.LODGroup;
-					MapSettings.Compression = OutputSetting.Compression;
-					MapSettings.IsSRGB = OutputSetting.bSRGB;
-					MapSettings.Width = (int32)OutputSetting.Width;
-					MapSettings.Height = (int32)OutputSetting.Height;
-					MapSettings.bSave = bSave;
-					TargetExportSettings.ExportPreset.push_back(std::pair<FName, FExportMapSettings>{ MapSettings.Name, MapSettings });
-					AnyValidExport = true;
-				}
-				else
-				{
-					//Log Error to Error System
-					if (!HasOutputs)
-					{
-						ErrorMessage += FString::Format(TEXT("Texture Export Error : No valid output found for OutputSetting {0}"), { OutputSetting.OutputName.ToString() });
-						ErrorMessage += "\n";
-					}
-					if (!IsPathValid)
+					FText PathReason;
+					bool IsPathValid = TextureExporter::IsPackageNameValid(Path, FileName);
+					IsFolderPathValid(Path, PathReason);
+					IsPathValid &= PathReason.IsEmpty();
+
+					bool HasOutputs = ExportBlobs.Num() > 0;
+
+					if (HasOutputs && IsNameValid && IsPathValid)
 					{
-						ErrorMessage += FString::Format(TEXT("Texture Export Error : Invalid path set for OutputSetting {0}"), { OutputSetting.OutputName.ToString() }) + (PathReason.IsEmpty() ? "" : " Folder " + PathReason.ToString());
-						ErrorMessage += "\n";
+						TiledBlobPtr Output = std::static_pointer_cast<TiledBlob>(ExportBlobs[0]);//Dealing with one output per Node for now
+						FExportMapSettings MapSettings = TextureExporter::GetExportSettingsForTarget(TargetExportSettings, std::static_pointer_cast<TiledBlob>(Output), *FileName);
+						MapSettings.Name = FName(*FileName);
+						MapSettings.Path = Path;
+						MapSettings.UseOverridePath = OverrideExportPath;
+						MapSettings.OverwriteTextures = OverwriteTextures;
+						MapSettings.LODGroup = OutputSetting.LODGroup;
+						MapSettings.Compression = OutputSetting.Compression;
+						MapSettings.IsSRGB = OutputSetting.bSRGB;
+						MapSettings.Width = (int32)OutputSetting.Width;
+						MapSettings.Height = (int32)OutputSetting.Height;
+						MapSettings.bSave = bSave;
+						TargetExportSettings.ExportPreset.push_back(std::pair<FName, FExportMapSettings>{ MapSettings.Name, MapSettings });
+						AnyValidExport = true;
 					}
-					if (!IsNameValid)
+					else
 					{
-						FString ReasonString = Reason.IsEmpty() ? "Invalid File Name" : Reason.ToString();
-						ErrorMessage += FString::Format(TEXT("Texture Export Error : {0} for OutputSetting {1}"), { ReasonString , OutputSetting.OutputName.ToString()});
-						ErrorMessage += "\n";
+						//Log Error to Error System
+						if (!HasOutputs)
+						{
+							ErrorMessage += FString::Format(TEXT("Texture Export Error : No valid output found for OutputSetting {0}"), { OutputSetting.OutputName.ToString() });
+							ErrorMessage += "\n";
+						}
+						if (!IsPathValid)
+						{
+							ErrorMessage += FString::Format(TEXT("Texture Export Error : Invalid path set for OutputSetting {0}"), { OutputSetting.OutputName.ToString() }) + (PathReason.IsEmpty() ? "" : " Folder " + PathReason.ToString());
+							ErrorMessage += "\n";
+						}
+						if (!IsNameValid)
+						{
+							FString ReasonString = Reason.IsEmpty() ? "Invalid File Name" : Reason.ToString();
+							ErrorMessage += FString::Format(TEXT("Texture Export Error : {0} for OutputSetting {1}"), { ReasonString , OutputSetting.OutputName.ToString()});
+							ErrorMessage += "\n";
+						}
 					}
 				}
 			}
-		}
-	});
+		});
 
 	if (!ErrorMessage.IsEmpty())
 	{
@@ -209,31 +210,41 @@ AsyncBool FTG_HelperFunctions::ExportAsync(UTextureGraph* InTextureGraph, FStrin
 
 		if (!AnyValidExport)
 		{
-			return cti::make_ready_continuable(false);
+			return nullptr;
 		}
 	}
 
-	return RenderAsync(InTextureGraph, Batch)
-	.then([InTextureGraph, &TargetExportSettings, ExportPath](auto result) 
-	{
-		return TextureExporter::ExportAsUAsset(InTextureGraph, TargetExportSettings, ExportPath);
-	})
-	.then([InTextureGraph, &TargetExportSettings]()
-	{
-		InTextureGraph->InvalidateAll();
+	return Batch;
+}
+
+AsyncBool FTG_HelperFunctions::ExportAsync(UTextureGraph* InTextureGraph, FString ExportPath, FString AssetName, FExportSettings& TargetExportSettings, bool OverrideExportPath, bool OverwriteTextures /*= true*/, bool ExportAllOutputs /*= false*/,bool bSave /*= true*/)
+{
+	JobBatchPtr Batch = InitExportBatch(InTextureGraph, ExportPath, AssetName, TargetExportSettings, OverrideExportPath, OverwriteTextures, ExportAllOutputs, bSave);
 
-		// Add analytics tag
-		if (FEngineAnalytics::IsAvailable())
+	if (!Batch)
+		return cti::make_ready_continuable<bool>(false);
+
+	return RenderAsync(InTextureGraph, Batch)
+		.then([InTextureGraph, &TargetExportSettings, ExportPath](auto result) 
+		{
+			return TextureExporter::ExportAsUAsset(InTextureGraph, TargetExportSettings, ExportPath);
+		})
+		.then([InTextureGraph, &TargetExportSettings]()
 		{
-			TArray<FAnalyticsEventAttribute> Attributes;
-			Attributes.Add(FAnalyticsEventAttribute(TEXT("NumExports"),  TargetExportSettings.MapsExported));
+			InTextureGraph->InvalidateAll();
+
+			// Add analytics tag
+			if (FEngineAnalytics::IsAvailable())
+			{
+				TArray<FAnalyticsEventAttribute> Attributes;
+				Attributes.Add(FAnalyticsEventAttribute(TEXT("NumExports"),  TargetExportSettings.MapsExported));
 					
-			// Send Analytics event 
-			FEngineAnalytics::GetProvider().RecordEvent(TEXT("Editor.Usage.TextureGraph.Export"), Attributes);
-		}
+				// Send Analytics event 
+				FEngineAnalytics::GetProvider().RecordEvent(TEXT("Editor.Usage.TextureGraph.Export"), Attributes);
+			}
 
-		return true;
-	});
+			return true;
+		});
 }
 
 void FTG_HelperFunctions::InitTargets(UTextureGraph* InTextureGraph)
@@ -263,7 +274,7 @@ void FTG_HelperFunctions::InitTargets(UTextureGraph* InTextureGraph)
 	}
 }
 
-AsyncBool FTG_HelperFunctions::RenderAsync(UTextureGraph* InTextureGraph, JobBatchPtr ExistingBatch /* = nullptr */)
+JobBatchPtr FTG_HelperFunctions::InitRenderBatch(UTextureGraph* InTextureGraph, JobBatchPtr ExistingBatch /* = nullptr */)
 {
 	/// Now run the update Cycle
 	JobBatchPtr NewBatch = nullptr;
@@ -277,13 +288,20 @@ AsyncBool FTG_HelperFunctions::RenderAsync(UTextureGraph* InTextureGraph, JobBat
 
 		NewBatch = JobBatch::Create(Details);
 		MixUpdateCyclePtr Cycle = NewBatch->GetCycle();
-		
+
 		/// Update the mix so that the rendering Cycle gets populated
 		InTextureGraph->Update(Cycle);
 
 		Batch = NewBatch;
 	}
 
+	return Batch;
+}
+
+AsyncBool FTG_HelperFunctions::RenderAsync(UTextureGraph* InTextureGraph, JobBatchPtr ExistingBatch /* = nullptr */)
+{
+	JobBatchPtr Batch = InitRenderBatch(InTextureGraph, ExistingBatch);
+
 	if (!Batch)
 	{
 		return cti::make_ready_continuable(true);
diff --git a/Engine/Plugins/Experimental/TextureGraph/Source/TextureGraph/Private/TG_Node.cpp b/Engine/Plugins/Experimental/TextureGraph/Source/TextureGraph/Private/TG_Node.cpp
--- a/Engine/Plugins/Experimental/TextureGraph/Source/TextureGraph/Private/TG_Node.cpp
+++ b/Engine/Plugins/Experimental/TextureGraph/Source/TextureGraph/Private/TG_Node.cpp
@@ -34,7 +34,7 @@ FTG_Name UTG_Node::GetNodeName() const
 #if WITH_EDITOR
 void UTG_Node::PostEditChangeProperty(FPropertyChangedEvent& PropertyChangedEvent)
 {
-	UE_LOG(LogTextureGraph, Log, TEXT("UTG_Node::PostEditChangeProperty."));
+	UE_LOG(LogTextureGraph, VeryVerbose, TEXT("UTG_Node::PostEditChangeProperty."));
 }
 
 void UTG_Node::PostEditUndo()
@@ -228,12 +228,12 @@ void UTG_Node::Serialize(FArchive& Ar)
 
 	if (!Expression)
 	{
-		UE_LOG(LogTextureGraph, Log, TEXT("    %s Node: NUll Expression???"),
+		UE_LOG(LogTextureGraph, Verbose, TEXT("    %s Node: NUll Expression???"),
 			(Ar.IsSaving() ? TEXT("Saved") : TEXT("Loaded")));
 		Expression = NewObject<UTG_Expression_Null>(this, UTG_Expression_Null::StaticClass(), NAME_None, RF_Transactional);
 	}
 
-	UE_LOG(LogTextureGraph, Log, TEXT("    %s Node: %s"),
+	UE_LOG(LogTextureGraph, VeryVerbose, TEXT("    %s Node: %s"),
 		(Ar.IsSaving() ? TEXT("Saved") : TEXT("Loaded")),
 		*GetId().ToString());
 }
diff --git a/Engine/Plugins/Experimental/TextureGraph/Source/TextureGraph/Private/TG_Pin.cpp b/Engine/Plugins/Experimental/TextureGraph/Source/TextureGraph/Private/TG_Pin.cpp
--- a/Engine/Plugins/Experimental/TextureGraph/Source/TextureGraph/Private/TG_Pin.cpp
+++ b/Engine/Plugins/Experimental/TextureGraph/Source/TextureGraph/Private/TG_Pin.cpp
@@ -17,7 +17,7 @@ void UTG_Pin::Serialize(FArchive& Ar)
 	// Serialize the Var with the knowledge of the 
 	SelfVar.Serialize(Ar, GetId(), GetArgument());
 
-	UE_LOG(LogTextureGraph, Log, TEXT("      %s Pin: %s : %s"),
+	UE_LOG(LogTextureGraph, VeryVerbose, TEXT("      %s Pin: %s : %s"),
 		(Ar.IsSaving() ? TEXT("Saved") : TEXT("Loaded")),
 		*GetId().ToString(),
 		*GetArgumentName().ToString());
@@ -272,7 +272,7 @@ FString UTG_Pin::GetEvaluatedVarValue() const
 	return CurrentValue;
 }
 
-void UTG_Pin::SetValue(const FString& InValueStr, bool bIsTweaking /*= false*/)
+void UTG_Pin::FromString(const FString& InValueStr, bool bIsTweaking /*= false*/)
 {
 	Modify();
 
@@ -401,6 +401,26 @@ bool UTG_Pin::GetValue(float& OutValue) const
 	return false;
 }
 
+bool UTG_Pin::GetValue(bool& OutValue) const
+{
+	if (IsArgBool())
+	{
+		OutValue = GetSelfVar()->GetAs<bool>();
+		return true;
+	}
+	return false;
+}
+
+bool UTG_Pin::GetValue(FString& OutValue) const
+{
+	if (IsArgString())
+	{
+		OutValue = GetSelfVar()->GetAs<FString>();
+		return true;
+	}
+	return false;
+}
+
 bool UTG_Pin::GetValue(FLinearColor& OutValue) const
 {
 	if (IsArgColor())
@@ -589,6 +609,44 @@ bool UTG_Pin::SetValue(float Value)
 	return false;
 }
 
+bool UTG_Pin::SetValue(bool bValue)
+{
+	if (IsArgBool())
+	{
+		Modify();
+		EditSelfVar()->EditAs<bool>() = bValue;
+		NotifyPinSelfVarChanged();
+
+		return true;
+	}
+	return false;
+}
+
+bool UTG_Pin::SetValue(FString Value)
+{
+	if (IsArgString())
+	{
+		Modify();
+		EditSelfVar()->EditAs<FString>() = Value;
+		NotifyPinSelfVarChanged();
+
+		return true;
+	}
+	return false;
+}
+
+bool UTG_Pin::SetValue(FName Value)
+{
+	if (IsArgString())
+	{
+		Modify();
+		EditSelfVar()->EditAs<FString>() = Value.GetPlainNameString();
+		NotifyPinSelfVarChanged();
+
+		return true;
+	}
+	return false;
+}
 
 bool UTG_Pin::SetValue(const FLinearColor& Value)
 {
diff --git a/Engine/Plugins/Experimental/TextureGraph/Source/TextureGraph/Private/TG_Signature.cpp b/Engine/Plugins/Experimental/TextureGraph/Source/TextureGraph/Private/TG_Signature.cpp
--- a/Engine/Plugins/Experimental/TextureGraph/Source/TextureGraph/Private/TG_Signature.cpp
+++ b/Engine/Plugins/Experimental/TextureGraph/Source/TextureGraph/Private/TG_Signature.cpp
@@ -12,10 +12,24 @@ bool FTG_Argument::IsScalar() const
 {
 	return (CPPTypeName == TEXT("float")) || (CPPTypeName == TEXT("int32")) || (CPPTypeName == TEXT("uint32")) || CPPTypeName == FTG_Variant::GetArgNameFromType(ETG_VariantType::Scalar);
 }
+
+bool FTG_Argument::IsBool() const
+{
+	return CPPTypeName == TEXT("bool");
+}
+
+bool FTG_Argument::IsString() const
+{
+	return CPPTypeName == TEXT("FString") || 
+		CPPTypeName == TEXT("FName") || 
+		CPPTypeName == TEXT("FText");
+}
+
 bool FTG_Argument::IsColor() const
 {
 	return CPPTypeName == TEXT("FLinearColor") || CPPTypeName == FTG_Variant::GetArgNameFromType(ETG_VariantType::Color);
 }
+
 bool FTG_Argument::IsVector() const
 {
 	return (CPPTypeName == TEXT("FVector4f")) || CPPTypeName == FTG_Variant::GetArgNameFromType(ETG_VariantType::Vector);
diff --git a/Engine/Plugins/Experimental/TextureGraph/Source/TextureGraph/Private/TG_Var.cpp b/Engine/Plugins/Experimental/TextureGraph/Source/TextureGraph/Private/TG_Var.cpp
--- a/Engine/Plugins/Experimental/TextureGraph/Source/TextureGraph/Private/TG_Var.cpp
+++ b/Engine/Plugins/Experimental/TextureGraph/Source/TextureGraph/Private/TG_Var.cpp
@@ -304,7 +304,7 @@ void VarPropertySerializer_StructProperty(FTG_Var::VarPropertySerialInfo& Info)
 	if (!WriterIt)
 	{
 		/// TODO: Perhaps think about doing a simple memcpy?
-		UE_LOG(LogTextureGraph, Log, TEXT("Fails serialize Var %s - Property %s FPClass %s CPPType %s"),
+		UE_LOG(LogTextureGraph, Warning, TEXT("Fails serialize Var %s - Property %s FPClass %s CPPType %s"),
 			*Info.Var->LogHead(),
 			*Info.Arg.GetName().ToString(),
 			*TypeName.ToString(),
@@ -420,7 +420,7 @@ bool FTG_Var::CopyGeneric(UTG_Expression* Owner, const FTG_Argument& Arg, bool C
 
 			if (!SerializerIt)
 			{
-				UE_LOG(LogTextureGraph, Log, TEXT("Fails serialize Var %s - Property %s FPClass %s CPPType %s"),
+				UE_LOG(LogTextureGraph, Warning, TEXT("Fails serialize Var %s - Property %s FPClass %s CPPType %s"),
 					*LogHead(),
 					*Arg.GetName().ToString(),
 					*PropertyClassName.ToString(),
@@ -479,11 +479,11 @@ void Generic_Simple_ArSerializer(FTG_Var::VarArchiveSerialInfo& Info)
 		Info.Var->ResetAs<T_ValueType>();
 	}
 	if (Info.Ar.IsSaving())
-		UE_LOG(LogTextureGraph, Log, TEXT("        Save Var %s: %s"), *Info.Var->GetId().ToString(), *Info.Var->LogValue());
+		UE_LOG(LogTextureGraph, VeryVerbose, TEXT("        Save Var %s: %s"), *Info.Var->GetId().ToString(), *Info.Var->LogValue());
 	T_ValueType& Value = Info.Var->EditAs<T_ValueType>();
 	Info.Ar << Value;
 	if (Info.Ar.IsLoading())
-		UE_LOG(LogTextureGraph, Log, TEXT("        Loaded Var %s: %s"), *Info.Var->GetId().ToString(), *Info.Var->LogValue());
+		UE_LOG(LogTextureGraph, VeryVerbose, TEXT("        Loaded Var %s: %s"), *Info.Var->GetId().ToString(), *Info.Var->LogValue());
 
 }
 
@@ -563,7 +563,7 @@ void FTG_Var::Serialize(FArchive& Ar, FTG_Id InPinId, const FTG_Argument& InArgu
 		}
 		else
 		{
-			UE_LOG(LogTextureGraph, Log, TEXT("serialize Var %s: NOT FOUND for %s"), *LogHead(), *InArgument.GetCPPTypeName().ToString());
+			UE_LOG(LogTextureGraph, Warning, TEXT("serialize Var %s: NOT FOUND for %s"), *LogHead(), *InArgument.GetCPPTypeName().ToString());
 		}
 	}
 }
diff --git a/Engine/Plugins/Experimental/TextureGraph/Source/TextureGraph/Private/TG_Variant.cpp b/Engine/Plugins/Experimental/TextureGraph/Source/TextureGraph/Private/TG_Variant.cpp
--- a/Engine/Plugins/Experimental/TextureGraph/Source/TextureGraph/Private/TG_Variant.cpp
+++ b/Engine/Plugins/Experimental/TextureGraph/Source/TextureGraph/Private/TG_Variant.cpp
@@ -39,6 +39,7 @@ FTG_Variant::FTG_Variant(FTG_Texture RHS)
 bool FTG_Variant::Serialize(FArchive& Ar)
 {
 	Ar << *this;
+	UE_LOG(LogTextureGraph, VeryVerbose, TEXT("FTG_Variant::Serialize"));
 	return true;
 }
 
diff --git a/Engine/Plugins/Experimental/TextureGraph/Source/TextureGraph/Private/TextureGraph.cpp b/Engine/Plugins/Experimental/TextureGraph/Source/TextureGraph/Private/TextureGraph.cpp
--- a/Engine/Plugins/Experimental/TextureGraph/Source/TextureGraph/Private/TextureGraph.cpp
+++ b/Engine/Plugins/Experimental/TextureGraph/Source/TextureGraph/Private/TextureGraph.cpp
@@ -107,7 +107,7 @@ void UTextureGraph::Serialize(FArchive& Ar)
 
 	int32 Version = Ar.CustomVer(FTG_CustomVersion::GUID);
 
-	UE_LOG(LogTextureGraph, Log, TEXT("%s TextureGraph: %s >>>> %s"),
+	UE_LOG(LogTextureGraph, Verbose, TEXT("%s TextureGraph: %s >>>> %s"),
 		(Ar.IsSaving() ? TEXT("Saved") : TEXT("Loaded")),
 		*GetName(),
 		*FString::FromInt(Version));
@@ -145,7 +145,7 @@ void UTextureGraph::PostLoad()
 void UTextureGraph::PreSave(FObjectPreSaveContext SaveContext)
 {
 	Super::PreSave(SaveContext);
-	UE_LOG(LogTextureGraph, Log, TEXT("PreSave Script: %s"), *GetName());
+	UE_LOG(LogTextureGraph, Verbose, TEXT("PreSave Script: %s"), *GetName());
 }
 
 void UTextureGraph::Update(MixUpdateCyclePtr InCycle)
@@ -188,6 +188,11 @@ void UTextureGraph::PostMeshLoad()
 	Invalidate(InvalidateInfo);
 }
 
+void UTextureGraph::FlushInvalidations()
+{
+	TextureGraphEngine::GetMixManager()->FlushMix(this);
+}
+
 void UTextureGraph::TriggerUpdate(bool Tweaking)
 {
 	FModelInvalidateInfo InvalidateInfo;
diff --git a/Engine/Plugins/Experimental/TextureGraph/Source/TextureGraph/Public/Blueprint/TG_BlueprintFunctionLibrary.h b/Engine/Plugins/Experimental/TextureGraph/Source/TextureGraph/Public/Blueprint/TG_BlueprintFunctionLibrary.h
--- a/Engine/Plugins/Experimental/TextureGraph/Source/TextureGraph/Public/Blueprint/TG_BlueprintFunctionLibrary.h
+++ b/Engine/Plugins/Experimental/TextureGraph/Source/TextureGraph/Public/Blueprint/TG_BlueprintFunctionLibrary.h
@@ -33,6 +33,22 @@ public:
 	UFUNCTION(BlueprintCallable, Category = "TextureGraph", meta = (Keywords = "GetScalarParameterValue", WorldContext = "WorldContextObject"))
 	static float GetScalarParameterValue(UObject* WorldContextObject, UTextureGraph* TextureGraph, FName ParameterName);
 
+	/** Sets a boolean parameter value on the TextureGraph instance. Logs if ParameterName is invalid. */
+	UFUNCTION(BlueprintCallable, Category = "TextureGraph", meta = (Keywords = "SetBoolParameterValue", WorldContext = "WorldContextObject"))
+	static void SetBoolParameterValue(UObject* WorldContextObject, UTextureGraph* TextureGraph, FName ParameterName, bool ParameterValue);
+
+	/** Gets a string parameter value from the TextureGraph instance. Logs if ParameterName is invalid. */
+	UFUNCTION(BlueprintCallable, Category = "TextureGraph", meta = (Keywords = "GetBoolParameterValue", WorldContext = "WorldContextObject"))
+	static bool GetBoolParameterValue(UObject* WorldContextObject, UTextureGraph* TextureGraph, FName ParameterName);
+
+	/** Sets a string parameter value on the TextureGraph instance. Logs if ParameterName is invalid. */
+	UFUNCTION(BlueprintCallable, Category = "TextureGraph", meta = (Keywords = "SetBoolParameterValue", WorldContext = "WorldContextObject"))
+	static void SetStringParameterValue(UObject* WorldContextObject, UTextureGraph* TextureGraph, FName ParameterName, FString ParameterValue);
+
+	/** Gets a String parameter value from the TextureGraph instance. Logs if ParameterName is invalid. */
+	UFUNCTION(BlueprintCallable, Category = "TextureGraph", meta = (Keywords = "GetBoolParameterValue", WorldContext = "WorldContextObject"))
+	static FString GetStringParameterValue(UObject* WorldContextObject, UTextureGraph* TextureGraph, FName ParameterName);
+
 	/** Sets a Vector parameter value on the TextureGraph instance. Logs if ParameterName is invalid. */
 	UFUNCTION(BlueprintCallable, Category = "TextureGraph", meta = (Keywords = "SetVectorParameterValue", WorldContext = "WorldContextObject"))
 	static void SetVectorParameterValue(UObject* WorldContextObject, UTextureGraph* TextureGraph, FName ParameterName, FVector4f ParameterValue);
@@ -59,6 +75,12 @@ public:
 	UFUNCTION(BlueprintCallable, Category = "TextureGraph", meta = (Keywords = "GetOutputSettingsParameterValue", WorldContext = "WorldContextObject"))
 	static FTG_OutputSettings GetSettingsParameterValue(UObject* WorldContextObject, UTextureGraph* TextureGraph, FName ParameterName , int& Width, int& Height);
 
+	/** Render the texture graph and return an array of texture render targets. */
+#if 0 /// This requires a lot of additional work. Will attempt after summer break
+	UFUNCTION(BlueprintCallable, Category = "TextureGraph", meta = (Keywords = "SetColorParameterValue", WorldContext = "WorldContextObject"))
+	static TArray<UTextureRenderTarget2D*> RenderTextureGraph(UObject* WorldContextObject, UTextureGraph* InTextureGraph);
+#endif 
+
 	static void AddParamWarning(FName ParamName, UObject* ObjectPtr, FString FunctionName);
 	static void AddError(UObject* ObjectPtr, FString FunctionName, FString Error);
 };
diff --git a/Engine/Plugins/Experimental/TextureGraph/Source/TextureGraph/Public/Expressions/Input/TG_Expression_InputParam.h b/Engine/Plugins/Experimental/TextureGraph/Source/TextureGraph/Public/Expressions/Input/TG_Expression_InputParam.h
--- a/Engine/Plugins/Experimental/TextureGraph/Source/TextureGraph/Public/Expressions/Input/TG_Expression_InputParam.h
+++ b/Engine/Plugins/Experimental/TextureGraph/Source/TextureGraph/Public/Expressions/Input/TG_Expression_InputParam.h
@@ -15,8 +15,8 @@ class TEXTUREGRAPH_API UTG_Expression_InputParam : public UTG_Expression
 
 protected:
 	// these 2 methods need to be implemented to provide concrete signatures for derived InputParam expression classes
-	virtual FTG_SignaturePtr BuildInputParameterSignature() const { return nullptr; }
-	virtual FTG_SignaturePtr BuildInputConstantSignature() const { return nullptr; }
+	virtual FTG_SignaturePtr BuildInputParameterSignature() const;
+	virtual FTG_SignaturePtr BuildInputConstantSignature() const;
 public:
 
 #if WITH_EDITOR
@@ -33,8 +33,6 @@ public:
 
 
 #define TG_DECLARE_INPUT_PARAM_EXPRESSION(Category) \
-	protected: virtual FTG_SignaturePtr BuildInputParameterSignature() const override; \
-	protected: virtual FTG_SignaturePtr BuildInputConstantSignature() const override; \
 	public:	virtual FTG_SignaturePtr GetSignature() const override { if (bIsConstant) { \
 			static FTG_SignaturePtr ConstantSignature = BuildInputConstantSignature(); return ConstantSignature; \
 		} else { \
diff --git a/Engine/Plugins/Experimental/TextureGraph/Source/TextureGraph/Public/Expressions/Input/TG_Expression_String.h b/Engine/Plugins/Experimental/TextureGraph/Source/TextureGraph/Public/Expressions/Input/TG_Expression_String.h
new file mode 100644
--- /dev/null
+++ b/Engine/Plugins/Experimental/TextureGraph/Source/TextureGraph/Public/Expressions/Input/TG_Expression_String.h
@@ -0,0 +1,31 @@
+// Copyright Epic Games, Inc. All Rights Reserved.
+
+#pragma once
+
+#include "CoreMinimal.h"
+#include "TG_Expression_InputParam.h"
+#include "TG_Expression_String.generated.h"
+
+UCLASS()
+class TEXTUREGRAPH_API UTG_Expression_String : public UTG_Expression_InputParam
+{
+	GENERATED_BODY()
+
+public:
+
+	virtual void Evaluate(FTG_EvaluationContext* InContext) override;
+
+	// The string constant
+	UPROPERTY(EditAnywhere, Category = NoCategory, meta = (TGType = "TG_InputParam", PinDisplayName = ""))
+		FString String;
+
+	// The output of the node, which is the string constant
+	UPROPERTY(meta = (TGType = "TG_Output", PinDisplayName = ""))
+		FString ValueOut;
+
+	virtual FTG_Name GetDefaultName() const override { return TEXT("String"); }
+	virtual void SetTitleName(FName NewName) override;
+	virtual FName GetTitleName() const override;
+	virtual FText GetTooltipText() const override { return FText::FromString(TEXT("String parameter that can be used as input in various places.")); } 
+};
+
diff --git a/Engine/Plugins/Experimental/TextureGraph/Source/TextureGraph/Public/Expressions/Input/TG_Expression_TexturePath.h b/Engine/Plugins/Experimental/TextureGraph/Source/TextureGraph/Public/Expressions/Input/TG_Expression_TexturePath.h
new file mode 100644
--- /dev/null
+++ b/Engine/Plugins/Experimental/TextureGraph/Source/TextureGraph/Public/Expressions/Input/TG_Expression_TexturePath.h
@@ -0,0 +1,41 @@
+// Copyright Epic Games, Inc. All Rights Reserved.
+
+#pragma once
+
+#include "CoreMinimal.h"
+
+#include "TG_Expression_InputParam.h"
+#include "Engine/Texture.h"
+#include "TG_Texture.h"
+
+#include "TG_Expression_TexturePath.generated.h"
+
+UCLASS()
+class TEXTUREGRAPH_API UTG_Expression_TexturePath : public UTG_Expression_InputParam
+{
+	GENERATED_BODY()
+
+public:
+
+	virtual void Evaluate(FTG_EvaluationContext* InContext) override;
+	virtual bool Validate(MixUpdateCyclePtr	Cycle) override;
+	
+	// The output of the node, which is the loaded texture from the path
+	UPROPERTY(meta = (TGType = "TG_Output", PinDisplayName = "", HideInnerPropertiesInNode))
+	FTG_Texture Output;
+
+	// Input file path of the texture
+	UPROPERTY(EditAnywhere, BlueprintReadWrite, Category = NoCategory, meta = (TGType = "TG_Input", TGPinNotConnectable, NoResetToDefault) )
+	FString Path;
+
+	// The input texture that was loaded from the path
+	UPROPERTY(meta = (TGType = "TG_InputParam"))
+	FTG_Texture Texture;
+
+	class ULayerChannel* Channel;
+	virtual FTG_Name GetDefaultName() const override { return TEXT("TexturePath"); }
+	virtual void SetTitleName(FName NewName) override;
+	virtual FName GetTitleName() const override;
+	virtual FText GetTooltipText() const override { return FText::FromString(TEXT("Loads a texture from a path.")); }
+};
+
diff --git a/Engine/Plugins/Experimental/TextureGraph/Source/TextureGraph/Public/TG_HelperFunctions.h b/Engine/Plugins/Experimental/TextureGraph/Source/TextureGraph/Public/TG_HelperFunctions.h
--- a/Engine/Plugins/Experimental/TextureGraph/Source/TextureGraph/Public/TG_HelperFunctions.h
+++ b/Engine/Plugins/Experimental/TextureGraph/Source/TextureGraph/Public/TG_HelperFunctions.h
@@ -23,9 +23,11 @@ public:
 
 	static void EnsureOutputIsTexture(MixUpdateCyclePtr Cycle, UTG_Node* OutputNode);
 
-	static AsyncBool ExportAsync(UTextureGraph* InTextureGraph, FString ExportPath, FString AssetName, FExportSettings& TargetExportSettings, bool OverrideExportPath, bool OverwriteTextures = true,bool ExportAllOutputs = false, bool bSave = true);
+	static JobBatchPtr InitExportBatch(UTextureGraph* InTextureGraph, FString ExportPath, FString AssetName, FExportSettings& TargetExportSettings, bool OverrideExportPath, bool OverwriteTextures, bool ExportAllOutputs, bool bSave);
+	static AsyncBool ExportAsync(UTextureGraph* InTextureGraph, FString ExportPath, FString AssetName, FExportSettings& TargetExportSettings, bool OverrideExportPath, bool OverwriteTextures = true, bool ExportAllOutputs = false, bool bSave = true);
+
+	static JobBatchPtr InitRenderBatch(UTextureGraph* InTextureGraph, JobBatchPtr ExistingBatch = nullptr);
 	static AsyncBool RenderAsync(UTextureGraph* InTextureGraph, JobBatchPtr ExistingBatch = nullptr);
-	
 	template <typename T_Type>
 	static TArray<T_Type> GetOutputsOfType(const UTG_Node* Node)
 	{
diff --git a/Engine/Plugins/Experimental/TextureGraph/Source/TextureGraph/Public/TG_Pin.h b/Engine/Plugins/Experimental/TextureGraph/Source/TextureGraph/Public/TG_Pin.h
--- a/Engine/Plugins/Experimental/TextureGraph/Source/TextureGraph/Public/TG_Pin.h
+++ b/Engine/Plugins/Experimental/TextureGraph/Source/TextureGraph/Public/TG_Pin.h
@@ -165,6 +165,8 @@ public:
 
 	// Check the type of the argument's Var expected to be contained in the Pin's Var
 	bool IsArgScalar() const { return Argument.IsScalar(); }
+	bool IsArgBool() const { return Argument.IsBool(); }
+	bool IsArgString() const { return Argument.IsString(); }
 	bool IsArgColor() const { return Argument.IsColor(); }
 	bool IsArgVector() const { return Argument.IsVector(); }
 	bool IsArgTexture() const { return Argument.IsTexture(); }
@@ -178,6 +180,12 @@ public:
 	// Scalar getter only on a scalar argument
 	bool GetValue(float& OutValue) const;
 
+	// bool getter only on a boolean argument
+	bool GetValue(bool& OutValue) const;
+
+	// bool getter only on a boolean argument
+	bool GetValue(FString& OutValue) const;
+
 	// Color or Vector getter only on a vector or color argument
 	bool GetValue(FLinearColor& OutValue) const;
 	bool GetValue(FVector4f& OutValue) const;
@@ -203,13 +211,16 @@ public:
 	// Set the self var value from string
 	// Used by UI to assign pin default value
 	// This is the same as the selfvar setter
-	void				SetValue(const FString& InValueStr, bool bIsTweaking = false);
+	void				FromString(const FString& InValueStr, bool bIsTweaking = false);
 
 	// Setters for Scalar, Color, Vector arguments
 	bool SetValue(float Value);
 	bool SetValue(double Value) { return SetValue((float)Value); }
 	bool SetValue(uint32 Value) { return SetValue((float)Value); }
 	bool SetValue(int32 Value) { return SetValue((float)Value); }
+	bool SetValue(bool bValue);
+	bool SetValue(FString Value);
+	bool SetValue(FName Value);
 	bool SetValue(const FLinearColor& Value);
 	bool SetValue(const FVector4f& Value);
 	bool SetValue(const float* Value, size_t Count);
diff --git a/Engine/Plugins/Experimental/TextureGraph/Source/TextureGraph/Public/TG_Signature.h b/Engine/Plugins/Experimental/TextureGraph/Source/TextureGraph/Public/TG_Signature.h
--- a/Engine/Plugins/Experimental/TextureGraph/Source/TextureGraph/Public/TG_Signature.h
+++ b/Engine/Plugins/Experimental/TextureGraph/Source/TextureGraph/Public/TG_Signature.h
@@ -151,6 +151,8 @@ public:
 	FName				GetCPPTypeName() const { return CPPTypeName; }
 
 	bool				IsScalar() const; // float or int
+	bool				IsBool() const; // boolean
+	bool				IsString() const; // string
 	bool				IsColor() const; // FlinearColor
 	bool				IsVector() const; // FLinearColor or ...
 	bool				IsTexture() const; // FTG_Texture or ...
diff --git a/Engine/Plugins/Experimental/TextureGraph/Source/TextureGraph/Public/TG_Var.h b/Engine/Plugins/Experimental/TextureGraph/Source/TextureGraph/Public/TG_Var.h
--- a/Engine/Plugins/Experimental/TextureGraph/Source/TextureGraph/Public/TG_Var.h
+++ b/Engine/Plugins/Experimental/TextureGraph/Source/TextureGraph/Public/TG_Var.h
@@ -175,26 +175,84 @@ private:
 		virtual FString LogValue() = 0;
 		virtual void SetValueFromString(const FString& String) = 0;
 		virtual TSharedPtr<FConcept> Clone() = 0;
+
+		/// Array related
+		virtual int32 Count() const 
+		{ 
+			return 1; 
+		}
+
+		virtual bool IsArray() 
+		{ 
+			return Count() > 1; 
+		}
 	};
+
 	template <class T>
 	struct TEXTUREGRAPH_API FModel : public FConcept
 	{
 		T Value;
-		FString LogValue() override { return TG_Var_LogValue(Value); }
+		FString LogValue() override 
+		{ 
+			return TG_Var_LogValue(Value); 
+		}
 	
-		void SetValueFromString(const FString& StrValue) override { TG_Var_SetValueFromString(Value, StrValue); }
+		void SetValueFromString(const FString& StrValue) override 
+		{ 
+			TG_Var_SetValueFromString(Value, StrValue); 
+		}
 
-		TSharedPtr<FConcept> Clone() override {
+		TSharedPtr<FConcept> Clone() override 
+		{
 			auto Ptr = MakeShared<FModel<T>>();
 			Ptr->Value = Value;
 			return Ptr;
 		}
 	};
 
+	template <class T>
+	struct TEXTUREGRAPH_API FModelArray : public FModel<TArray<T> >
+	{
+		using FModelBase = FModel<TArray<T>>;
+		virtual int32 Count() const override
+		{ 
+			return FModelBase::Value.Num(); 
+		}
+		virtual bool IsArray() 
+		{ 
+			return true; 
+		}
+
+		FString LogValue() override 
+		{ 
+			FString Log = "[";
+			for (int32 Index = 0; Index < FModelBase::Value.Num() - 1; Index++)
+				Log += (TG_Var_LogValue<T>(FModelBase::Value[Index]) + TEXT(","));
+
+			if (!FModelBase::Value.IsEmpty()) {
+				Log += TG_Var_LogValue<T>(FModelBase::Value[FModelBase::Value.Num() - 1]);
+			}
+
+			Log += TEXT("]");
+
+			return Log; 
+		}
+
+		void SetValueFromString(const FString& StrValue) override 
+		{ 
+			TArray<FString> ArrayValues;
+			StrValue.ParseIntoArray(ArrayValues, TEXT(","));
+			FModelBase::Value.Empty();
+			FModelBase::Value.SetNum(ArrayValues.Num());
+
+			for (int32 Index = 0; Index < ArrayValues.Num(); Index++)
+				TG_Var_SetValueFromString<T>(FModelBase::Value[Index], ArrayValues[Index]);
+		}
+	};
+
 	mutable TSharedPtr<FConcept> Concept;
 };
 
-
 template <> FString TG_Var_LogValue(bool& Value);
 template <> FString TG_Var_LogValue(int& Value);
 template <> FString TG_Var_LogValue(float& Value);
diff --git a/Engine/Plugins/Experimental/TextureGraph/Source/TextureGraph/Public/TextureGraph.h b/Engine/Plugins/Experimental/TextureGraph/Source/TextureGraph/Public/TextureGraph.h
--- a/Engine/Plugins/Experimental/TextureGraph/Source/TextureGraph/Public/TextureGraph.h
+++ b/Engine/Plugins/Experimental/TextureGraph/Source/TextureGraph/Public/TextureGraph.h
@@ -46,6 +46,7 @@ public:
 	void PostMeshLoad() override;
 
 	void TriggerUpdate(bool Tweaking);
+	void FlushInvalidations();
 
 	void UpdateGlobalTGSettings();
 
diff --git a/Engine/Plugins/Experimental/TextureGraph/Source/TextureGraphEditor/Private/EdGraph/TG_EdGraphNode.cpp b/Engine/Plugins/Experimental/TextureGraph/Source/TextureGraphEditor/Private/EdGraph/TG_EdGraphNode.cpp
--- a/Engine/Plugins/Experimental/TextureGraph/Source/TextureGraphEditor/Private/EdGraph/TG_EdGraphNode.cpp
+++ b/Engine/Plugins/Experimental/TextureGraph/Source/TextureGraphEditor/Private/EdGraph/TG_EdGraphNode.cpp
@@ -519,7 +519,7 @@ void UTG_EdGraphNode::PinDefaultValueChangedWithTweaking(UEdGraphPin* Pin, bool
 	const UTG_EdGraphSchema* Schema = Cast<const UTG_EdGraphSchema>(GetSchema());
 	UTG_Pin* TGPin = Schema->GetTGPinFromEdPin(Pin);
 
-	TGPin->SetValue(Pin->DefaultValue, bIsTweaking);
+	TGPin->FromString(Pin->DefaultValue, bIsTweaking);
 
 	// This updates the UObject (Texture/Material) picker UI in the Node to get updated
 	FProperty* Property = TGPin->GetExpressionProperty();
diff --git a/Engine/Plugins/Experimental/TextureGraph/Source/TextureGraphEditor/Private/TG_Editor.cpp b/Engine/Plugins/Experimental/TextureGraph/Source/TextureGraphEditor/Private/TG_Editor.cpp
--- a/Engine/Plugins/Experimental/TextureGraph/Source/TextureGraphEditor/Private/TG_Editor.cpp
+++ b/Engine/Plugins/Experimental/TextureGraph/Source/TextureGraphEditor/Private/TG_Editor.cpp
@@ -168,6 +168,7 @@ void FTG_Editor::OnTextureGraphPreSave(UObject* Object, FObjectPreSaveContext Sa
 	if (EditedTextureGraph == Object)
 		EditedTextureGraph->TriggerUpdate(false);
 }
+
 void FTG_Editor::InitEditor(const EToolkitMode::Type Mode, const TSharedPtr< class IToolkitHost >& InitToolkitHost, UTextureGraph* InTextureGraph)
 {
 	TG_Parameters = NewObject<UTG_Parameters>();
@@ -1588,6 +1589,14 @@ void FTG_Editor::OnClose()
 				
 	}
 	ITG_Editor::OnClose();
+
+	if (EditedTextureGraph)
+	{
+		/// We need to flush any invalidations coming for this graph. This is because if the user decided to save
+		/// graph on exit then this queues a mix update that never gets finished as the engine is being shutdown
+		/// and results in a cleanup assertion in Device.cpp
+		EditedTextureGraph->FlushInvalidations();
+	}
 }
 
 bool FTG_Editor::UpdateOriginalTextureGraph()
diff --git a/Engine/Plugins/Experimental/TextureGraph/Source/TextureGraphEngine/2D/Tex.cpp b/Engine/Plugins/Experimental/TextureGraph/Source/TextureGraphEngine/2D/Tex.cpp
--- a/Engine/Plugins/Experimental/TextureGraph/Source/TextureGraphEngine/2D/Tex.cpp
+++ b/Engine/Plugins/Experimental/TextureGraph/Source/TextureGraphEngine/2D/Tex.cpp
@@ -33,6 +33,7 @@
 #include "UnrealClient.h"
 #include "Widgets/SViewport.h"
 #include <Rendering/Texture2DResource.h>
+#include "Math/GuardedInt.h"
 
 #include "Profiling/StatGroup.h"
 #include "UObject/Package.h"
@@ -52,7 +53,6 @@
 #pragma pop_macro("DLLEXPORT")
 
 DECLARE_CYCLE_STAT(TEXT("Tex_InitRT"), STAT_Tex_InitRT, STATGROUP_TextureGraphEngine);
-extern ENGINE_API UEngine* GEngine;
 
 TexDescriptor::TexDescriptor()
 {
@@ -62,7 +62,7 @@ TexDescriptor::TexDescriptor(uint32 InWidth, uint32 InHeight, EPixelFormat InFor
 	: Name(Util::RandomID())
 	, Width(InWidth)
 	, Height(InHeight)
-	, Format(InFormat)
+	, Format(Format)
 	, NumChannels(TextureHelper::GetChannelsFromPixelFormat(InFormat))
 {
 }
@@ -502,23 +502,63 @@ AsyncActionResultPtr Tex::LoadFlat()
 	});
 }
 
-
 UTexture2D* Tex::CreateTexture(uint32 Width, uint32 Height, EPixelFormat Format, bool sRGB, UObject* Package)
 {
-	UTexture2D* NewTexture = NULL;
+	UTexture2D* NewTexture = nullptr;
 
 	if (Width > 0 && Height > 0 &&
 		(Width % GPixelFormats[Format].BlockSizeX) == 0 &&
 		(Height % GPixelFormats[Format].BlockSizeY) == 0)
 	{
-		NewTexture = UTexture2D::CreateTransient(Width, Height, Format, *Desc.Name);
+		if (!Package)
+			Package = Util::GetTexturesPackage();
+
+		const FName UniqueName = MakeUniqueObjectName(Package, UTexture2D::StaticClass(), *Desc.Name);
+		NewTexture = NewObject<UTexture2D>(Package, UniqueName);
 		NewTexture->SRGB = sRGB;
-		//NewObject<UTexture2D>(Package, *_desc.Name, UModelObject::s_maskNoGC);
-	//	NewTexture->PlatformData->PixelFormat = Format;
-	//	NewTexture->PlatformData = new FTexturePlatformData();
-	//	NewTexture->PlatformData->SizeX = Width;
-	//	NewTexture->PlatformData->SizeY = Height;
-	////	NewTexture->PlatformData->PixelFormat = Format;
+
+		const int32 NumBlocksX = Width / GPixelFormats[Format].BlockSizeX;
+		const int32 NumBlocksY = Height / GPixelFormats[Format].BlockSizeY;
+
+		if (Width <= 0 || Height <= 0)
+		{
+			UE_LOG(LogTexture, Warning, TEXT("Negative size specified for UTexture2D::CreateTransient()"));
+			return nullptr;
+		}
+
+		if ((Width % GPixelFormats[Format].BlockSizeX) ||
+			(Height % GPixelFormats[Format].BlockSizeY))
+		{
+			UE_LOG(LogTexture, Warning, TEXT("Size specified isn't valid for block-based pixel format in UTexture2D::CreateTransient()"));
+			return nullptr;
+		}
+
+		FGuardedInt64 BytesForImageValidation = FGuardedInt64(NumBlocksX) * NumBlocksY * GPixelFormats[Format].BlockBytes;
+		if (BytesForImageValidation.IsValid() == false)
+		{
+			UE_LOG(LogTexture, Warning, TEXT("Size specified overflows in UTexture2D::CreateTransient()"));
+			return nullptr;
+		}
+
+		int64 BytesForImage = BytesForImageValidation.Get(0);
+
+		NewTexture->SetPlatformData(new FTexturePlatformData());
+		NewTexture->GetPlatformData()->SizeX = Width;
+		NewTexture->GetPlatformData()->SizeY = Height;
+		NewTexture->GetPlatformData()->SetNumSlices(1);
+		NewTexture->GetPlatformData()->PixelFormat = Format;
+
+		// Allocate first mipmap.
+		FTexture2DMipMap* Mip = new FTexture2DMipMap(Width, Height, 1);
+		NewTexture->GetPlatformData()->Mips.Add(Mip);
+		Mip->BulkData.Lock(LOCK_READ_WRITE);
+		void* DestImageData = Mip->BulkData.Realloc(BytesForImage);
+		memset(DestImageData, 0, BytesForImage);
+		Mip->BulkData.Unlock();
+
+		NewTexture->UpdateResource();
+
+		return NewTexture;
 
 	}
 	else
@@ -643,14 +683,32 @@ UTexture2D* Tex::InitTextureDefault(int32 Width, int32 Height, EPixelFormat Pixe
 	check(IsInGameThread());
 	check(UncompressedData);
 
-	UTexture2D* NewTexture;
-	
 	Desc.Name += MakeUniqueObjectName(Package, UTexture2D::StaticClass()).ToString(); // , FName());
 
-	NewTexture = CreateTexture(Width, Height, PixelFormat, Desc.bIsSRGB, Package);
+	UTexture2D* NewTexture = CreateTexture(Width, Height, PixelFormat, Desc.bIsSRGB, Package);
 
 	if (NewTexture)
 	{
+#if 0
+		check(UncompressedData && UncompressedDataSize);
+
+		NewTexture->SetPlatformData(new FTexturePlatformData());
+		NewTexture->GetPlatformData()->SizeX = Width;
+		NewTexture->GetPlatformData()->SizeY = Height;
+		NewTexture->GetPlatformData()->SetNumSlices(1);
+		NewTexture->GetPlatformData()->PixelFormat = PixelFormat;
+
+		// Allocate first mipmap.
+		FTexture2DMipMap* Mip = new FTexture2DMipMap(Width, Height, 1);
+		NewTexture->GetPlatformData()->Mips.Add(Mip);
+
+		Mip->BulkData.Lock(LOCK_READ_WRITE);
+		void* DestImageData = Mip->BulkData.Realloc(UncompressedDataSize);
+		FMemory::Memcpy(DestImageData, UncompressedData, UncompressedDataSize);
+		Mip->BulkData.Unlock();
+		NewTexture->UpdateResource();
+#endif 
+
 		uint8* MipData = static_cast<uint8*>(NewTexture->GetPlatformData()->Mips[0].BulkData.Lock(LOCK_READ_WRITE));
 		int64 MipSize = NewTexture->GetPlatformData()->Mips[0].BulkData.GetBulkDataSize();
 
@@ -658,16 +716,9 @@ UTexture2D* Tex::InitTextureDefault(int32 Width, int32 Height, EPixelFormat Pixe
 
 		// Bulk Data was already allocated for the correct size when we called CreateTransient above
 		FMemory::Memcpy(MipData, UncompressedData, MipSize);
-
 		NewTexture->GetPlatformData()->Mips[0].BulkData.Unlock();
-		
-		// If mipmaps are required, this Tex will be turned into a rendertarget and autogenerate mips from mip0
-		// If needed, We could allocate and fill the mipmaps from sys mem Data in the flow here
-		// if (_desc.mipmaps)
-		//		AllocateTextureMips(_texture, some_init_data_for_the_mipmaps);
 
 		NewTexture->UpdateResource();
-		
 	}
 
 	return NewTexture;
@@ -1046,7 +1097,18 @@ AsyncTiledBlobRef Tex::ToBlob(int32 XTiles, int32 YTiles, uint32 Width /* = 0 */
 	Desc.Width = Width;
 	Desc.Height = Height;
 
-	if (!TextureHelper::CanSplitToTiles(Texture,XTiles,YTiles))
+	int32 CheckSizeX = Width;
+	int32 CheckSizeY = Height;
+
+#if WITH_EDITOR
+	if (Texture->Source.IsValid())
+	{
+		CheckSizeX = Texture->Source.GetSizeX();
+		CheckSizeY = Texture->Source.GetSizeY();
+	} 
+#endif 
+
+	if (!TextureHelper::CanSplitToTiles(CheckSizeX, CheckSizeY, XTiles, YTiles))
 	{
 		/// TOOD: need to properly calculate Hash over here
 		return ToSingleBlob(nullptr, TransferToRT);
@@ -1181,7 +1243,7 @@ bool Tex::IsValidVirtualTexture()
 	return(IsValid(Texture) && Texture->IsCurrentlyVirtualTextured()); 
 }
 
-void Tex::LoadAsset(FSoftObjectPath& SoftPath, const DesiredImageProperties* Props /* = nullptr */)
+bool Tex::LoadAsset(FSoftObjectPath& SoftPath, const DesiredImageProperties* Props /* = nullptr */)
 {
 	check(IsInGameThread());
 
@@ -1224,9 +1286,10 @@ void Tex::LoadAsset(FSoftObjectPath& SoftPath, const DesiredImageProperties* Pro
 	Desc.Format = OutPixelFormat;
 #endif
 	
+	return true;
 }
 
-AsyncActionResultPtr Tex::LoadAsync(const FString& Filename, const DesiredImageProperties* Props /* = nullptr */)
+bool Tex::LoadFile(const FString& Filename, const DesiredImageProperties* Props /* = nullptr */)
 {
 	if (Props)
 	{
@@ -1240,69 +1303,65 @@ AsyncActionResultPtr Tex::LoadAsync(const FString& Filename, const DesiredImageP
 			Desc.bMipMaps = true;
 	}
 
-	return cti::make_continuable<ActionResultPtr>([this, Filename](auto&& LoadPromise)
+	TArray<uint8> Buffer;
+	UE_LOG(LogTexture, Log, TEXT("Trying to load file %s"), *Filename);
+	bool DidLoadBuffer = false;
+
+	if (FFileHelper::LoadFileToArray(Buffer, *Filename))
 	{
-		AsyncTask(ENamedThreads::AnyBackgroundThreadNormalTask, [this, Filename, FWD_PROMISE(LoadPromise)]() mutable
+		EPixelFormat PixelFormat;
+		int32 Width;
+		int32 Height;
+		TArray<uint8> UncompressedData;
+
+		bool Success = CopyImageToBuffer(PixelFormat, Width, Height, Buffer, UncompressedData);
+
+		if (Success)
 		{
-			TArray<uint8> Buffer;
-			UE_LOG(LogTexture, Log, TEXT("Trying to load file %s"), *Filename);
-			bool DidLoadBuffer = false;
+			DidLoadBuffer = true;
+			UE_LOG(LogTexture, Log, TEXT("Loading file %s Success"), *Filename);
 
-			if (FFileHelper::LoadFileToArray(Buffer, *Filename))
+			try
 			{
-				EPixelFormat PixelFormat;
-				int32 Width;
-				int32 Height;
-				TArray<uint8> UncompressedData;
-
-				bool Success = CopyImageToBuffer(PixelFormat, Width, Height, Buffer, UncompressedData);
+				auto Package = Util::GetTexturesPackage();
+				UTexture2D* NewTexture = nullptr;
 
-				if (Success)
+				if (FPaths::GetExtension(Filename) == TEXT("HDR"))
 				{
-					DidLoadBuffer = true;
-					UE_LOG(LogTexture, Log, TEXT("Loading file %s Success"), *Filename);
-					AsyncTask(ENamedThreads::GameThread, [this, Filename, PixelFormat, Width, Height, UncompressedData = std::move(UncompressedData), LoadPromise = std::forward<decltype(LoadPromise)>(LoadPromise)]() mutable
-					{
-						try
-						{
-							auto Package = Util::GetTexturesPackage();
-							UTexture2D* NewTexture = nullptr;
-
-							if (FPaths::GetExtension(Filename) == TEXT("HDR"))
-							{
-								//NewTexture = InitTextureHDR(Buffer, Package);
-							}
-							else
-							{
-								NewTexture = InitTextureDefault(Width, Height, PixelFormat, UncompressedData.GetData(), UncompressedData.Num(), Package);
-							}
-
-							 TexDescriptor NewDesc(NewTexture);
-
-							 NewDesc.ClearColor = Desc.ClearColor;
-							 NewDesc.bMipMaps = Desc.bMipMaps;
-							 Desc = NewDesc;
-
-							Texture = NewTexture;
-
-							LoadPromise.set_value(std::make_shared<ActionResult>(nullptr));
-						}
-						catch (const std::exception_ptr e)
-						{
-							LoadPromise.set_exception(e);
-						}
-					});
+					//NewTexture = InitTextureHDR(Buffer, Package);
+				}
+				else
+				{
+					NewTexture = InitTextureDefault(Width, Height, PixelFormat, UncompressedData.GetData(), UncompressedData.Num(), Package);
 				}
-			}
 
-			if (!DidLoadBuffer)
+				TexDescriptor NewDesc(NewTexture);
+
+				NewDesc.ClearColor = Desc.ClearColor;
+				NewDesc.bMipMaps = Desc.bMipMaps;
+				Desc = NewDesc;
+
+				Texture = NewTexture;
+			}
+			catch (const std::exception_ptr e)
 			{
-				/// Set the error on the Promise
-				LoadPromise.set_exception(std::make_exception_ptr(std::runtime_error("Unable to load image Buffer!")));
-				UE_LOG(LogTexture, Log, TEXT("Loading file %s failed"), *Filename);
+				UE_LOG(LogTexture, Error, TEXT("Exception while loading filename: %s"), *Filename);
+				return false;
 			}
-		});
-	});
+		}
+		else
+		{
+			UE_LOG(LogTexture, Log, TEXT("CopyImageToBuffer failed for file: "), *Filename);
+			return false;
+		}
+	}
+	else
+	{
+		UE_LOG(LogTexture, Log, TEXT("FFileHelper::LoadFileToArray failed for file: "), *Filename);
+		return false;
+	}
+
+	return true;
 }
 
 UTexture* Tex::GetTexture() const
diff --git a/Engine/Plugins/Experimental/TextureGraph/Source/TextureGraphEngine/2D/Tex.h b/Engine/Plugins/Experimental/TextureGraph/Source/TextureGraphEngine/2D/Tex.h
--- a/Engine/Plugins/Experimental/TextureGraph/Source/TextureGraphEngine/2D/Tex.h
+++ b/Engine/Plugins/Experimental/TextureGraph/Source/TextureGraphEngine/2D/Tex.h
@@ -140,8 +140,8 @@ public:
 	void							ReleaseTexture();
 	void							TransferTextureToRT(FRHICommandListImmediate& RHI, UTexture2D** PrevTexture, bool FreeAfterUse);
 	void							TransferVirtualTextureToRT(FRHICommandListImmediate& RHI, UTexture2D** TextureToTransfer, bool FreeAfterUse);
-	void							LoadAsset(FSoftObjectPath& SoftPath, const DesiredImageProperties* Props = nullptr);
-	AsyncActionResultPtr			LoadAsync(const FString& Filename, const DesiredImageProperties* Props = nullptr);
+	bool							LoadAsset(FSoftObjectPath& SoftPath, const DesiredImageProperties* Props = nullptr);
+	bool							LoadFile(const FString& Filename, const DesiredImageProperties* Props = nullptr);
 	AsyncActionResultPtr			LoadFlat();
 	AsyncTiledBlobRef				ToBlob(int32 XTiles, int32 YTiles, uint32 Width = 0, uint32 Height = 0, bool TransferToRT = false); ///There is a visual difference between having material transfer to RT or directly copy from Texture
 	AsyncTiledBlobRef				ToSingleBlob(CHashPtr Hash, bool TransferToRT = false, bool ResolveOnRenderThread = false, bool NoCache = false);
diff --git a/Engine/Plugins/Experimental/TextureGraph/Source/TextureGraphEngine/2D/TextureHelper.cpp b/Engine/Plugins/Experimental/TextureGraph/Source/TextureGraphEngine/2D/TextureHelper.cpp
--- a/Engine/Plugins/Experimental/TextureGraph/Source/TextureGraphEngine/2D/TextureHelper.cpp
+++ b/Engine/Plugins/Experimental/TextureGraph/Source/TextureGraphEngine/2D/TextureHelper.cpp
@@ -776,37 +776,43 @@ RawBufferPtr TextureHelper::RawFromTexture(UTexture2D* Texture, const BufferDesc
 {
 	check(IsInRenderingThread());
 
+	uint8* RawData = nullptr;
+	size_t DataSize = 0;
+
 	//For editor we are using the source texture until we support the compressed Formats
 #if WITH_EDITOR
-	check(Texture->Source.IsValid());
-	uint8* RawData = nullptr;
-	// Possibly the data was stored in TextureSource rather than PlatformData
-	// The difference is the source data (residing in Utexture) is editor only and can be saved to disk
-	// This is mostly used for thumbnail, where thumbnail data can exist across sessions (e.g UAsset)
-	// platform data, as the name suggests, is used during cook to convert to platform specific texture. 
-	// This is legal because we might not be cooking our data.
-
-	TArray64<uint8> MipDataSrc;
-	Texture->Source.GetMipData(MipDataSrc, 0);
-	size_t DataSize = MipDataSrc.Num();
-	RawData = new uint8[DataSize];
-
-	FMemory::Memcpy(RawData, MipDataSrc.GetData(), DataSize);
-#else
-	const uint8* MipData = static_cast<uint8*>(Texture->GetPlatformData()->Mips[0].BulkData.Lock(LOCK_READ_ONLY));
-	size_t DataSize = Texture->GetPlatformData()->Mips[0].BulkData.GetBulkDataSize();
-	uint8* RawData = nullptr;
-	if (DataSize != 0)
+	if (Texture->Source.IsValid())
 	{
-		size_t DescSize = Desc.Size();
-		check(DataSize == DescSize);
+		// Possibly the data was stored in TextureSource rather than PlatformData
+		// The difference is the source data (residing in Utexture) is editor only and can be saved to disk
+		// This is mostly used for thumbnail, where thumbnail data can exist across sessions (e.g UAsset)
+		// platform data, as the name suggests, is used during cook to convert to platform specific texture. 
+		// This is legal because we might not be cooking our data.
+
+		TArray64<uint8> MipDataSrc;
+		Texture->Source.GetMipData(MipDataSrc, 0);
+		DataSize = MipDataSrc.Num();
 		RawData = new uint8[DataSize];
-
-		// Bulk data was already allocated for the correct size when we called CreateTransient above
-		FMemory::Memcpy(RawData, MipData, DataSize);
-		Texture->GetPlatformData()->Mips[0].BulkData.Unlock();
+		FMemory::Memcpy(RawData, MipDataSrc.GetData(), DataSize);
 	}
 #endif
+
+	if (!RawData)
+	{
+		const uint8* MipData = static_cast<uint8*>(Texture->GetPlatformData()->Mips[0].BulkData.Lock(LOCK_READ_ONLY));
+		DataSize = Texture->GetPlatformData()->Mips[0].BulkData.GetBulkDataSize();
+
+		if (DataSize != 0)
+		{
+			size_t DescSize = Desc.Size();
+			check(DataSize == DescSize);
+			RawData = new uint8[DataSize];
+
+			// Bulk data was already allocated for the correct size when we called CreateTransient above
+			FMemory::Memcpy(RawData, MipData, DataSize);
+			Texture->GetPlatformData()->Mips[0].BulkData.Unlock();
+		}
+	}
 	
 	return std::make_shared<RawBuffer>(RawData, DataSize, Desc);
 }
@@ -1450,23 +1456,15 @@ bool TextureHelper::CanSupportTexture(UTexture* Tex)
 	return CanSupport;
 }
 
-bool TextureHelper::CanSplitToTiles(UTexture* Texture, int TilesX, int TilesY)
+bool TextureHelper::CanSplitToTiles(int Width, int Height, int TilesX, int TilesY)
 {
 	bool bSplitToTiles = true;
 
-#if WITH_EDITOR
-	check(Texture->Source.IsValid());
-	int Width = Texture->Source.GetSizeX();
-	int Height = Texture->Source.GetSizeY();
-#else
-	int Width = Texture->GetSurfaceWidth();
-	int Height = Texture->GetSurfaceHeight();
-#endif
 	bool bSingleTile = (TilesX == 1 && TilesY == 1);
 	//Height and width both must be big enough to support tiling
 	bool bSizeNotBigEnough = (Width <= TilesX || Height <= TilesY);
-	
-	if(bSingleTile || bSizeNotBigEnough)
+
+	if (bSingleTile || bSizeNotBigEnough)
 	{
 		bSplitToTiles = false;
 	}
diff --git a/Engine/Plugins/Experimental/TextureGraph/Source/TextureGraphEngine/2D/TextureHelper.h b/Engine/Plugins/Experimental/TextureGraph/Source/TextureGraphEngine/2D/TextureHelper.h
--- a/Engine/Plugins/Experimental/TextureGraph/Source/TextureGraphEngine/2D/TextureHelper.h
+++ b/Engine/Plugins/Experimental/TextureGraph/Source/TextureGraphEngine/2D/TextureHelper.h
@@ -95,7 +95,7 @@ struct TEXTUREGRAPHENGINE_API TextureHelper
 	static FLinearColor					GetPixelValueFromRaw(RawBufferPtr RawObj, int32 Width, int32 Height, int32 X, int32 Y);
 	static AsyncBool					ExportRaw(RawBufferPtr RawObj, const FString& CompletePath);
 	static bool							CanSupportTexture(UTexture* Tex);
-	static bool							CanSplitToTiles(UTexture* Texture, int TilesX, int TilesY);
+	static bool							CanSplitToTiles(int Width, int Height, int TilesX, int TilesY);
 	static size_t						RoundUpTo(size_t Size, size_t DesiredRounding);
 
 	//////////////////////////////////////////////////////////////////////////
diff --git a/Engine/Plugins/Experimental/TextureGraph/Source/TextureGraphEngine/Model/Mix/MixManager.cpp b/Engine/Plugins/Experimental/TextureGraph/Source/TextureGraphEngine/Model/Mix/MixManager.cpp
--- a/Engine/Plugins/Experimental/TextureGraph/Source/TextureGraphEngine/Model/Mix/MixManager.cpp
+++ b/Engine/Plugins/Experimental/TextureGraph/Source/TextureGraphEngine/Model/Mix/MixManager.cpp
@@ -32,6 +32,12 @@ MixManager::~MixManager()
 	Queue.clear();
 }
 
+void MixManager::FlushMix(UMixInterface* MixObj)
+{
+	check(IsInGameThread());
+	MixesToFlush.push_back(MixObj);
+}
+
 void MixManager::Update(float Delta)
 {
 	/// Mix manager has been suspended for the time being. 
@@ -44,6 +50,25 @@ void MixManager::Update(float Delta)
 	if (!MixInfos.size())
 		return;
 
+	if (!MixesToFlush.empty())
+	{
+		for (UMixInterface* MixToFlush : MixesToFlush)
+		{
+			for (auto Iter = MixInfos.begin(); Iter != MixInfos.end();)
+			{
+				MixInvalidateInfo& MixInfo = *Iter;
+				UMixInterface* MixObj = MixInfo.MixObj;
+
+				if (MixObj == MixToFlush)
+					Iter = MixInfos.erase(Iter);
+				else
+					Iter++;
+			}
+		}
+
+		MixesToFlush.clear();
+	}
+
 	std::vector<MixInvalidateInfo> BucketedMixInfos;
 	
 	// First merge invalidation touching on the same mix
diff --git a/Engine/Plugins/Experimental/TextureGraph/Source/TextureGraphEngine/Model/Mix/MixManager.h b/Engine/Plugins/Experimental/TextureGraph/Source/TextureGraphEngine/Model/Mix/MixManager.h
--- a/Engine/Plugins/Experimental/TextureGraph/Source/TextureGraphEngine/Model/Mix/MixManager.h
+++ b/Engine/Plugins/Experimental/TextureGraph/Source/TextureGraphEngine/Model/Mix/MixManager.h
@@ -41,6 +41,7 @@ class TEXTUREGRAPHENGINE_API MixManager
 private:
 	std::atomic_bool				bIsSuspended = false;			/// Whether the mix manager is suspended or not
 	MixPriorityQueue				Queue;							/// Mix queue
+	std::vector<UMixInterface*>		MixesToFlush;						/// The mixes to flush
 	
 public:
 									MixManager();
@@ -52,6 +53,7 @@ public:
 	void							Exit();
 
 	void							InvalidateMix(UMixInterface* MixObj, const FInvalidationDetails &Details, int32 Priority = (int32)E_Priority::kNormal);
+	void							FlushMix(UMixInterface* MixObj);
 	
 	//////////////////////////////////////////////////////////////////////////
 	/// Inline functions
diff --git a/Engine/Plugins/Experimental/TextureGraph/Source/TextureGraphEngine/Model/StaticImageResource.cpp b/Engine/Plugins/Experimental/TextureGraph/Source/TextureGraphEngine/Model/StaticImageResource.cpp
--- a/Engine/Plugins/Experimental/TextureGraph/Source/TextureGraphEngine/Model/StaticImageResource.cpp
+++ b/Engine/Plugins/Experimental/TextureGraph/Source/TextureGraphEngine/Model/StaticImageResource.cpp
@@ -20,21 +20,41 @@ AsyncTiledBlobRef UStaticImageResource::Load(MixUpdateCyclePtr Cycle)
 		return cti::make_ready_continuable(TiledBlobRef(TextureHelper::GBlack));
 	
 	// if the path is a valid package path (content browser asset)
-	check(FPackageName::IsValidPath(FileName));
-	Tex* TexObj = nullptr;
+	check(bIsFilesystem || FPackageName::IsValidPath(FileName));
 
 	// we load the asset -> UTexture2d -> Tex -> Blob and create a SourceAsset from that
-	return PromiseUtil::OnGameThread().then([=](int32) mutable
+	return PromiseUtil::OnGameThread().then([=, this](int32) mutable
 	{
-		FSoftObjectPath SoftPath(FileName);
+		Tex* TexObj = nullptr;
 
+		// we load the asset -> UTexture2d -> Tex -> Blob and create a SourceAsset from that
 		TexObj = new Tex();
-		TexObj->LoadAsset(SoftPath);
+		bool bDidLoad = false;
 
-		int32 NumXTiles = Cycle->GetMix()->GetNumXTiles();
-		int32 NumYTiles = Cycle->GetMix()->GetNumYTiles();
+		try
+		{
+			if (!bIsFilesystem)
+			{
+				FSoftObjectPath SoftPath(FileName);
+				bDidLoad = TexObj->LoadAsset(SoftPath);
+			}
+			else
+				bDidLoad = TexObj->LoadFile(FileName);
+		}
+		catch (std::exception e)
+		{
+		}
 
-		return TexObj->ToBlob(NumXTiles, NumYTiles, 0, 0, false);
+		if (bDidLoad)
+		{
+			int32 NumXTiles = Cycle->GetMix()->GetNumXTiles();
+			int32 NumYTiles = Cycle->GetMix()->GetNumYTiles();
+
+			return TexObj->ToBlob(NumXTiles, NumYTiles, 0, 0, false);
+		}
+
+		BlobObj = TextureHelper::GetMagenta();
+		return (AsyncTiledBlobRef)cti::make_ready_continuable<TiledBlobRef>(TiledBlobRef(BlobObj));
 	})
 	.then([=](TiledBlobRef LoadedBlob) mutable
 	{
diff --git a/Engine/Plugins/Experimental/TextureGraph/Source/TextureGraphEngine/Model/StaticImageResource.h b/Engine/Plugins/Experimental/TextureGraph/Source/TextureGraphEngine/Model/StaticImageResource.h
--- a/Engine/Plugins/Experimental/TextureGraph/Source/TextureGraphEngine/Model/StaticImageResource.h
+++ b/Engine/Plugins/Experimental/TextureGraph/Source/TextureGraphEngine/Model/StaticImageResource.h
@@ -21,8 +21,8 @@ private:
 	
 	UPROPERTY()
 	FString 						AssetUUID;						/// Unique id for the asset within the entire system
-	
 	TiledBlobPtr					BlobObj;						/// The blob that represents the data for this source
+	bool							bIsFilesystem = false;			/// Is loading directly from the filesystem
 
 	virtual AsyncTiledBlobRef		Load(MixUpdateCyclePtr Cycle);
 	
@@ -36,5 +36,7 @@ public:
 	//////////////////////////////////////////////////////////////////////////
 	FORCEINLINE const FString&		GetAssetUUID() const { return AssetUUID; }
 	FORCEINLINE void				SetAssetUUID(const FString& UUID) { AssetUUID = UUID; }
+	FORCEINLINE void				SetIsFileSystem(bool bInIsFileSystem) { bIsFilesystem = bInIsFileSystem; }
+	FORCEINLINE bool				IsFileSystem() const { return bIsFilesystem; }
 };
 
