Fixing String and TexturePath nodes

diff --git a/Engine/Plugins/Experimental/TextureGraph/Source/TextureGraph/Private/Expressions/Input/TG_Expression_TexturePath.cpp b/Engine/Plugins/Experimental/TextureGraph/Source/TextureGraph/Private/Expressions/Input/TG_Expression_TexturePath.cpp
--- a/Engine/Plugins/Experimental/TextureGraph/Source/TextureGraph/Private/Expressions/Input/TG_Expression_TexturePath.cpp
+++ b/Engine/Plugins/Experimental/TextureGraph/Source/TextureGraph/Private/Expressions/Input/TG_Expression_TexturePath.cpp
@@ -6,32 +6,64 @@
 #include "2D/TextureHelper.h"
 #include "Model/StaticImageResource.h"
 
-void UTG_Expression_TexturePath::Evaluate(FTG_EvaluationContext* InContext)
+// Special case for TexturePath Constant signature, we want to keep the Path Input connectable in that case
+// so do this in the override version of BuildInputConstantSignature()
+FTG_SignaturePtr  UTG_Expression_TexturePath::BuildInputConstantSignature() const
 {
-	Super::Evaluate(InContext);
-	if (Texture)
+	FTG_Signature::FInit SignatureInit = GetSignatureInitArgsFromClass();
+	for (auto& Arg : SignatureInit.Arguments)
 	{
-		Output = Texture;
+		if (Arg.IsInput() && Arg.IsParam())
+		{
+			Arg.ArgumentType = Arg.ArgumentType.Unparamed();
+		}
 	}
-	else if (!Path.IsEmpty())
+	return MakeShared<FTG_Signature>(SignatureInit);
+}
+
+bool UTG_Expression_TexturePath::ValidateInputPath(FString& ValidatedPath) const
+{
+	// empty but that's ok
+	if (Path.IsEmpty())
 	{
-		UStaticImageResource* StaticImageResource = UStaticImageResource::CreateNew<UStaticImageResource>();
+		return true;
+	}
 
-		FString LocalPath = Path;
-		FPackagePath PackagePath;
-		FString PathExt = FPaths::GetExtension(Path);
+	// Check that the local path exists
+	FString LocalPath = Path.TrimQuotes();
+	FPackagePath PackagePath;
+	FString PathExt = FPaths::GetExtension(Path);
 
-		if (FPackagePath::TryFromMountedName(Path, PackagePath))
-		{
-			LocalPath = PackagePath.GetLocalFullPath();
-			FString LocalPathExt = FPaths::GetExtension(LocalPath);
+	// Try to find a file in a mounted package
+	if (FPackagePath::TryFromMountedName(LocalPath, PackagePath))
+	{
+		LocalPath = PackagePath.GetLocalFullPath();
+		FString LocalPathExt = FPaths::GetExtension(LocalPath);
 
-			if (LocalPathExt != PathExt)
-			{
-				LocalPath = FPaths::ChangeExtension(LocalPath, PathExt);
-			}
+		if (LocalPathExt != PathExt)
+		{
+			LocalPath = FPaths::ChangeExtension(LocalPath, PathExt);
 		}
+		ValidatedPath = LocalPath;
+		return true;
+	}
+	else if(FPaths::FileExists(LocalPath))
+	{
+		ValidatedPath = LocalPath;
+		return true;
+	}
+
+	return false;
+}
+
+void UTG_Expression_TexturePath::Evaluate(FTG_EvaluationContext* InContext)
+{
+	Super::Evaluate(InContext);
 
+	FString LocalPath;
+	if (!Path.IsEmpty() && ValidateInputPath(LocalPath))
+	{
+		UStaticImageResource* StaticImageResource = UStaticImageResource::CreateNew<UStaticImageResource>();
 		StaticImageResource->SetAssetUUID(LocalPath);
 		StaticImageResource->SetIsFileSystem(true);
 
@@ -51,8 +83,17 @@ void UTG_Expression_TexturePath::Evaluate(FTG_EvaluationContext* InContext)
 
 bool UTG_Expression_TexturePath::Validate(MixUpdateCyclePtr Cycle)
 {
+	FString LocalPath;
+	if (!ValidateInputPath(LocalPath)) 
+	{
+		auto ErrorType = static_cast<int32>(ETextureGraphErrorType::NODE_WARNING);
+		TextureGraphEngine::GetErrorReporter(Cycle->GetMix())->ReportWarning(ErrorType, FString::Printf(TEXT("Input Path <%s> is not a valid local path"), *Path), GetParentNode());
+		return true;
+	}
+
 	return true;
 }
+
 void UTG_Expression_TexturePath::SetTitleName(FName NewName)
 {
 	GetParentNode()->GetPin(GET_MEMBER_NAME_CHECKED(UTG_Expression_TexturePath, Path))->SetAliasName(NewName);
diff --git a/Engine/Plugins/Experimental/TextureGraph/Source/TextureGraph/Private/TG_Var.cpp b/Engine/Plugins/Experimental/TextureGraph/Source/TextureGraph/Private/TG_Var.cpp
--- a/Engine/Plugins/Experimental/TextureGraph/Source/TextureGraph/Private/TG_Var.cpp
+++ b/Engine/Plugins/Experimental/TextureGraph/Source/TextureGraph/Private/TG_Var.cpp
@@ -41,7 +41,10 @@ template <> FString TG_Var_LogValue(FName& Value)
 	FString LogMessage = Value.ToString();
 	return LogMessage;
 }
-
+template <> FString TG_Var_LogValue(FString& Value)
+{
+	return Value;
+}
 template <> FString TG_Var_LogValue(FLinearColor& Value)
 {
 	FString LogMessage = Value.ToString();
@@ -96,7 +99,10 @@ template <> void TG_Var_SetValueFromString(FName& Value, const FString& StrVal)
 {
 	Value = FName(StrVal);
 }
-
+template <> void TG_Var_SetValueFromString(FString& Value, const FString& StrVal)
+{
+	Value = StrVal;
+}
 template <> void TG_Var_SetValueFromString(FLinearColor& Value, const FString& StrVal)
 {
 	Value.InitFromString(StrVal);
@@ -239,7 +245,10 @@ void VarPropertySerializer_FName(FTG_Var::VarPropertySerialInfo& Info)
 {
 	Generic_Simple_Serializer<FNameProperty, FName>(Info);
 }
-
+void VarPropertySerializer_FString(FTG_Var::VarPropertySerialInfo& Info)
+{
+	Generic_Simple_Serializer<FStrProperty, FString>(Info);
+}
 void VarPropertySerializer_UObjectPtr(FTG_Var::VarPropertySerialInfo& Info)
 {
 	Generic_Struct_Serializer<FTG_Texture>(Info);
@@ -376,6 +385,7 @@ FTG_Var::VarPropertySerializerMap FTG_Var::DefaultPropertySerializers
 		VAR_PROPERTY_SERIALIZER_DEF(uint32),
 		VAR_PROPERTY_SERIALIZER_DEF(float),
 		VAR_PROPERTY_SERIALIZER_DEF(FName),
+		VAR_PROPERTY_SERIALIZER_DEF(FString),
 		VAR_PROPERTY_SERIALIZER_DEF(FTG_Texture),
 		VAR_PROPERTY_SERIALIZER_DEF(FVector4f),
 		VAR_PROPERTY_SERIALIZER_DEF(FVector2f),
diff --git a/Engine/Plugins/Experimental/TextureGraph/Source/TextureGraph/Public/Expressions/Input/TG_Expression_String.h b/Engine/Plugins/Experimental/TextureGraph/Source/TextureGraph/Public/Expressions/Input/TG_Expression_String.h
--- a/Engine/Plugins/Experimental/TextureGraph/Source/TextureGraph/Public/Expressions/Input/TG_Expression_String.h
+++ b/Engine/Plugins/Experimental/TextureGraph/Source/TextureGraph/Public/Expressions/Input/TG_Expression_String.h
@@ -10,6 +10,7 @@ UCLASS()
 class TEXTUREGRAPH_API UTG_Expression_String : public UTG_Expression_InputParam
 {
 	GENERATED_BODY()
+	TG_DECLARE_INPUT_PARAM_EXPRESSION(TG_Category::Input);
 
 public:
 
diff --git a/Engine/Plugins/Experimental/TextureGraph/Source/TextureGraph/Public/Expressions/Input/TG_Expression_TexturePath.h b/Engine/Plugins/Experimental/TextureGraph/Source/TextureGraph/Public/Expressions/Input/TG_Expression_TexturePath.h
--- a/Engine/Plugins/Experimental/TextureGraph/Source/TextureGraph/Public/Expressions/Input/TG_Expression_TexturePath.h
+++ b/Engine/Plugins/Experimental/TextureGraph/Source/TextureGraph/Public/Expressions/Input/TG_Expression_TexturePath.h
@@ -14,7 +14,16 @@ UCLASS()
 class TEXTUREGRAPH_API UTG_Expression_TexturePath : public UTG_Expression_InputParam
 {
 	GENERATED_BODY()
+	TG_DECLARE_INPUT_PARAM_EXPRESSION(TG_Category::Input);
 
+protected:
+	// Special case for TexturePath Constant signature, we want to keep the Path Input connectable in that case
+	// so do this in the override version of BuildInputConstantSignature()
+	virtual FTG_SignaturePtr BuildInputConstantSignature() const override;
+
+	// Validate the input path, returning the actual path to use
+	// empty if the input path is NOT valid
+	bool ValidateInputPath(FString& ValidatedPath) const;
 public:
 
 	virtual void Evaluate(FTG_EvaluationContext* InContext) override;
@@ -25,13 +34,9 @@ public:
 	FTG_Texture Output;
 
 	// Input file path of the texture
-	UPROPERTY(EditAnywhere, BlueprintReadWrite, Category = NoCategory, meta = (TGType = "TG_Input", TGPinNotConnectable, NoResetToDefault) )
+	UPROPERTY(EditAnywhere, BlueprintReadWrite, Category = NoCategory, meta = (TGType = "TG_InputParam", NoResetToDefault) )
 	FString Path;
 
-	// The input texture that was loaded from the path
-	UPROPERTY(meta = (TGType = "TG_InputParam"))
-	FTG_Texture Texture;
-
 	class ULayerChannel* Channel;
 	virtual FTG_Name GetDefaultName() const override { return TEXT("TexturePath"); }
 	virtual void SetTitleName(FName NewName) override;
diff --git a/Engine/Plugins/Experimental/TextureGraph/Source/TextureGraph/Public/TG_Var.h b/Engine/Plugins/Experimental/TextureGraph/Source/TextureGraph/Public/TG_Var.h
--- a/Engine/Plugins/Experimental/TextureGraph/Source/TextureGraph/Public/TG_Var.h
+++ b/Engine/Plugins/Experimental/TextureGraph/Source/TextureGraph/Public/TG_Var.h
@@ -260,6 +260,7 @@ template <> FString TG_Var_LogValue(FLinearColor& Value);
 template <> FString TG_Var_LogValue(FVector4f& Value);
 template <> FString TG_Var_LogValue(FVector2f& Value);
 template <> FString TG_Var_LogValue(FName& Value);
+template <> FString TG_Var_LogValue(FString& Value);
 template <> FString TG_Var_LogValue(TObjectPtr<UObject>& Value);
 template <> FString TG_Var_LogValue(struct FTG_Texture& Value);
 template <> FString TG_Var_LogValue(struct FTG_Scalar& Value);
@@ -277,6 +278,7 @@ template <> void TG_Var_SetValueFromString(FLinearColor& Value, const FString& S
 template <> void TG_Var_SetValueFromString(FVector4f& Value, const FString& StrVal);
 template <> void TG_Var_SetValueFromString(FVector2f& Value, const FString& StrVal);
 template <> void TG_Var_SetValueFromString(FName& Value, const FString& StrVal);
+template <> void TG_Var_SetValueFromString(FString& Value, const FString& StrVal);
 
 //template <> void TG_Var_SetValueFromString(struct FTG_Texture& Value, const FString& StrVal);
 template <> void TG_Var_SetValueFromString(struct FTG_Scalar& Value, const FString& StrVal);
