Delay initialization of TextureGraphEditor module until editor engine is initialized so that InitEngineInternal does not early out during Create, leaving an unusable instance and causing crashes when TextureGraph assets are loaded before engine initializaiton.

diff --git a/Engine/Plugins/Experimental/TextureGraph/TextureGraph.uplugin b/Engine/Plugins/Experimental/TextureGraph/TextureGraph.uplugin
--- a/Engine/Plugins/Experimental/TextureGraph/TextureGraph.uplugin
+++ b/Engine/Plugins/Experimental/TextureGraph/TextureGraph.uplugin
@@ -49,7 +49,7 @@
 		{
 			"Name": "TextureGraphEditor",
 			"Type": "Editor",
-			"LoadingPhase": "Default",
+			"LoadingPhase": "PostEngineInit",
 			"PlatformAllowList": [
 				"Win64"
 			]
