Fix for trying to unbind a tiled buffer that isn't ready yet.

diff --git a/Engine/Plugins/Experimental/TextureGraph/Source/TextureGraphEngine/Data/Blob.cpp b/Engine/Plugins/Experimental/TextureGraph/Source/TextureGraphEngine/Data/Blob.cpp
--- a/Engine/Plugins/Experimental/TextureGraph/Source/TextureGraphEngine/Data/Blob.cpp
+++ b/Engine/Plugins/Experimental/TextureGraph/Source/TextureGraphEngine/Data/Blob.cpp
@@ -231,7 +231,6 @@ AsyncBufferResultPtr Blob::Bind(const BlobTransform* Transform, const ResourceBi
 
 AsyncBufferResultPtr Blob::Unbind(const BlobTransform* Transform, const ResourceBindInfo& BindInfo)
 {
-	//check(_buffer);
 	return Buffer->Unbind(Transform, BindInfo);
 }
 
diff --git a/Engine/Plugins/Experimental/TextureGraph/Source/TextureGraphEngine/Data/TiledBlob.cpp b/Engine/Plugins/Experimental/TextureGraph/Source/TextureGraphEngine/Data/TiledBlob.cpp
--- a/Engine/Plugins/Experimental/TextureGraph/Source/TextureGraphEngine/Data/TiledBlob.cpp
+++ b/Engine/Plugins/Experimental/TextureGraph/Source/TextureGraphEngine/Data/TiledBlob.cpp
@@ -764,8 +764,14 @@ AsyncBufferResultPtr TiledBlob::Unbind(const BlobTransform* Transform, const Res
 	if (Tiles.Rows() > 1 || Tiles.Cols() > 1)
 	{
 		/// If this is not a tiled Target and it was a write Target, then the combined Buffer is already ready
-		if (BindInfo.bWriteTarget)
+		if (BindInfo.bWriteTarget) {
 			bReady = true;
+			check(Buffer); /// Buffer must be valid for write targets
+		}
+
+		/// Don't try to Unbind a buffer that isn't ready yet. The buffer might not even be created yet.
+		if (!bReady || !Buffer)
+			return cti::make_ready_continuable<BufferResultPtr>(std::make_shared<BufferResult>());
 
 		//_ready = false;
 		return Blob::Unbind(Transform, BindInfo).then([this](BufferResultPtr Result)
