Fixed tests failing
- enabled TextureGraph engine test mode for maths tests
- For now, rerouted specialized struct serialization in TG_Var for structs defined in TG_Expression_Pattern. Creating a JIRA for later to resolve in a better way.

diff --git a/Engine/Plugins/Experimental/TextureGraph/Source/TextureGraph/Private/TG_Var.cpp b/Engine/Plugins/Experimental/TextureGraph/Source/TextureGraph/Private/TG_Var.cpp
--- a/Engine/Plugins/Experimental/TextureGraph/Source/TextureGraph/Private/TG_Var.cpp
+++ b/Engine/Plugins/Experimental/TextureGraph/Source/TextureGraph/Private/TG_Var.cpp
@@ -10,6 +10,7 @@
 #include "Model/ModelObject.h"
 
 #include "Expressions/TG_Expression.h"
+#include "Expressions/Procedural/TG_Expression_Pattern.h"
 
 template <> FString TG_Var_LogValue(uint8& Value)
 {
@@ -78,6 +79,36 @@ template <> FString TG_Var_LogValue(FTG_OutputSettings& Value)
 	return LogMessage;
 }
 
+template <> FString TG_Var_LogValue(FPatternMaskPlacement_TS& Value)
+{
+	FString LogMessage = Value.ToString();
+	return LogMessage;
+}
+
+template <> FString TG_Var_LogValue(FPatternMaskJitter_TS& Value)
+{
+	FString LogMessage = Value.ToString();
+	return LogMessage;
+}
+
+template <> FString TG_Var_LogValue(FPatternMaskBevel_TS& Value)
+{
+	FString LogMessage = Value.ToString();
+	return LogMessage;
+}
+
+template <> FString TG_Var_LogValue(FPatternMaskCutout_TS& Value)
+{
+	FString LogMessage = Value.ToString();
+	return LogMessage;
+}
+
+template <> FString TG_Var_LogValue(FGradientDir_TS& Value)
+{
+	FString LogMessage = Value.ToString();
+	return LogMessage;
+}
+
 template <> void TG_Var_SetValueFromString(int& Value, const FString& StrVal)
 {
 	Value = FCString::Atoi(*StrVal);
@@ -135,6 +166,26 @@ template <> void TG_Var_SetValueFromString(FTG_OutputSettings& Value, const FStr
 {
 	Value.InitFromString(StrVal);
 }
+template <> void TG_Var_SetValueFromString(FPatternMaskPlacement_TS& Value, const FString& StrVal)
+{
+	Value.InitFromString(StrVal);
+}
+template <> void TG_Var_SetValueFromString(FPatternMaskJitter_TS& Value, const FString& StrVal)
+{
+	Value.InitFromString(StrVal);
+}
+template <> void TG_Var_SetValueFromString(FPatternMaskBevel_TS& Value, const FString& StrVal)
+{
+	Value.InitFromString(StrVal);
+}
+template <> void TG_Var_SetValueFromString(FPatternMaskCutout_TS& Value, const FString& StrVal)
+{
+	Value.InitFromString(StrVal);
+}
+template <> void TG_Var_SetValueFromString(FGradientDir_TS& Value, const FString& StrVal)
+{
+	Value.InitFromString(StrVal);
+}
 
 
 FString FTG_Var::LogHead() const
@@ -372,6 +423,26 @@ void VarPropertySerializer_FTG_OutputSettings(FTG_Var::VarPropertySerialInfo& In
 	Generic_Struct_Serializer<FTG_OutputSettings>(Info);
 }
 
+void VarPropertySerializer_FGradientDir_TS(FTG_Var::VarPropertySerialInfo& Info)
+{
+	Generic_Struct_Serializer<FGradientDir_TS>(Info);
+}
+void VarPropertySerializer_FPatternMaskCutout_TS(FTG_Var::VarPropertySerialInfo& Info)
+{
+	Generic_Struct_Serializer<FPatternMaskCutout_TS>(Info);
+}
+void VarPropertySerializer_FPatternMaskBevel_TS(FTG_Var::VarPropertySerialInfo& Info)
+{
+	Generic_Struct_Serializer<FPatternMaskBevel_TS>(Info);
+}
+void VarPropertySerializer_FPatternMaskJitter_TS(FTG_Var::VarPropertySerialInfo& Info)
+{
+	Generic_Struct_Serializer<FPatternMaskJitter_TS>(Info);
+}
+void VarPropertySerializer_FPatternMaskPlacement_TS(FTG_Var::VarPropertySerialInfo& Info)
+{
+	Generic_Struct_Serializer<FPatternMaskPlacement_TS>(Info);
+}
 //////////////////////////////////////////////////////////////////////////////////////////////////////////////////////////////////////
 FTG_Var::VarPropertySerializerMap FTG_Var::DefaultPropertySerializers
 (
@@ -392,6 +463,11 @@ FTG_Var::VarPropertySerializerMap FTG_Var::DefaultPropertySerializers
 		VAR_PROPERTY_SERIALIZER_DEF(FLinearColor),
 		VAR_PROPERTY_SERIALIZER_DEF(FTG_OutputSettings),
 		VAR_PROPERTY_SERIALIZER_DEF(FTG_Variant),
+		VAR_PROPERTY_SERIALIZER_DEF(FGradientDir_TS),
+		VAR_PROPERTY_SERIALIZER_DEF(FPatternMaskCutout_TS),
+		VAR_PROPERTY_SERIALIZER_DEF(FPatternMaskBevel_TS),
+		VAR_PROPERTY_SERIALIZER_DEF(FPatternMaskJitter_TS),
+		VAR_PROPERTY_SERIALIZER_DEF(FPatternMaskPlacement_TS),
 	}
 );
 
diff --git a/Engine/Plugins/Experimental/TextureGraph/Source/TextureGraph/Public/Expressions/Procedural/TG_Expression_Pattern.h b/Engine/Plugins/Experimental/TextureGraph/Source/TextureGraph/Public/Expressions/Procedural/TG_Expression_Pattern.h
--- a/Engine/Plugins/Experimental/TextureGraph/Source/TextureGraph/Public/Expressions/Procedural/TG_Expression_Pattern.h
+++ b/Engine/Plugins/Experimental/TextureGraph/Source/TextureGraph/Public/Expressions/Procedural/TG_Expression_Pattern.h
@@ -7,6 +7,7 @@
 
 #include "Expressions/TG_Expression.h"
 #include "TG_Texture.h"
+#include "Misc/OutputDeviceNull.h"
 
 #include "TG_Expression_Pattern.generated.h"
 
@@ -51,6 +52,19 @@ struct TEXTUREGRAPH_API FPatternMaskPlacement_TS
 	Vector2								Repeat() const { return Vector2(RepeatX, RepeatY); }
 	Vector2								Spacing() const { return Vector2(SpacingX, SpacingY); }
 	Vector2								OffsetValue() const { return (OffsetHorizontal ? Vector2( Offset, 0.0f ) : Vector2( 0.0f, Offset )); }
+
+	void InitFromString(const FString& StrVal)
+	{
+		FOutputDeviceNull NullOut;
+		FPatternMaskPlacement_TS::StaticStruct()->ImportText(*StrVal, this, /*OwnerObject*/nullptr, 0, &NullOut, FPatternMaskPlacement_TS::StaticStruct()->GetName(), /*bAllowNativeOverride*/true);
+	}
+
+	FString ToString() const
+	{
+		FString ExportString;
+		FPatternMaskPlacement_TS::StaticStruct()->ExportText(ExportString, this, this, /*OwnerObject*/nullptr, /*PortFlags*/0, /*ExportRootScope*/nullptr);
+		return ExportString;
+	}
 };
 
 USTRUCT(BlueprintType)
@@ -112,6 +126,19 @@ struct TEXTUREGRAPH_API FPatternMaskJitter_TS
 
 	Vector2								Angle() const { return Vector2(AngleXAmount, AngleYAmount); }
 	Vector2								Tilt() const { return Vector2(TiltXAmount, TiltYAmount); }
+
+	void InitFromString(const FString& StrVal)
+	{
+		FOutputDeviceNull NullOut;
+		FPatternMaskJitter_TS::StaticStruct()->ImportText(*StrVal, this, /*OwnerObject*/nullptr, 0, &NullOut, FPatternMaskJitter_TS::StaticStruct()->GetName(), /*bAllowNativeOverride*/true);
+	}
+
+	FString ToString() const
+	{
+		FString ExportString;
+		FPatternMaskJitter_TS::StaticStruct()->ExportText(ExportString, this, this, /*OwnerObject*/nullptr, /*PortFlags*/0, /*ExportRootScope*/nullptr);
+		return ExportString;
+	}
 };
 
 USTRUCT(BlueprintType)
@@ -126,6 +153,19 @@ struct TEXTUREGRAPH_API FPatternMaskBevel_TS
 	// The bevel curve
 	UPROPERTY(EditAnywhere, BlueprintReadWrite, Category = NoCategory, meta = (TGType = "TG_Setting", UIMin = "-1", ClampMin = "-1", UIMax = "1.0", ClampMax = "1.0" ))
 	float								BevelCurve = 0;
+	
+	void InitFromString(const FString& StrVal)
+	{
+		FOutputDeviceNull NullOut;
+		FPatternMaskBevel_TS::StaticStruct()->ImportText(*StrVal, this, /*OwnerObject*/nullptr, 0, &NullOut, FPatternMaskBevel_TS::StaticStruct()->GetName(), /*bAllowNativeOverride*/true);
+	}
+
+	FString ToString() const
+	{
+		FString ExportString;
+		FPatternMaskBevel_TS::StaticStruct()->ExportText(ExportString, this, this, /*OwnerObject*/nullptr, /*PortFlags*/0, /*ExportRootScope*/nullptr);
+		return ExportString;
+	}
 };
 
 USTRUCT(BlueprintType)
@@ -140,6 +180,19 @@ struct TEXTUREGRAPH_API FPatternMaskCutout_TS
 	// The seed to control the cutoff randomness. 0 means no randomness
 	UPROPERTY(EditAnywhere, BlueprintReadWrite, Category = NoCategory, meta = (TGType = "TG_Setting", UIMin = "0", ClampMin = "0", UIMax = "100", ClampMax = "100" ))
 	float								CutoffSeed = 0;
+	
+	void InitFromString(const FString& StrVal)
+	{
+		FOutputDeviceNull NullOut;
+		FPatternMaskCutout_TS::StaticStruct()->ImportText(*StrVal, this, /*OwnerObject*/nullptr, 0, &NullOut, FPatternMaskCutout_TS::StaticStruct()->GetName(), /*bAllowNativeOverride*/true);
+	}
+
+	FString ToString() const
+	{
+		FString ExportString;
+		FPatternMaskCutout_TS::StaticStruct()->ExportText(ExportString, this, this, /*OwnerObject*/nullptr, /*PortFlags*/0, /*ExportRootScope*/nullptr);
+		return ExportString;
+	}
 };
 
 USTRUCT(BlueprintType)
@@ -156,6 +209,19 @@ struct TEXTUREGRAPH_API FGradientDir_TS
 	float								Y = 0;
 
 	Vector2								Value() const { return Vector2(X, Y); }
+
+	void InitFromString(const FString& StrVal)
+	{
+		FOutputDeviceNull NullOut;
+		FGradientDir_TS::StaticStruct()->ImportText(*StrVal, this, /*OwnerObject*/nullptr, 0, &NullOut, FGradientDir_TS::StaticStruct()->GetName(), /*bAllowNativeOverride*/true);
+	}
+
+	FString ToString() const
+	{
+		FString ExportString;
+		FGradientDir_TS::StaticStruct()->ExportText(ExportString, this, this, /*OwnerObject*/nullptr, /*PortFlags*/0, /*ExportRootScope*/nullptr);
+		return ExportString;
+	}
 };
 
 UCLASS()
