Fixed Issue: Blended output is rendered as transparent in the Node Preview and Thumbnail

Fix/How: This is happening because when we use blends like subtract then alpha channel also get zero. This creates confusion because preview and thumbnail become empty. Introducing IgnoreAlpha flag will fix this issue.
IgnoreAlpha flag is on by default.

diff --git a/Engine/Plugins/Experimental/TextureGraph/Shaders/Expressions/Expression_Blend.usf b/Engine/Plugins/Experimental/TextureGraph/Shaders/Expressions/Expression_Blend.usf
--- a/Engine/Plugins/Experimental/TextureGraph/Shaders/Expressions/Expression_Blend.usf
+++ b/Engine/Plugins/Experimental/TextureGraph/Shaders/Expressions/Expression_Blend.usf
@@ -24,9 +24,13 @@ float4 FSH_BlendNormal(in float2 uv : TEXCOORD0) : SV_Target0
 	float4 ForegroundTex = ForegroundTexture.Sample(SamplerStates_Linear_Clamp, uv);
 
 	float MaskingValue = GetMask(MaskTexture, uv, Opacity);
-	
+
+#if IGNORE_ALPHA
+	float3 FinalRGB = lerp(BackgroundTex.rgb, ForegroundTex.rgb, MaskingValue);
+	float4 FinalColor = float4(FinalRGB, BackgroundTex.a);
+#else
 	float4 FinalColor = lerp(BackgroundTex, ForegroundTex, MaskingValue);
-	
+#endif
     return FinalColor;
 }
 
@@ -36,8 +40,12 @@ float4 FSH_BlendAdd(in float2 uv : TEXCOORD0) : SV_Target0
 	float4 ForegroundTex = ForegroundTexture.Sample(SamplerStates_Linear_Clamp, uv);
 	
 	float MaskingValue = GetMask(MaskTexture, uv, Opacity);
-	
+#if IGNORE_ALPHA
+	float3 FinalRGB = lerp(BackgroundTex.rgb, BackgroundTex.rgb + ForegroundTex.rgb, MaskingValue);
+	float4 FinalColor = float4(FinalRGB, BackgroundTex.a);
+#else
 	float4 FinalColor = lerp(BackgroundTex, BackgroundTex + ForegroundTex, MaskingValue);
+#endif
 	
 	return FinalColor;
 }
@@ -48,8 +56,13 @@ float4 FSH_BlendSubtract(in float2 uv : TEXCOORD0) : SV_Target0
 	float4 ForegroundTex = ForegroundTexture.Sample(SamplerStates_Linear_Clamp, uv);
 
 	float MaskingValue = GetMask(MaskTexture, uv, Opacity);
-	
+
+#if IGNORE_ALPHA
+	float3 FinalRGB = lerp(BackgroundTex, saturate(BackgroundTex - ForegroundTex), MaskingValue);
+	float4 FinalColor = float4(FinalRGB, BackgroundTex.a);
+#else
 	float4 FinalColor = lerp(BackgroundTex, saturate(BackgroundTex - ForegroundTex), MaskingValue);
+#endif
 	
 	return FinalColor;
 }
@@ -60,8 +73,13 @@ float4 FSH_BlendMultiply(in float2 uv : TEXCOORD0) : SV_Target0
 	float4 ForegroundTex = ForegroundTexture.Sample(SamplerStates_Linear_Clamp, uv);
 
 	float MaskingValue = GetMask(MaskTexture, uv, Opacity);
-	
+
+#if IGNORE_ALPHA
+	float3 FinalRGB = lerp(BackgroundTex, saturate(ForegroundTex * BackgroundTex), MaskingValue);
+	float4 FinalColor = float4(FinalRGB, BackgroundTex.a);
+#else
 	float4 FinalColor = lerp(BackgroundTex, saturate(ForegroundTex * BackgroundTex), MaskingValue);
+#endif
 	
 	return FinalColor;
 }
@@ -74,8 +92,13 @@ float4 FSH_BlendDivide(in float2 uv : TEXCOORD0) : SV_Target0
 	float MaskingValue = GetMask(MaskTexture, uv, Opacity);
 	
 	float4 DivisionResult = saturate(BackgroundTex/ForegroundTex);
-	
+
+#if IGNORE_ALPHA
+	float3 FinalRGB = lerp(BackgroundTex.rgb, DivisionResult.rgb, MaskingValue);
+	float4 FinalColor = float4(FinalRGB, BackgroundTex.a);
+#else
 	float4 FinalColor = lerp(BackgroundTex, DivisionResult, MaskingValue);
+#endif
 	
 	return FinalColor;
 }
@@ -86,8 +109,13 @@ float4 FSH_BlendDifference(in float2 uv : TEXCOORD0) : SV_Target0
 	float4 ForegroundTex = ForegroundTexture.Sample(SamplerStates_Linear_Clamp, uv);
 
 	float MaskingValue = GetMask(MaskTexture, uv, Opacity);
-	
+
+#if IGNORE_ALPHA
+	float3 FinalRGB = lerp(BackgroundTex.rgb, abs(BackgroundTex.rgb - ForegroundTex.rgb), MaskingValue);
+	float4 FinalColor = float4(FinalRGB, BackgroundTex.a);
+#else
 	float4 FinalColor = lerp(BackgroundTex, abs(BackgroundTex - ForegroundTex), MaskingValue);
+#endif
 	
 	return FinalColor;
 }
@@ -98,9 +126,13 @@ float4 FSH_BlendMax(in float2 uv : TEXCOORD0) : SV_Target0
 	float4 ForegroundTex = ForegroundTexture.Sample(SamplerStates_Linear_Clamp, uv);
 
 	float MaskingValue = GetMask(MaskTexture, uv, Opacity);
-	
+
+#if IGNORE_ALPHA
+	float3 FinalRGB = lerp(BackgroundTex.rgb, max(BackgroundTex.rgb, ForegroundTex.rgb), MaskingValue);
+	float4 FinalColor = float4(FinalRGB, BackgroundTex.a);
+#else
 	float4 FinalColor = lerp(BackgroundTex, max(BackgroundTex, ForegroundTex), MaskingValue);
-	
+#endif
 	return FinalColor;
 }
 
@@ -110,8 +142,13 @@ float4 FSH_BlendMin(in float2 uv : TEXCOORD0) : SV_Target0
 	float4 ForegroundTex = ForegroundTexture.Sample(SamplerStates_Linear_Clamp, uv);
 
 	float MaskingValue = GetMask(MaskTexture, uv, Opacity);
-	
+
+#if IGNORE_ALPHA
+	float3 FinalRGB = lerp(BackgroundTex.rgb, min(BackgroundTex.rgb, ForegroundTex.rgb), MaskingValue);
+	float4 FinalColor = float4(FinalRGB, BackgroundTex.a);
+#else
 	float4 FinalColor = lerp(BackgroundTex, min(BackgroundTex, ForegroundTex), MaskingValue);
+#endif
 	
 	return FinalColor;
 }
@@ -122,10 +159,14 @@ float4 FSH_BlendStep(in float2 uv : TEXCOORD0) : SV_Target0
 	float4 ForegroundTex = ForegroundTexture.Sample(SamplerStates_Linear_Clamp, uv);
 
 	float MaskingValue = GetMask(MaskTexture, uv, Opacity);
-
 	float4 StepValue = 1 - step(ForegroundTex, BackgroundTex);
-	
+
+#if IGNORE_ALPHA
+	float3 FinalRGB = lerp(BackgroundTex.rgb, StepValue.rgb, MaskingValue);
+	float4 FinalColor = float4(FinalRGB, BackgroundTex.a);
+#else
 	float4 FinalColor = lerp(BackgroundTex, StepValue, MaskingValue);
+#endif
 	
 	return FinalColor;
 }
@@ -135,15 +176,25 @@ float4 OverlayBlend( float4 base , float4 overlay )
 	return lerp(2 * base * overlay, 1 - 2 * (1 - base) * (1 - overlay), 1 - step(base, 0.5));
 }
 
+float3 OverlayBlend( float3 base , float3 overlay )
+{
+	return lerp(2 * base * overlay, 1 - 2 * (1 - base) * (1 - overlay), 1 - step(base, 0.5));
+}
+
 float4 FSH_BlendOverlay(in float2 uv : TEXCOORD0) : SV_Target0
 {
 	float4 BackgroundTex = BackgroundTexture.Sample(SamplerStates_Linear_Clamp, uv);
 	float4 ForegroundTex = ForegroundTexture.Sample(SamplerStates_Linear_Clamp, uv);
 
 	float MaskingValue = GetMask(MaskTexture, uv, Opacity);
-	
+
+#if IGNORE_ALPHA
+	float3 FinalRGB = lerp(BackgroundTex.rgb, OverlayBlend(BackgroundTex.rgb, ForegroundTex.rgb), MaskingValue);
+	float4 FinalColor = float4(FinalRGB, BackgroundTex.a);
+#else
 	float4 FinalColor = lerp(BackgroundTex, OverlayBlend(BackgroundTex, ForegroundTex), MaskingValue);
-	
+#endif
+
 	return FinalColor;
 }
 
diff --git a/Engine/Plugins/Experimental/TextureGraph/Source/TextureGraph/Private/Expressions/Maths/TG_Expression_Blend.cpp b/Engine/Plugins/Experimental/TextureGraph/Source/TextureGraph/Private/Expressions/Maths/TG_Expression_Blend.cpp
--- a/Engine/Plugins/Experimental/TextureGraph/Source/TextureGraph/Private/Expressions/Maths/TG_Expression_Blend.cpp
+++ b/Engine/Plugins/Experimental/TextureGraph/Source/TextureGraph/Private/Expressions/Maths/TG_Expression_Blend.cpp
@@ -26,5 +26,12 @@ void UTG_Expression_Blend::Evaluate(FTG_EvaluationContext* InContext)
 		return;
 	}
 
-	Output = T_Blend::Create(InContext->Cycle, Output.GetBufferDescriptor(), Background, Foreground, Mask, Opacity, InContext->TargetId, BlendMode);
+	T_Blend::FBlendSettings BlendSettings;
+	BlendSettings.ForegroundTexture = Foreground;
+	BlendSettings.BackgroundTexture = Background;
+	BlendSettings.Mask = Mask;
+	BlendSettings.Opacity = Opacity;
+	BlendSettings.bIgnoreAlpha = bIgnoreAlpha;
+	
+	Output = T_Blend::Create(InContext->Cycle, Output.GetBufferDescriptor(), InContext->TargetId, BlendMode, &BlendSettings);
 }
diff --git a/Engine/Plugins/Experimental/TextureGraph/Source/TextureGraph/Public/Expressions/Maths/TG_Expression_Blend.h b/Engine/Plugins/Experimental/TextureGraph/Source/TextureGraph/Public/Expressions/Maths/TG_Expression_Blend.h
--- a/Engine/Plugins/Experimental/TextureGraph/Source/TextureGraph/Public/Expressions/Maths/TG_Expression_Blend.h
+++ b/Engine/Plugins/Experimental/TextureGraph/Source/TextureGraph/Public/Expressions/Maths/TG_Expression_Blend.h
@@ -25,6 +25,10 @@ public:
 	// The opacity to be used in conjunction with the mask. The final blend value = 'Mask[RedChannel] * Opacity'
 	UPROPERTY(EditAnywhere, Category = NoCategory, meta = (TGType = "TG_Setting", UIMin = "0", ClampMin = "0", UIMax = "1", ClampMax = "1"))
 	float Opacity = 1.0;
+
+	// If true, it will ignore the alpha channel blending and keep the background alpha.  
+	UPROPERTY(EditAnywhere, Category = NoCategory, meta = (TGType = "TG_Setting", DisplayName = "Ignore Alpha"))
+	bool bIgnoreAlpha = true;
 	
 	UPROPERTY(EditAnywhere, Category = NoCategory, meta = (TGType = "TG_Output", PinDisplayName = ""))
 	FTG_Texture Output;
diff --git a/Engine/Plugins/Experimental/TextureGraph/Source/TextureGraphEngine/Transform/Expressions/T_Blend.cpp b/Engine/Plugins/Experimental/TextureGraph/Source/TextureGraphEngine/Transform/Expressions/T_Blend.cpp
--- a/Engine/Plugins/Experimental/TextureGraph/Source/TextureGraphEngine/Transform/Expressions/T_Blend.cpp
+++ b/Engine/Plugins/Experimental/TextureGraph/Source/TextureGraphEngine/Transform/Expressions/T_Blend.cpp
@@ -5,7 +5,6 @@
 #include "TextureGraphEngine.h"
 #include "FxMat/MaterialManager.h"
 #include "Job/JobBatch.h"
-#include "2D/TargetTextureSet.h"
 #include "Model/Mix/MixInterface.h"
 #include "Model/Mix/MixSettings.h"
 
@@ -29,31 +28,30 @@ T_Blend::~T_Blend()
 {
 }
 
-TiledBlobPtr T_Blend::Create(MixUpdateCyclePtr InCycle, BufferDescriptor DesiredDesc, TiledBlobPtr InBackgroundTexture,
-	TiledBlobPtr InForeGroundTexture, TiledBlobPtr InMask, float InOpacity, int InTargetId, EBlendModes InBlendMode)
+TiledBlobPtr T_Blend::Create(MixUpdateCyclePtr InCycle, BufferDescriptor DesiredDesc, int InTargetId, EBlendModes InBlendMode, const FBlendSettings* InBlendSettings)
 {
 	switch(InBlendMode)
 	{
 		case EBlendModes::Normal:
-			return CreateNormal(InCycle, DesiredDesc, InBackgroundTexture, InForeGroundTexture, InMask, InOpacity, InTargetId);
+			return CreateNormal(InCycle, DesiredDesc, InTargetId, InBlendSettings);
 		case EBlendModes::Add:
-			return CreateAdd(InCycle, DesiredDesc, InBackgroundTexture, InForeGroundTexture, InMask, InOpacity, InTargetId);
+			return CreateAdd(InCycle, DesiredDesc, InTargetId, InBlendSettings);
 		case EBlendModes::Subtract:
-			return CreateSubtract(InCycle, DesiredDesc, InBackgroundTexture, InForeGroundTexture, InMask, InOpacity, InTargetId);
+			return CreateSubtract(InCycle, DesiredDesc, InTargetId, InBlendSettings);
 		case EBlendModes::Multiply:
-			return CreateMultiply(InCycle, DesiredDesc, InBackgroundTexture, InForeGroundTexture, InMask, InOpacity, InTargetId);
+			return CreateMultiply(InCycle, DesiredDesc, InTargetId, InBlendSettings);
 		case EBlendModes::Divide:
-			return CreateDivide(InCycle, DesiredDesc, InBackgroundTexture, InForeGroundTexture, InMask, InOpacity, InTargetId);
+			return CreateDivide(InCycle, DesiredDesc, InTargetId, InBlendSettings);
 		case EBlendModes::Difference:
-			return CreateDifference(InCycle, DesiredDesc, InBackgroundTexture, InForeGroundTexture, InMask, InOpacity, InTargetId);
+			return CreateDifference(InCycle, DesiredDesc, InTargetId, InBlendSettings);
 		case EBlendModes::Max:
-			return CreateMax(InCycle, DesiredDesc, InBackgroundTexture, InForeGroundTexture, InMask, InOpacity, InTargetId);
+			return CreateMax(InCycle, DesiredDesc, InTargetId, InBlendSettings);
 		case EBlendModes::Min:
-			return CreateMin(InCycle, DesiredDesc, InBackgroundTexture, InForeGroundTexture, InMask, InOpacity, InTargetId);
+			return CreateMin(InCycle, DesiredDesc, InTargetId, InBlendSettings);
 		case EBlendModes::Step:
-			return CreateStep(InCycle, DesiredDesc, InBackgroundTexture, InForeGroundTexture, InMask, InOpacity, InTargetId);
+			return CreateStep(InCycle, DesiredDesc, InTargetId, InBlendSettings);
 		case EBlendModes::Overlay:
-			return CreateOverlay(InCycle, DesiredDesc, InBackgroundTexture, InForeGroundTexture, InMask, InOpacity, InTargetId);
+			return CreateOverlay(InCycle, DesiredDesc, InTargetId, InBlendSettings);
 		// case EBlendModes::Distort:
 		// 	return CreateDistort(InCycle, DesiredDesc, InBackgroundTexture, InForeGroundTexture, InMask, InOpacity, InTargetId);
 
@@ -65,34 +63,41 @@ TiledBlobPtr T_Blend::Create(MixUpdateCyclePtr InCycle, BufferDescriptor Desired
 }
 
 template <typename FSH_Type>
-TiledBlobPtr CreateGenericBlend(MixUpdateCyclePtr InCycle, int32 InTargetId, BufferDescriptor DesiredDesc, TiledBlobPtr InBackgroundTexture,
-	TiledBlobPtr InForegroundTexture, TiledBlobPtr InMask, float InOpacity, FString InTransformName)
+TiledBlobPtr CreateGenericBlend(MixUpdateCyclePtr InCycle, int32 InTargetId, BufferDescriptor DesiredDesc,FString InTransformName, const T_Blend::FBlendSettings* InBlendSettings)
 {
-	const RenderMaterial_FXPtr RenderMaterial = TextureGraphEngine::GetMaterialManager()->CreateMaterial_FX<VSH_Simple, FSH_Type>(InTransformName);
+	FSH_BlendBase::FPermutationDomain PermutationVector;
+	PermutationVector.Set<FSH_BlendBase::FIgnoreAlpha>(InBlendSettings->bIgnoreAlpha);
+	
+	const RenderMaterial_FXPtr RenderMaterial = TextureGraphEngine::GetMaterialManager()->CreateMaterial_FX<VSH_Simple, FSH_Type>(InTransformName, PermutationVector);
 	check(RenderMaterial);
 
-	if (!InBackgroundTexture)
+	TiledBlobPtr BackgroundTexture = InBlendSettings->BackgroundTexture; 
+	TiledBlobPtr ForegroundTexture = InBlendSettings->ForegroundTexture; 
+	TiledBlobPtr MaskTexture = InBlendSettings->Mask; 
+	
+	if (!BackgroundTexture)
 	{
-		InBackgroundTexture = TextureHelper::GetBlack();
+		BackgroundTexture = TextureHelper::GetBlack();
 	}
 	
-	if (!InForegroundTexture)
+	if (!ForegroundTexture)
 	{
-		InForegroundTexture = TextureHelper::GetBlack();
+		ForegroundTexture = TextureHelper::GetBlack();
 	}
 	
-	if(!InMask)
+	if(!MaskTexture)
 	{
-		InMask = TextureHelper::GetWhite();
-	}  
-
+		MaskTexture = TextureHelper::GetWhite();
+	}
+	
 	JobUPtr JobPtr = std::make_unique<Job>(InCycle->GetMix(), InTargetId, std::static_pointer_cast<BlobTransform>(RenderMaterial));
 	
 	JobPtr
-		->AddArg(ARG_BLOB(InBackgroundTexture, "BackgroundTexture"))
-		->AddArg(ARG_BLOB(InForegroundTexture, "ForegroundTexture"))
-		->AddArg(ARG_BLOB(InMask, "MaskTexture"))
-		->AddArg(ARG_FLOAT(InOpacity, "Opacity"))
+		->AddArg(ARG_BLOB(BackgroundTexture, "BackgroundTexture"))
+		->AddArg(ARG_BLOB(ForegroundTexture, "ForegroundTexture"))
+		->AddArg(ARG_BLOB(MaskTexture, "MaskTexture"))
+		->AddArg(ARG_FLOAT(InBlendSettings->Opacity, "Opacity"))
+		->AddArg(WithUnbounded(ARG_BOOL(InBlendSettings->bIgnoreAlpha, "IgnoreAlpha")))
 		;
 
 	const FString Name = FString::Printf(TEXT("[%llu] - Blend - %s"), InCycle->GetBatch()->GetBatchId(), *InTransformName);
@@ -104,57 +109,57 @@ TiledBlobPtr CreateGenericBlend(MixUpdateCyclePtr InCycle, int32 InTargetId, Buf
 	return Result;
 }
 
-TiledBlobPtr T_Blend::CreateNormal(MixUpdateCyclePtr InCycle, BufferDescriptor DesiredDesc, TiledBlobPtr InBackgroundTexture, TiledBlobPtr InForegroundTexture, TiledBlobPtr InMask, float InOpacity, int InTargetId)
+TiledBlobPtr T_Blend::CreateNormal(MixUpdateCyclePtr InCycle, BufferDescriptor DesiredDesc, int InTargetId, const FBlendSettings* InBlendSettings)
 {
-	return CreateGenericBlend<FSH_BlendNormal>(InCycle, InTargetId, DesiredDesc, InBackgroundTexture, InForegroundTexture, InMask,InOpacity, "T_BlendNormal");
+	return CreateGenericBlend<FSH_BlendNormal>(InCycle, InTargetId, DesiredDesc, "T_BlendNormal", InBlendSettings);
 }
 
-TiledBlobPtr T_Blend::CreateAdd(MixUpdateCyclePtr InCycle, BufferDescriptor DesiredDesc, TiledBlobPtr InBackgroundTexture, TiledBlobPtr InForegroundTexture, TiledBlobPtr InMask, float InOpacity, int InTargetId)
+TiledBlobPtr T_Blend::CreateAdd(MixUpdateCyclePtr InCycle, BufferDescriptor DesiredDesc, int InTargetId, const FBlendSettings* InBlendSettings)
 {
-	return CreateGenericBlend<FSH_BlendAdd>(InCycle, InTargetId, DesiredDesc, InBackgroundTexture, InForegroundTexture, InMask,InOpacity, "T_BlendAdd");
+	return CreateGenericBlend<FSH_BlendAdd>(InCycle, InTargetId, DesiredDesc, "T_BlendAdd", InBlendSettings);
 }
 
-TiledBlobPtr T_Blend::CreateSubtract(MixUpdateCyclePtr InCycle, BufferDescriptor DesiredDesc, TiledBlobPtr InBackgroundTexture, TiledBlobPtr InForegroundTexture, TiledBlobPtr InMask, float InOpacity, int InTargetId)
+TiledBlobPtr T_Blend::CreateSubtract(MixUpdateCyclePtr InCycle, BufferDescriptor DesiredDesc, int InTargetId, const FBlendSettings* InBlendSettings)
 {
-	return CreateGenericBlend<FSH_BlendSubtract>(InCycle, InTargetId, DesiredDesc, InBackgroundTexture, InForegroundTexture, InMask,InOpacity, "T_BlendSubtract");
+	return CreateGenericBlend<FSH_BlendSubtract>(InCycle, InTargetId, DesiredDesc, "T_BlendSubtract", InBlendSettings);
 }
 
-TiledBlobPtr T_Blend::CreateMultiply(MixUpdateCyclePtr InCycle, BufferDescriptor DesiredDesc, TiledBlobPtr InBackgroundTexture, TiledBlobPtr InForeGroundTexture, TiledBlobPtr InMask, float InOpacity, int InTargetId)
+TiledBlobPtr T_Blend::CreateMultiply(MixUpdateCyclePtr InCycle, BufferDescriptor DesiredDesc, int InTargetId, const FBlendSettings* InBlendSettings)
 {
-	return CreateGenericBlend<FSH_BlendMultiply>(InCycle, InTargetId, DesiredDesc, InBackgroundTexture, InForeGroundTexture, InMask,InOpacity, "T_BlendMultiply");
+	return CreateGenericBlend<FSH_BlendMultiply>(InCycle, InTargetId, DesiredDesc, "T_BlendMultiply", InBlendSettings);
 }
 
-TiledBlobPtr T_Blend::CreateDivide(MixUpdateCyclePtr InCycle, BufferDescriptor DesiredDesc, TiledBlobPtr InBackgroundTexture, TiledBlobPtr InForegroundTexture, TiledBlobPtr InMask, float InOpacity, int InTargetId)
+TiledBlobPtr T_Blend::CreateDivide(MixUpdateCyclePtr InCycle, BufferDescriptor DesiredDesc, int InTargetId, const FBlendSettings* InBlendSettings)
 {
-	return CreateGenericBlend<FSH_BlendDivide>(InCycle, InTargetId, DesiredDesc, InBackgroundTexture, InForegroundTexture, InMask,InOpacity, "T_BlendDivide");
+	return CreateGenericBlend<FSH_BlendDivide>(InCycle, InTargetId, DesiredDesc, "T_BlendDivide", InBlendSettings);
 }
 
-TiledBlobPtr T_Blend::CreateDifference(MixUpdateCyclePtr InCycle, BufferDescriptor DesiredDesc, TiledBlobPtr InBackgroundTexture, TiledBlobPtr InForegroundTexture, TiledBlobPtr InMask, float InOpacity, int InTargetId)
+TiledBlobPtr T_Blend::CreateDifference(MixUpdateCyclePtr InCycle, BufferDescriptor DesiredDesc, int InTargetId, const FBlendSettings* InBlendSettings)
 {
-	return CreateGenericBlend<FSH_BlendDifference>(InCycle, InTargetId, DesiredDesc, InBackgroundTexture, InForegroundTexture, InMask,InOpacity, "T_BlendDifference");
+	return CreateGenericBlend<FSH_BlendDifference>(InCycle, InTargetId, DesiredDesc, "T_BlendDifference", InBlendSettings);
 }
 
-TiledBlobPtr T_Blend::CreateMax(MixUpdateCyclePtr InCycle, BufferDescriptor DesiredDesc, TiledBlobPtr InBackgroundTexture, TiledBlobPtr InForegroundTexture, TiledBlobPtr InMask, float InOpacity, int InTargetId)
+TiledBlobPtr T_Blend::CreateMax(MixUpdateCyclePtr InCycle, BufferDescriptor DesiredDesc,int InTargetId, const FBlendSettings* InBlendSettings)
 {
-	return CreateGenericBlend<FSH_BlendMax>(InCycle, InTargetId, DesiredDesc, InBackgroundTexture, InForegroundTexture, InMask,InOpacity, "T_BlendMax");
+	return CreateGenericBlend<FSH_BlendMax>(InCycle, InTargetId, DesiredDesc, "T_BlendMax", InBlendSettings);
 }
 
-TiledBlobPtr T_Blend::CreateMin(MixUpdateCyclePtr InCycle, BufferDescriptor DesiredDesc, TiledBlobPtr InBackgroundTexture, TiledBlobPtr InForegroundTexture, TiledBlobPtr InMask, float InOpacity, int InTargetId)
+TiledBlobPtr T_Blend::CreateMin(MixUpdateCyclePtr InCycle, BufferDescriptor DesiredDesc,int InTargetId, const FBlendSettings* InBlendSettings)
 {
-	return CreateGenericBlend<FSH_BlendMin>(InCycle, InTargetId, DesiredDesc, InBackgroundTexture, InForegroundTexture, InMask,InOpacity, "T_BlendMin");
+	return CreateGenericBlend<FSH_BlendMin>(InCycle, InTargetId, DesiredDesc, "T_BlendMin", InBlendSettings);
 }
 
-TiledBlobPtr T_Blend::CreateStep(MixUpdateCyclePtr InCycle, BufferDescriptor DesiredDesc, TiledBlobPtr InBackgroundTexture, TiledBlobPtr InForegroundTexture, TiledBlobPtr InMask, float InOpacity, int InTargetId)
+TiledBlobPtr T_Blend::CreateStep(MixUpdateCyclePtr InCycle, BufferDescriptor DesiredDesc,int InTargetId, const FBlendSettings* InBlendSettings)
 {
-	return CreateGenericBlend<FSH_BlendStep>(InCycle, InTargetId, DesiredDesc, InBackgroundTexture, InForegroundTexture, InMask,InOpacity, "T_BlendStep");
+	return CreateGenericBlend<FSH_BlendStep>(InCycle, InTargetId, DesiredDesc, "T_BlendStep", InBlendSettings);
 }
 
-TiledBlobPtr T_Blend::CreateOverlay(MixUpdateCyclePtr InCycle, BufferDescriptor DesiredDesc, TiledBlobPtr InBackgroundTexture, TiledBlobPtr InForegroundTexture, TiledBlobPtr InMask, float InOpacity, int InTargetId)
+TiledBlobPtr T_Blend::CreateOverlay(MixUpdateCyclePtr InCycle, BufferDescriptor DesiredDesc,int InTargetId, const FBlendSettings* InBlendSettings)
 {
-	return CreateGenericBlend<FSH_BlendOverlay>(InCycle, InTargetId, DesiredDesc, InBackgroundTexture, InForegroundTexture, InMask,InOpacity, "T_BlendOverlay");
+	return CreateGenericBlend<FSH_BlendOverlay>(InCycle, InTargetId, DesiredDesc, "T_BlendOverlay", InBlendSettings);
 }
 
-TiledBlobPtr T_Blend::CreateDistort(MixUpdateCyclePtr InCycle, BufferDescriptor DesiredDesc, TiledBlobPtr InBackgroundTexture, TiledBlobPtr InForegroundTexture, TiledBlobPtr InMask, float InOpacity, int InTargetId)
+TiledBlobPtr T_Blend::CreateDistort(MixUpdateCyclePtr InCycle, BufferDescriptor DesiredDesc, int InTargetId, const FBlendSettings* InBlendSettings)
 {
-	return CreateGenericBlend<FSH_BlendDistort>(InCycle, InTargetId, DesiredDesc, InBackgroundTexture, InForegroundTexture, InMask,InOpacity, "T_BlendDistort");
+	return CreateGenericBlend<FSH_BlendDistort>(InCycle, InTargetId, DesiredDesc, "T_BlendDistort", InBlendSettings);
 }
\ No newline at end of file
diff --git a/Engine/Plugins/Experimental/TextureGraph/Source/TextureGraphEngine/Transform/Expressions/T_Blend.h b/Engine/Plugins/Experimental/TextureGraph/Source/TextureGraphEngine/Transform/Expressions/T_Blend.h
--- a/Engine/Plugins/Experimental/TextureGraph/Source/TextureGraphEngine/Transform/Expressions/T_Blend.h
+++ b/Engine/Plugins/Experimental/TextureGraph/Source/TextureGraphEngine/Transform/Expressions/T_Blend.h
@@ -25,6 +25,9 @@ public:
 		SHADER_PARAMETER(float, Opacity)
 	END_SHADER_PARAMETER_STRUCT()
 
+	class FIgnoreAlpha : SHADER_PERMUTATION_BOOL("IGNORE_ALPHA");
+	using FPermutationDomain = TShaderPermutationDomain<FIgnoreAlpha>;
+
 public:
 	static bool ShouldCompilePermutation(const FGlobalShaderPermutationParameters& Parameters)
 	{
@@ -153,20 +156,28 @@ public:
 	T_Blend();
 	~T_Blend();
 
+	struct FBlendSettings
+	{
+		TiledBlobPtr BackgroundTexture;
+		TiledBlobPtr ForegroundTexture;
+		TiledBlobPtr Mask;
+		float Opacity;
+		bool bIgnoreAlpha;
+	};
 	//////////////////////////////////////////////////////////////////////////
 	/// Static functions
 	//////////////////////////////////////////////////////////////////////////
 
-	static TiledBlobPtr				Create(MixUpdateCyclePtr InCycle, BufferDescriptor DesiredDesc, TiledBlobPtr InBackgroundTexture, TiledBlobPtr InForeGroundTexture, TiledBlobPtr InMask, float InOpacity, int InTargetId, EBlendModes InBlendMode);
-	static TiledBlobPtr				CreateNormal(MixUpdateCyclePtr InCycle, BufferDescriptor DesiredDesc, TiledBlobPtr InBackgroundTexture, TiledBlobPtr InForeGroundTexture, TiledBlobPtr InMask, float InOpacity, int InTargetId);
-	static TiledBlobPtr				CreateAdd(MixUpdateCyclePtr InCycle, BufferDescriptor DesiredDesc, TiledBlobPtr InBackgroundTexture, TiledBlobPtr InForeGroundTexture, TiledBlobPtr InMask, float InOpacity, int InTargetId);
-	static TiledBlobPtr				CreateSubtract(MixUpdateCyclePtr InCycle, BufferDescriptor DesiredDesc, TiledBlobPtr InBackgroundTexture, TiledBlobPtr InForeGroundTexture, TiledBlobPtr InMask, float InOpacity, int InTargetId);
-	static TiledBlobPtr				CreateMultiply(MixUpdateCyclePtr InCycle, BufferDescriptor DesiredDesc, TiledBlobPtr InBackgroundTexture, TiledBlobPtr InForeGroundTexture, TiledBlobPtr InMask, float InOpacity, int InTargetId);
-	static TiledBlobPtr				CreateDivide(MixUpdateCyclePtr InCycle, BufferDescriptor DesiredDesc, TiledBlobPtr InBackgroundTexture, TiledBlobPtr InForeGroundTexture, TiledBlobPtr InMask, float InOpacity, int InTargetId);
-	static TiledBlobPtr				CreateDifference(MixUpdateCyclePtr InCycle, BufferDescriptor DesiredDesc, TiledBlobPtr InBackgroundTexture, TiledBlobPtr InForeGroundTexture, TiledBlobPtr InMask, float InOpacity, int InTargetId);
-	static TiledBlobPtr				CreateMax(MixUpdateCyclePtr InCycle, BufferDescriptor DesiredDesc, TiledBlobPtr InBackgroundTexture, TiledBlobPtr InForeGroundTexture, TiledBlobPtr InMask, float InOpacity, int InTargetId);
-	static TiledBlobPtr				CreateMin(MixUpdateCyclePtr InCycle, BufferDescriptor DesiredDesc, TiledBlobPtr InBackgroundTexture, TiledBlobPtr InForeGroundTexture, TiledBlobPtr InMask, float InOpacity, int InTargetId);
-	static TiledBlobPtr				CreateStep(MixUpdateCyclePtr InCycle, BufferDescriptor DesiredDesc, TiledBlobPtr InBackgroundTexture, TiledBlobPtr InForeGroundTexture, TiledBlobPtr InMask, float InOpacity, int InTargetId);
-	static TiledBlobPtr				CreateOverlay(MixUpdateCyclePtr InCycle, BufferDescriptor DesiredDesc, TiledBlobPtr InBackgroundTexture, TiledBlobPtr InForeGroundTexture, TiledBlobPtr InMask, float InOpacity, int InTargetId);
-	static TiledBlobPtr				CreateDistort(MixUpdateCyclePtr InCycle, BufferDescriptor DesiredDesc, TiledBlobPtr InBackgroundTexture, TiledBlobPtr InForeGroundTexture, TiledBlobPtr InMask, float InOpacity, int InTargetId);
+	static TiledBlobPtr				Create(MixUpdateCyclePtr InCycle, BufferDescriptor DesiredDesc, int InTargetId, EBlendModes InBlendMode, const FBlendSettings* InBlendSettings);
+	static TiledBlobPtr				CreateNormal(MixUpdateCyclePtr InCycle, BufferDescriptor DesiredDesc, int InTargetId, const FBlendSettings* InBlendSettings);
+	static TiledBlobPtr				CreateAdd(MixUpdateCyclePtr InCycle, BufferDescriptor DesiredDesc, int InTargetId, const FBlendSettings* InBlendSettings);
+	static TiledBlobPtr				CreateSubtract(MixUpdateCyclePtr InCycle, BufferDescriptor DesiredDesc, int InTargetId, const FBlendSettings* InBlendSettings);
+	static TiledBlobPtr				CreateMultiply(MixUpdateCyclePtr InCycle, BufferDescriptor DesiredDesc, int InTargetId, const FBlendSettings* InBlendSettings);
+	static TiledBlobPtr				CreateDivide(MixUpdateCyclePtr InCycle, BufferDescriptor DesiredDesc, int InTargetId, const FBlendSettings* InBlendSettings);
+	static TiledBlobPtr				CreateDifference(MixUpdateCyclePtr InCycle, BufferDescriptor DesiredDesc, int InTargetId, const FBlendSettings* InBlendSettings);
+	static TiledBlobPtr				CreateMax(MixUpdateCyclePtr InCycle, BufferDescriptor DesiredDesc, int InTargetId, const FBlendSettings* InBlendSettings);
+	static TiledBlobPtr				CreateMin(MixUpdateCyclePtr InCycle, BufferDescriptor DesiredDesc, int InTargetId, const FBlendSettings* InBlendSettings);
+	static TiledBlobPtr				CreateStep(MixUpdateCyclePtr InCycle, BufferDescriptor DesiredDesc, int InTargetId, const FBlendSettings* InBlendSettings);
+	static TiledBlobPtr				CreateOverlay(MixUpdateCyclePtr InCycle, BufferDescriptor DesiredDesc, int InTargetId, const FBlendSettings* InBlendSettings);
+	static TiledBlobPtr				CreateDistort(MixUpdateCyclePtr InCycle, BufferDescriptor DesiredDesc, int InTargetId, const FBlendSettings* InBlendSettings);
 };
