Blueprint crash when stopping blueprint execution on breakpoint

#tests 66a0bf5926265d82e91d5c44

diff --git a/Engine/Plugins/Experimental/TextureGraph/Source/TextureGraph/Private/Blueprint/TG_AsyncExportTask.cpp b/Engine/Plugins/Experimental/TextureGraph/Source/TextureGraph/Private/Blueprint/TG_AsyncExportTask.cpp
--- a/Engine/Plugins/Experimental/TextureGraph/Source/TextureGraph/Private/Blueprint/TG_AsyncExportTask.cpp
+++ b/Engine/Plugins/Experimental/TextureGraph/Source/TextureGraph/Private/Blueprint/TG_AsyncExportTask.cpp
@@ -11,19 +11,20 @@
 UTG_AsyncExportTask::UTG_AsyncExportTask(const FObjectInitializer& ObjectInitializer)
 	: Super(ObjectInitializer)
 {
-	
+
 }
 
 UTG_AsyncExportTask* UTG_AsyncExportTask::TG_AsyncExportTask( UTextureGraph* InTextureGraph, const bool OverwriteTextures, const bool bSave)
 {
 	UTG_AsyncExportTask* Task = NewObject<UTG_AsyncExportTask>();
+	Task->SetFlags(RF_Standalone);
 	Task->OverwriteTextures = OverwriteTextures;
 	Task->bSave = bSave;
 
 	if (InTextureGraph != nullptr)
 	{
 		Task->OrignalTextureGraphPtr = InTextureGraph;
-		Task->TextureGraphPtr = (UTextureGraph*)StaticDuplicateObject(Task->OrignalTextureGraphPtr, GetTransientPackage(), NAME_None, ~RF_Standalone, UTextureGraph::StaticClass());
+		Task->TextureGraphPtr = (UTextureGraph*)StaticDuplicateObject(Task->OrignalTextureGraphPtr, GetTransientPackage(), NAME_None, RF_Standalone, UTextureGraph::StaticClass());
 		FTG_HelperFunctions::InitTargets(Task->TextureGraphPtr);
 		Task->RegisterWithTGAsyncTaskManger();
 	}
@@ -37,7 +38,7 @@ void UTG_AsyncExportTask::Activate()
 	Super::Activate();
 	UE_LOG(LogTextureGraph, Log, TEXT("TG_AsyncExportTask:: Activate"));
 
-	if (TextureGraphPtr == nullptr)
+	if (!IsValid(TextureGraphPtr))
 	{
 		UE_LOG(LogTextureGraph, Warning, TEXT("TG_AsyncExportTask::Cannot export Texture Graph not selected"));
 		return;
@@ -46,12 +47,7 @@ void UTG_AsyncExportTask::Activate()
 	TargetExportSettings = FExportSettings();
 	TargetExportSettings.OnDone.BindUFunction(this, "OnExportDone");
 
-	FTG_HelperFunctions::ExportAsync(TextureGraphPtr, "", "", TargetExportSettings, false, OverwriteTextures, true,bSave)
-		.then([this](bool)
-		{
-			OnExportDone();
-		});
-
+	FTG_HelperFunctions::ExportAsync(TextureGraphPtr, "", "", TargetExportSettings, false, OverwriteTextures, true, bSave);
 }
 
 void UTG_AsyncExportTask::OnExportDone()
@@ -59,6 +55,8 @@ void UTG_AsyncExportTask::OnExportDone()
 	TargetExportSettings.ExportPreset.clear();
 	TargetExportSettings.OnDone.Unbind();
 	OnDone.Broadcast();
+	TextureGraphPtr->FlushInvalidations();
+	ClearFlags(RF_Standalone);
 	SetReadyToDestroy();
 }
 
@@ -67,6 +65,7 @@ void UTG_AsyncExportTask::FinishDestroy()
 	if (TextureGraphPtr != nullptr)
 	{
 		TextureGraphPtr->GetSettings()->FreeTargets();
+		TextureGraphPtr->ClearFlags(RF_Standalone);
 		TextureGraphPtr = nullptr;
 		OrignalTextureGraphPtr = nullptr;
 	}
diff --git a/Engine/Plugins/Experimental/TextureGraph/Source/TextureGraph/Private/Blueprint/TG_AsyncRenderTask.cpp b/Engine/Plugins/Experimental/TextureGraph/Source/TextureGraph/Private/Blueprint/TG_AsyncRenderTask.cpp
--- a/Engine/Plugins/Experimental/TextureGraph/Source/TextureGraph/Private/Blueprint/TG_AsyncRenderTask.cpp
+++ b/Engine/Plugins/Experimental/TextureGraph/Source/TextureGraph/Private/Blueprint/TG_AsyncRenderTask.cpp
@@ -11,6 +11,7 @@
 #include "UObject/Package.h"
 #include "Device/FX/DeviceBuffer_FX.h"
 #include "Model/Mix/MixSettings.h"
+#include "GameDelegates.h"
 
 UTG_AsyncRenderTask::UTG_AsyncRenderTask(const FObjectInitializer& ObjectInitializer)
 	: Super(ObjectInitializer)
@@ -21,13 +22,15 @@ UTG_AsyncRenderTask::UTG_AsyncRenderTask(const FObjectInitializer& ObjectInitial
 UTG_AsyncRenderTask* UTG_AsyncRenderTask::TG_AsyncRenderTask( UTextureGraph* InTextureGraph)
 {
 	UTG_AsyncRenderTask* Task = NewObject<UTG_AsyncRenderTask>();
+	Task->SetFlags(RF_Standalone);
 
 	if (InTextureGraph != nullptr)
 	{
 		Task->OriginalTextureGraphPtr = InTextureGraph;
-		Task->TextureGraphPtr = (UTextureGraph*)StaticDuplicateObject(Task->OriginalTextureGraphPtr, GetTransientPackage(), NAME_None, ~RF_Standalone, UTextureGraph::StaticClass());
+		Task->TextureGraphPtr = (UTextureGraph*)StaticDuplicateObject(Task->OriginalTextureGraphPtr, GetTransientPackage(), NAME_None, RF_Standalone, UTextureGraph::StaticClass());
 		FTG_HelperFunctions::InitTargets(Task->TextureGraphPtr);
 		Task->RegisterWithTGAsyncTaskManger();
+		FWorldDelegates::OnWorldCleanup.AddUObject(Task, &UTG_AsyncRenderTask::OnWorldCleanup);
 	}
 	
 	return Task;
@@ -47,6 +50,7 @@ void UTG_AsyncRenderTask::Activate()
 
 	OutputBlobs.Empty();
 	OutputRts.Empty();
+	TextureGraphPtr->FlushInvalidations();
 
 	FTG_HelperFunctions::RenderAsync(TextureGraphPtr)
 	.then([this](bool bRenderResult) mutable
@@ -73,7 +77,14 @@ void UTG_AsyncRenderTask::Activate()
 	.then([this](bool bRtResult)
 	{
 		OnDone.Broadcast(OutputRts);
-		SetReadyToDestroy();
+		bRenderComplete = true;
+
+		UE_LOG(LogTextureGraph, Log, TEXT("UTG_AsyncRenderTask:: OnDone : bShouldDestroyOnRenderComplete %i"), bShouldDestroyOnRenderComplete ? 1 : 0);
+		if (bShouldDestroyOnRenderComplete)
+		{
+			SetReadyToDestroy();
+		}
+
 		return bRtResult;
 	});
 }
@@ -115,11 +126,35 @@ AsyncBool UTG_AsyncRenderTask::GetRenderTextures()
 	});
 }
 
+void UTG_AsyncRenderTask::OnWorldCleanup(UWorld* World, bool bSessionEnded, bool CleanupResources)
+{
+	FWorldDelegates::OnWorldCleanup.RemoveAll(this);
+	bShouldDestroyOnRenderComplete = true;
+
+	//Destroy now if rendering is already complete
+	if (bRenderComplete)
+	{
+		SetReadyToDestroy();
+	}
+
+	UE_LOG(LogTextureGraph, Log, TEXT("UTG_AsyncRenderTask:: OnWorldCleanup"));
+}
+
+void UTG_AsyncRenderTask::SetReadyToDestroy()
+{
+	UE_LOG(LogTextureGraph, Log, TEXT("UTG_AsyncRenderTask:: SetReadyToDestroy"));
+	TextureGraphPtr->FlushInvalidations();
+	ClearFlags(RF_Standalone);
+	UTG_AsyncTask::SetReadyToDestroy();
+}
+
 void UTG_AsyncRenderTask::FinishDestroy()
 {
+	UE_LOG(LogTextureGraph, Log, TEXT("UTG_AsyncRenderTask:: FinishDestroy"));
 	if (TextureGraphPtr != nullptr)
 	{
 		TextureGraphPtr->GetSettings()->FreeTargets();
+		TextureGraphPtr->ClearFlags(RF_Standalone);
 		TextureGraphPtr = nullptr;
 		OriginalTextureGraphPtr = nullptr;
 	}
diff --git a/Engine/Plugins/Experimental/TextureGraph/Source/TextureGraph/Private/Blueprint/TG_BlueprintFunctionLibrary.cpp b/Engine/Plugins/Experimental/TextureGraph/Source/TextureGraph/Private/Blueprint/TG_BlueprintFunctionLibrary.cpp
--- a/Engine/Plugins/Experimental/TextureGraph/Source/TextureGraph/Private/Blueprint/TG_BlueprintFunctionLibrary.cpp
+++ b/Engine/Plugins/Experimental/TextureGraph/Source/TextureGraph/Private/Blueprint/TG_BlueprintFunctionLibrary.cpp
@@ -85,11 +85,26 @@ T_Expr_Value GetParameterValue_Generic(UObject* WorldContextObject, UTextureGrap
 void UTG_BlueprintFunctionLibrary::SetTextureParameterValue(UObject* WorldContextObject, UTextureGraph* InTextureGraph, FName ParameterName, UTexture* ParameterValue)
 {
 	UTG_Expression_Texture* ExpressionPtr = GetParamExpression<UTG_Expression_Texture>(WorldContextObject, InTextureGraph, ParameterName);
-	if (ExpressionPtr)
+	UTextureRenderTarget2D* RenderTarget2D = Cast<UTextureRenderTarget2D>(ParameterValue);
+
+	if (ExpressionPtr && ExpressionPtr->CanHandleAsset(ParameterValue))
+	{
+		UTexture2D* Texture2D = Cast<UTexture2D>(ParameterValue);
+		if (Texture2D)
+		{
+			UTexture2D* DupTexture = (UTexture2D*)StaticDuplicateObject(ParameterValue, GetTransientPackage(), NAME_None, RF_Transient, UTexture2D::StaticClass());
+			ExpressionPtr->SetAsset(DupTexture);
+		}
+		
+	}
+#if WITH_EDITOR
+	else if (ExpressionPtr && RenderTarget2D)
 	{
-		UTexture2D* DupTexture = (UTexture2D*)StaticDuplicateObject(ParameterValue, GetTransientPackage(), NAME_None, ~RF_Standalone, UTexture2D::StaticClass());
-		ExpressionPtr->SetAsset(DupTexture);
+		UTexture2D* Texture = NewObject<UTexture2D>(GetTransientPackage(), NAME_None, RF_Transient);
+		UKismetRenderingLibrary::ConvertRenderTargetToTexture2DEditorOnly(WorldContextObject, RenderTarget2D, Texture);
+		ExpressionPtr->SetAsset(Texture);
 	}
+#endif
 	else
 	{
 		AddParamWarning(ParameterName, InTextureGraph, "SetTextureParameterValue");
diff --git a/Engine/Plugins/Experimental/TextureGraph/Source/TextureGraph/Public/Blueprint/TG_AsyncRenderTask.h b/Engine/Plugins/Experimental/TextureGraph/Source/TextureGraph/Public/Blueprint/TG_AsyncRenderTask.h
--- a/Engine/Plugins/Experimental/TextureGraph/Source/TextureGraph/Public/Blueprint/TG_AsyncRenderTask.h
+++ b/Engine/Plugins/Experimental/TextureGraph/Source/TextureGraph/Public/Blueprint/TG_AsyncRenderTask.h
@@ -11,6 +11,7 @@
 #include "Data/Blob.h"
 #include "Data/TiledBlob.h"
 #include "TG_AsyncTask.h"
+#include "Engine/World.h"
 #include "TG_AsyncRenderTask.generated.h"
 
 DECLARE_DYNAMIC_MULTICAST_DELEGATE_OneParam(FTSRenderTaskDelegate, const TArray<UTextureRenderTarget2D*>&, OutputRts);
@@ -28,6 +29,8 @@ public:
 
 	virtual void FinishDestroy() override;
 
+	virtual void SetReadyToDestroy() override;
+
 	UPROPERTY(BlueprintAssignable, Category = "TextureGraph")
 	FTSRenderTaskDelegate OnDone;
 
@@ -36,11 +39,18 @@ private:
 
 	AsyncBool GetRenderTextures();
 
+	void OnWorldCleanup(UWorld* World, bool bSessionEnded, bool CleanupResources);
+
 	TArray<UTextureRenderTarget2D*> OutputRts;
 
 	TArray<BlobPtr> OutputBlobs;
 
 	UTextureGraph* OriginalTextureGraphPtr;
+
 	UTextureGraph* TextureGraphPtr;
+
+	bool bShouldDestroyOnRenderComplete = false;
+
+	bool bRenderComplete = false;
 };
 
