Updated the error message for invalid file name and path in texture graph

diff --git a/Engine/Plugins/Experimental/TextureGraph/Source/TextureGraph/Private/TG_HelperFunctions.cpp b/Engine/Plugins/Experimental/TextureGraph/Source/TextureGraph/Private/TG_HelperFunctions.cpp
--- a/Engine/Plugins/Experimental/TextureGraph/Source/TextureGraph/Private/TG_HelperFunctions.cpp
+++ b/Engine/Plugins/Experimental/TextureGraph/Source/TextureGraph/Private/TG_HelperFunctions.cpp
@@ -27,6 +27,12 @@ bool FTG_HelperFunctions::IsFileNameValid(FName FileName, FText& Reason)
 bool FTG_HelperFunctions::IsFolderPathValid(FString FolderPath, FText& Reason)
 {
 	FName::IsValidXName(FolderPath, INVALID_OBJECTPATH_CHARACTERS INVALID_LONGPACKAGE_CHARACTERS, &Reason);
+	
+	if (!Reason.IsEmpty())
+	{
+		Reason = FText::FromString(Reason.ToString().Replace(TEXT("Name"), TEXT("Path")));
+	}
+
 	FPaths::ValidatePath(FolderPath, &Reason);
 	return Reason.IsEmpty();
 }
@@ -185,12 +191,15 @@ JobBatchPtr FTG_HelperFunctions::InitExportBatch(UTextureGraph* InTextureGraph,
 						}
 						if (!IsPathValid)
 						{
-							ErrorMessage += FString::Format(TEXT("Texture Export Error : Invalid path set for OutputSetting {0}"), { OutputSetting.OutputName.ToString() }) + (PathReason.IsEmpty() ? "" : " Folder " + PathReason.ToString());
+							FString InvalidFolderPath = "Invalid folder path. ";
+							FString PathReasonString = PathReason.IsEmpty() ? InvalidFolderPath : InvalidFolderPath + PathReason.ToString();
+							ErrorMessage += FString::Format(TEXT("Texture Export Error : {0} for OutputSetting {1}"), { PathReasonString, OutputSetting.OutputName.ToString() });
 							ErrorMessage += "\n";
 						}
 						if (!IsNameValid)
 						{
-							FString ReasonString = Reason.IsEmpty() ? "Invalid File Name" : Reason.ToString();
+							FString InvalidFileName = "Invalid File Name. ";
+							FString ReasonString = Reason.IsEmpty() ? InvalidFileName : InvalidFileName + Reason.ToString();
 							ErrorMessage += FString::Format(TEXT("Texture Export Error : {0} for OutputSetting {1}"), { ReasonString , OutputSetting.OutputName.ToString()});
 							ErrorMessage += "\n";
 						}
