Fixed thumbs going black for Levels node. Also added some logs and enabled some that were disabled.
Added a new JobArg property to ignore descriptor values when merging. This is useful when arguments like the histogram which don't directly contribute to the dimensions of the result blob need to be ignored in the descriptor merge code.

diff --git a/Engine/Plugins/Experimental/TextureGraph/Source/TextureGraphEditor/Private/Customizations/TG_ScalarCustomization.h b/Engine/Plugins/Experimental/TextureGraph/Source/TextureGraphEditor/Private/Customizations/TG_ScalarCustomization.h
--- a/Engine/Plugins/Experimental/TextureGraph/Source/TextureGraphEditor/Private/Customizations/TG_ScalarCustomization.h
+++ b/Engine/Plugins/Experimental/TextureGraph/Source/TextureGraphEditor/Private/Customizations/TG_ScalarCustomization.h
@@ -138,6 +138,13 @@ public:
 	void OnEndSliderMovement()
 	{
 		bIsUsingSlider = false;
+
+		// set value once more with default flags so the TextureGraph system recognizes a non-interactive change as well.
+		float OrgValue(0);
+		if (ScalarHandle->GetValue(OrgValue) != FPropertyAccess::Fail)
+		{
+			ScalarHandle->SetValue(OrgValue);
+		}
 		GEditor->EndTransaction();
 	}
 
diff --git a/Engine/Plugins/Experimental/TextureGraph/Source/TextureGraphEngine/Job/BlobHasherService.cpp b/Engine/Plugins/Experimental/TextureGraph/Source/TextureGraphEngine/Job/BlobHasherService.cpp
--- a/Engine/Plugins/Experimental/TextureGraph/Source/TextureGraphEngine/Job/BlobHasherService.cpp
+++ b/Engine/Plugins/Experimental/TextureGraph/Source/TextureGraphEngine/Job/BlobHasherService.cpp
@@ -27,7 +27,7 @@ void BlobHasherService::Add(BlobRef BlobObj)
 
 AsyncJobResultPtr BlobHasherService::Tick()
 {
-	UE_LOG(LogIdle_Svc, Verbose, TEXT("Svc_BlobHasher::Tick"));
+	UE_LOG(LogIdle_Svc, VeryVerbose, TEXT("Svc_BlobHasher::Tick"));
 
 	static constexpr size_t MaxBlobs = 8;
 	BlobPtr BlobObjs[MaxBlobs];
diff --git a/Engine/Plugins/Experimental/TextureGraph/Source/TextureGraphEngine/Job/HistogramService.cpp b/Engine/Plugins/Experimental/TextureGraph/Source/TextureGraphEngine/Job/HistogramService.cpp
--- a/Engine/Plugins/Experimental/TextureGraph/Source/TextureGraphEngine/Job/HistogramService.cpp
+++ b/Engine/Plugins/Experimental/TextureGraph/Source/TextureGraphEngine/Job/HistogramService.cpp
@@ -42,7 +42,7 @@ AsyncJobResultPtr HistogramService::Tick()
 {
 	check(IsInGameThread());
 
-	UE_LOG(LogIdle_Svc, Verbose, TEXT("Svc_Histogram::Tick"));
+	UE_LOG(LogIdle_Svc, VeryVerbose, TEXT("Svc_Histogram::Tick"));
 
 	if (Batch)
 	{
diff --git a/Engine/Plugins/Experimental/TextureGraph/Source/TextureGraphEngine/Job/Job.cpp b/Engine/Plugins/Experimental/TextureGraph/Source/TextureGraphEngine/Job/Job.cpp
--- a/Engine/Plugins/Experimental/TextureGraph/Source/TextureGraphEngine/Job/Job.cpp
+++ b/Engine/Plugins/Experimental/TextureGraph/Source/TextureGraphEngine/Job/Job.cpp
@@ -280,7 +280,7 @@ BufferDescriptor Job::GetCombinedDesc(BufferDescriptor& ArgsDescCombined, size_t
 		JobArgPtr Arg = Args[ArgIndex];
 		const BufferDescriptor* ArgDesc = Arg->GetDescriptor();
 
-		if (ArgDesc)
+		if (ArgDesc && !Arg->IgnoreDesc())
 		{
 			Descs.push_back(*ArgDesc);
 		}
@@ -662,7 +662,7 @@ bool Job::CheckCulled(JobRunInfo InRunInfo)
 
 	/// Nothing to do over here. We just call the BeginNative and EndNative here
 	/// so that the job collects timing information
-	BeginNative(RunInfo);
+	BeginNative(InRunInfo);
 	EndNative();
 
 	/// Make sure that the promise is resolved over here, if the job is culled
@@ -1042,9 +1042,10 @@ AsyncJobResultPtr Job::Run(JobRunInfo InRunInfo)
 //////////////////////////////////////////////////////////////////////////
 AsyncInt Job::BeginNative(JobRunInfo InRunInfo)
 {
+	RunInfo = InRunInfo;
+
 	UE_LOG(LogJob, VeryVerbose, TEXT("Job::BeginNative: %llu.%llu.%s"), RunInfo.Batch->GetBatchId(), Id, *Transform->GetName());
 
-	RunInfo = InRunInfo;
 	Stats.BeginNativeTime = Util::Time();
 
 	if (bIsCulled)
@@ -1063,10 +1064,13 @@ void Job::MarkJobDone()
 
 	check((!Result || Result->IsFinalised()) && (!ResultOrg || ResultOrg->IsFinalised()));
 	if (RunInfo.Batch)
+	{
 		RunInfo.Batch->OnJobDone(this, GetJobId());
+		
+		UE_LOG(LogJob, VeryVerbose, TEXT("Job::Done: %llu.%d.%s"), RunInfo.Batch->GetBatchId(), GetJobId(), *Transform->GetName());
+	}
 
 	DeviceNativeTask::bIsDone = true;
-	UE_LOG(LogJob, VeryVerbose, TEXT("Job::Done: %llu.%d.%s"), RunInfo.Batch->GetBatchId(), GetJobId(), *Transform->GetName());
 
 	Prev.clear();
 }
diff --git a/Engine/Plugins/Experimental/TextureGraph/Source/TextureGraphEngine/Job/Job.h b/Engine/Plugins/Experimental/TextureGraph/Source/TextureGraphEngine/Job/Job.h
--- a/Engine/Plugins/Experimental/TextureGraph/Source/TextureGraphEngine/Job/Job.h
+++ b/Engine/Plugins/Experimental/TextureGraph/Source/TextureGraphEngine/Job/Job.h
@@ -13,7 +13,7 @@ class Scheduler;
 class JobBatch;
 class MixUpdateCycle;
 
-DECLARE_LOG_CATEGORY_EXTERN(LogJob, Log, Verbose);
+DECLARE_LOG_CATEGORY_EXTERN(LogJob, Log, All);
 
 class RenderMesh;
 
diff --git a/Engine/Plugins/Experimental/TextureGraph/Source/TextureGraphEngine/Job/JobArgs.h b/Engine/Plugins/Experimental/TextureGraph/Source/TextureGraphEngine/Job/JobArgs.h
--- a/Engine/Plugins/Experimental/TextureGraph/Source/TextureGraphEngine/Job/JobArgs.h
+++ b/Engine/Plugins/Experimental/TextureGraph/Source/TextureGraphEngine/Job/JobArgs.h
@@ -42,6 +42,8 @@ struct TEXTUREGRAPHENGINE_API JobArgBindInfo
 class TEXTUREGRAPHENGINE_API JobArg
 {
 protected:
+	bool							bIgnoreDesc = false;/// Whether to ignore the descriptor for this argument
+														/// when combining descriptors during merge
 	bool							bIgnoreHash = false;/// Whether to ignore this argument in hash calculation or not
 	bool							bUnbound = false;	/// Sometimes you need to add an argument for hashing calculation 
 														/// and not necessarily binding to a BlobTransform
@@ -72,6 +74,9 @@ public:
 
 	FORCEINLINE bool				Unbounded() const { return bUnbound; }
 	FORCEINLINE void 				WithUnbounded(bool unbounded) { bUnbound = unbounded; }
+
+	FORCEINLINE bool				IgnoreDesc() const { return bIgnoreDesc; }
+	FORCEINLINE void 				WithIgnoreDesc(bool ignoreDesc) { bIgnoreDesc = ignoreDesc; }
 };
 
 typedef std::shared_ptr<JobArg>		JobArgPtr;
@@ -395,6 +400,11 @@ public:
 //Blob can be fetched in HLSL using GetFullBlob(inout float 4) method. See AdjustUVGeneric.usf and TiledFetch_Combined.ush for help
 #define ARG_COMBINEDBLOB(v, n)				std::make_shared<JobArg_Blob_Combined>(v, n)	//With Custom SRV
 
+FORCEINLINE JobArgPtr				WithIgnoreDesc(JobArgPtr Arg, bool bIgnoreDesc = true)
+{
+	Arg->WithIgnoreDesc(bIgnoreDesc);
+	return Arg;
+}
 
 FORCEINLINE JobArgPtr				WithIgnoreHash(JobArgPtr Arg, bool bIgnoreHash = true)
 {
diff --git a/Engine/Plugins/Experimental/TextureGraph/Source/TextureGraphEngine/Job/JobBatch.h b/Engine/Plugins/Experimental/TextureGraph/Source/TextureGraphEngine/Job/JobBatch.h
--- a/Engine/Plugins/Experimental/TextureGraph/Source/TextureGraphEngine/Job/JobBatch.h
+++ b/Engine/Plugins/Experimental/TextureGraph/Source/TextureGraphEngine/Job/JobBatch.h
@@ -20,7 +20,7 @@ typedef GenericTaskPriorityQueue<std::shared_ptr<Job>, DeviceNativeTask_CompareS
 
 struct FInvalidationDetails;
 
-DECLARE_LOG_CATEGORY_EXTERN(LogBatch, Log, Log);
+DECLARE_LOG_CATEGORY_EXTERN(LogBatch, Log, All);
 
 //////////////////////////////////////////////////////////////////////////
 /// JobBatch: A series of jobs put together that need to execute at the 
diff --git a/Engine/Plugins/Experimental/TextureGraph/Source/TextureGraphEngine/Job/MinMaxService.cpp b/Engine/Plugins/Experimental/TextureGraph/Source/TextureGraphEngine/Job/MinMaxService.cpp
--- a/Engine/Plugins/Experimental/TextureGraph/Source/TextureGraphEngine/Job/MinMaxService.cpp
+++ b/Engine/Plugins/Experimental/TextureGraph/Source/TextureGraphEngine/Job/MinMaxService.cpp
@@ -16,7 +16,7 @@ MinMaxService::~MinMaxService()
 
 AsyncJobResultPtr MinMaxService::Tick()
 {
-	UE_LOG(LogIdle_Svc, Verbose, TEXT("Svc_MinMax::Tick"));
+	UE_LOG(LogIdle_Svc, VeryVerbose, TEXT("Svc_MinMax::Tick"));
 
 	static constexpr size_t MaxBlobs = 8;
 	BlobPtr BlobObjs[MaxBlobs] = { 0 };
diff --git a/Engine/Plugins/Experimental/TextureGraph/Source/TextureGraphEngine/Job/MipMapService.cpp b/Engine/Plugins/Experimental/TextureGraph/Source/TextureGraphEngine/Job/MipMapService.cpp
--- a/Engine/Plugins/Experimental/TextureGraph/Source/TextureGraphEngine/Job/MipMapService.cpp
+++ b/Engine/Plugins/Experimental/TextureGraph/Source/TextureGraphEngine/Job/MipMapService.cpp
@@ -16,7 +16,7 @@ MipMapService::~MipMapService()
 
 AsyncJobResultPtr MipMapService::Tick()
 {
-	UE_LOG(LogIdle_Svc, Verbose, TEXT("Svc_MipMap::Tick"));
+	UE_LOG(LogIdle_Svc, VeryVerbose, TEXT("Svc_MipMap::Tick"));
 
 	static constexpr size_t MaxBlobs = 8;
 	BlobPtr BlobObjs[MaxBlobs] = { 0 };
diff --git a/Engine/Plugins/Experimental/TextureGraph/Source/TextureGraphEngine/Job/Scheduler.cpp b/Engine/Plugins/Experimental/TextureGraph/Source/TextureGraphEngine/Job/Scheduler.cpp
--- a/Engine/Plugins/Experimental/TextureGraph/Source/TextureGraphEngine/Job/Scheduler.cpp
+++ b/Engine/Plugins/Experimental/TextureGraph/Source/TextureGraphEngine/Job/Scheduler.cpp
@@ -263,7 +263,7 @@ void Scheduler::UpdateIdle()
 			return;
 	}
 
-	UE_LOG(LogIdle_Svc, Verbose, TEXT("Scheduler::UpdateIdle"));
+	UE_LOG(LogIdle_Svc, VeryVerbose, TEXT("Scheduler::UpdateIdle"));
 
 	if (TextureGraphEngine::IsDestroying())
 		return;
diff --git a/Engine/Plugins/Experimental/TextureGraph/Source/TextureGraphEngine/Job/TempHashService.cpp b/Engine/Plugins/Experimental/TextureGraph/Source/TextureGraphEngine/Job/TempHashService.cpp
--- a/Engine/Plugins/Experimental/TextureGraph/Source/TextureGraphEngine/Job/TempHashService.cpp
+++ b/Engine/Plugins/Experimental/TextureGraph/Source/TextureGraphEngine/Job/TempHashService.cpp
@@ -14,7 +14,7 @@ TempHashService::~TempHashService()
 
 AsyncJobResultPtr TempHashService::Tick()
 {
-	UE_LOG(LogIdle_Svc, Verbose, TEXT("Svc_BlobHasher::Tick"));
+	UE_LOG(LogIdle_Svc, VeryVerbose, TEXT("Svc_BlobHasher::Tick"));
 
 	check(IsInGameThread());
 
diff --git a/Engine/Plugins/Experimental/TextureGraph/Source/TextureGraphEngine/Job/ThumbnailsService.cpp b/Engine/Plugins/Experimental/TextureGraph/Source/TextureGraphEngine/Job/ThumbnailsService.cpp
--- a/Engine/Plugins/Experimental/TextureGraph/Source/TextureGraphEngine/Job/ThumbnailsService.cpp
+++ b/Engine/Plugins/Experimental/TextureGraph/Source/TextureGraphEngine/Job/ThumbnailsService.cpp
@@ -21,7 +21,7 @@ AsyncJobResultPtr ThumbnailsService::Tick()
  	if (!Batch)
  		return cti::make_ready_continuable<JobResultPtr>(std::make_shared<JobResult>());
  
- 	UE_LOG(LogIdle_Svc, Verbose, TEXT("Svc_BlobHasher::Tick"));
+ 	UE_LOG(LogIdle_Svc, VeryVerbose, TEXT("ThumbnailsService::Tick"));
  
  	/// Move over to the next cycle
  	JobBatchPtr CurrentBatch = GetNextUpdateCycle();
@@ -34,13 +34,21 @@ AsyncJobResultPtr ThumbnailsService::Tick()
 	return CurrentBatch->Exec([=, this](JobBatch*)	/// Instead of passing it as an argument, JobBatch should be a return type; this is to keep it from cyclic dependancy
 	{
 		OnUpdateThumbnailDelegate.Broadcast(CurrentBatch);
+		UE_LOG(LogIdle_Svc, Verbose, TEXT("ThumbnailsService, Batch fully queued: %llu. Triggering Observer::BatchJobsDone ..."), CurrentBatch ? CurrentBatch->GetBatchId() : -1);
+
 		TextureGraphEngine::GetInstance()->GetScheduler()->GetObserverSource()->BatchJobsDone(CurrentBatch);
+
+		UE_LOG(LogIdle_Svc, Verbose, TEXT("ThumbnailsService Observer::BatchJobsDone finished for Batch: %llu"), CurrentBatch ? CurrentBatch->GetBatchId() : -1);
 		
 	})
  	.then([this, CurrentBatch]()
  	{
+ 		UE_LOG(LogBatch, Verbose, TEXT("ThumbnailsService triggering Observer::BatchDone for Batch: %llu ..."), CurrentBatch ? CurrentBatch->GetBatchId() : -1);
+
  		TextureGraphEngine::GetInstance()->GetScheduler()->GetObserverSource()->BatchDone(CurrentBatch); // notify observer
- 		
+
+ 		UE_LOG(LogBatch, Verbose, TEXT("ThumbnailsService Observer::BatchDone finished for Batch: %llu"), CurrentBatch ? CurrentBatch->GetBatchId() : -1);
+			
  		return std::make_shared<JobResult>();
  	});
 }
@@ -82,7 +90,7 @@ JobBatchPtr ThumbnailsService::CreateNewUpdateCycle(UMixInterface* Mix)
 	SceneTargetUpdatePtr Target = std::make_shared<MixTargetUpdate>(Mix, 0);
 	Target->InvalidateAllTiles();
 	NewBatch->GetCycle()->AddTarget(Target);
-
+	// NewBatch->SetCaptureRenderDoc(true);
 	return NewBatch;
 }
 
diff --git a/Engine/Plugins/Experimental/TextureGraph/Source/TextureGraphEngine/Transform/Expressions/T_Levels.cpp b/Engine/Plugins/Experimental/TextureGraph/Source/TextureGraphEngine/Transform/Expressions/T_Levels.cpp
--- a/Engine/Plugins/Experimental/TextureGraph/Source/TextureGraphEngine/Transform/Expressions/T_Levels.cpp
+++ b/Engine/Plugins/Experimental/TextureGraph/Source/TextureGraphEngine/Transform/Expressions/T_Levels.cpp
@@ -329,7 +329,7 @@ TiledBlobPtr	T_Levels::Create(MixUpdateCyclePtr Cycle, BufferDescriptor DesiredO
 		->AddArg(ARG_FLOAT(InLevels->OutLow, "OutLow"))
 		->AddArg(ARG_FLOAT(InLevels->OutHigh, "OutHigh"))
 		->AddArg(ARG_FLOAT(OutputRange, "OutputRange"))
-		->AddArg(std::make_shared<JobArg_Blob>(JobArg_Blob(Histogram, "Histogram").WithNotHandleTiles()))
+		->AddArg(WithIgnoreDesc(std::make_shared<JobArg_Blob>(JobArg_Blob(Histogram, "Histogram").WithNotHandleTiles())))
 		;
 
 
