Fixed issue: Color node output does not render in Node preview

diff --git a/Engine/Plugins/Experimental/TextureGraph/Source/TextureGraphEditor/Private/STG_NodePreview.cpp b/Engine/Plugins/Experimental/TextureGraph/Source/TextureGraphEditor/Private/STG_NodePreview.cpp
--- a/Engine/Plugins/Experimental/TextureGraph/Source/TextureGraphEditor/Private/STG_NodePreview.cpp
+++ b/Engine/Plugins/Experimental/TextureGraph/Source/TextureGraphEditor/Private/STG_NodePreview.cpp
@@ -69,6 +69,11 @@ TOptional<TVariant<FColor, FLinearColor>> FNodeViewer::GetCurrentImagePixelColor
 	return {};
 }
 
+UE::ImageWidgets::SImageViewport::FDrawSettings FNodeViewer::GetDrawSettings() const
+{
+	return DrawSettings;
+}
+
 FText FNodeViewer::GetFormatLabelText() const
 {
 	if (NodeTexture != nullptr)
@@ -87,12 +92,13 @@ bool FNodeViewer::IsSingleChannel() const
 	return NodeDescriptor.ItemsPerPoint == 1;
 }
 
-void FNodeViewer::SetTexture(const BlobPtr& InBlob)
+void FNodeViewer::SetTexture(const BlobPtr& InBlob, const FLinearColor InClearColor /*= FLinearColor(0.1, 0.1, 0.1, 1)*/)
 {
 	NodeTexture = nullptr;
 	CurrentBlob = InBlob;
 
 	NodeDescriptor = {};
+	DrawSettings.ClearColor = InClearColor;
 
 	if (InBlob)
 	{
@@ -112,6 +118,11 @@ void FNodeViewer::SetRGBA(bool bR, bool bG, bool bB, bool bA)
 	bRGBA[3] = bA;
 }
 
+void FNodeViewer::SetDrawSettings(const UE::ImageWidgets::SImageViewport::FDrawSettings& InDrawSettings)
+{
+	DrawSettings = InDrawSettings;
+}
+
 void FNodeViewer::DrawTexture(const FTextureResource* TextureResource, FCanvas* Canvas, const FDrawProperties::FPlacement& Placement,
                               const FDrawProperties::FMip& Mip) const
 {
@@ -216,10 +227,22 @@ STG_NodePreviewWidget::~STG_NodePreviewWidget()
 	FPreviewViewerCommands::Unregister();
 }
 
+UE::ImageWidgets::SImageViewport::FDrawSettings STG_NodePreviewWidget::GetDrawSettings() const
+{
+	return NodeViewer->GetDrawSettings();
+}
+
 void STG_NodePreviewWidget::Construct(const FArguments& InArgs)
 {
 	NodeViewer = MakeShared<FNodeViewer>();
 
+	NodeViewer->SetDrawSettings(UE::ImageWidgets::SImageViewport::FDrawSettings{
+					.ClearColor = FVector3f(0.1f),
+					.bBorderEnabled = false,
+					.bBackgroundColorEnabled = true,
+					.BackgroundColor = FLinearColor::Black,
+					.bBackgroundCheckerEnabled = true});
+
 	OnNodeBlobChanged = InArgs._OnNodeBlobChanged;
 
 	FPreviewViewerCommands::Register();
@@ -249,13 +272,7 @@ void STG_NodePreviewWidget::Construct(const FArguments& InArgs)
 			SAssignNew(Viewport, UE::ImageWidgets::SImageViewport, NodeViewer.ToSharedRef())
 				.ToolbarExtender(ToolbarExtender)
 				.StatusBarExtender(StatusBarExtender)
-				.DrawSettings(UE::ImageWidgets::SImageViewport::FDrawSettings{
-					.ClearColor = FVector3f(0.1f),
-					.bBorderEnabled = false,
-					.bBackgroundColorEnabled = true,
-					.BackgroundColor = FLinearColor::Black,
-					.bBackgroundCheckerEnabled = true
-				})
+				.DrawSettings(this,&STG_NodePreviewWidget::GetDrawSettings)
 				.ControllerSettings(UE::ImageWidgets::SImageViewport::FControllerSettings{
 					.DefaultZoomMode = UE::ImageWidgets::SImageViewport::FControllerSettings::EDefaultZoomMode::Fill
 				})
@@ -297,54 +314,78 @@ void STG_NodePreviewWidget::NodeDeleted(const UTG_Node* Node)
 void STG_NodePreviewWidget::Update() const
 {
 	const UTG_Node* PreviewNode = LockedNode ? LockedNode : SelectedNode;
-
-	// Get blob if preview node is valid or assign a nullptr if not.
-	const BlobPtr Blob = [PreviewNode]() -> BlobPtr
+	BlobPtr Blob;
+	
+	// Get the Variant if preview node is valid or assign a nullptr if not.
+	const FTG_Variant* Variant = nullptr;
+	TArray<FTG_Variant> OutVariants;
+	TArray<FName>* OutNames = nullptr;
+
+	if(PreviewNode)
 	{
-		if (PreviewNode)
+		PreviewNode->GetAllOutputValues(OutVariants, OutNames);
+
+		if (!OutVariants.IsEmpty())
 		{
-			TArray<FTG_Texture> OutTextures;
-			PreviewNode->GetAllOutputValues(OutTextures);
+			Variant = &OutVariants[0];
+		}	
+	}
 
-			if (!OutTextures.IsEmpty())
+	// Lambda for actually updating the preview.
+	auto UpdatePreview = [this, Variant]
+	{
+		if(Variant)
+		{
+			if(Variant->IsTexture())
 			{
-				return OutTextures[0].RasterBlob;
+				NodeViewer->SetTexture(Variant->GetTexture().RasterBlob);
+			}
+			else if(Variant->IsColor())
+			{
+				NodeViewer->SetTexture(nullptr, Variant->GetColor());
+			}
+			else
+			{
+				NodeViewer->SetTexture(nullptr);	
 			}
 		}
-
-		return nullptr;
-	}();
-
-	// Lambda for actually updating the preview.
-	auto UpdatePreview = [this, Blob]
-	{
-		NodeViewer->SetTexture(Blob);
+		else
+		{
+			NodeViewer->SetTexture(nullptr);
+		}
+		
 		Viewport->ResetZoom(NodeViewer->GetCurrentImageInfo().Size);
 	};
 
-	if (Blob)
+	if(Variant && Variant->IsTexture())
 	{
-		// Set non-empty preview if blob is valid.
-		// Note that if the blob is tiled, it needs to be combined first.
-		if (Blob->IsTiled())
+		Blob = Variant->GetTexture().RasterBlob;
+
+		if(Blob)
 		{
-			Blob->OnFinalise()
-			    .then([Blob]()
-			    {
-				    const TiledBlobPtr BlobTiled = std::static_pointer_cast<TiledBlob>(Blob);
-				    return BlobTiled->CombineTiles(false, false);
-			    })
-			    .then(UpdatePreview);
+			if (Blob->IsTiled())
+			{
+				Blob->OnFinalise()
+					.then([Blob]()
+					{
+						const TiledBlobPtr BlobTiled = std::static_pointer_cast<TiledBlob>(Blob);
+						return BlobTiled->CombineTiles(false, false);
+					})
+					.then(UpdatePreview);
+			}
+			else
+			{
+				Blob->OnFinalise()
+					.then(UpdatePreview);
+			}
 		}
 		else
 		{
-			Blob->OnFinalise()
-			    .then(UpdatePreview);
+			UpdatePreview();
 		}
 	}
 	else
 	{
-		// Set black preview, i.e. Blob is nullptr.
 		UpdatePreview();
 	}
 
diff --git a/Engine/Plugins/Experimental/TextureGraph/Source/TextureGraphEditor/Private/STG_NodePreview.h b/Engine/Plugins/Experimental/TextureGraph/Source/TextureGraphEditor/Private/STG_NodePreview.h
--- a/Engine/Plugins/Experimental/TextureGraph/Source/TextureGraphEditor/Private/STG_NodePreview.h
+++ b/Engine/Plugins/Experimental/TextureGraph/Source/TextureGraphEditor/Private/STG_NodePreview.h
@@ -3,6 +3,7 @@
 #pragma once
 
 #include "IImageViewer.h"
+#include "SImageViewport.h"
 #include "Data/Blob.h"
 #include "Widgets/SCompoundWidget.h"
 
@@ -41,11 +42,16 @@ public:
 	bool IsSingleChannel() const;
 
 	/** Sets the node buffer to a given node. */
-	void SetTexture(const BlobPtr& InBlob);
+	void SetTexture(const BlobPtr& InBlob, FLinearColor InClearColor = FLinearColor(0.1, 0.1, 0.1, 1));
 
 	/** Toggles RGBA components for display. */
 	void SetRGBA(bool bR, bool bG, bool bB, bool bA);
 
+	/** Sets the required Draw settings for the viewport. */ 
+	void SetDrawSettings(const UE::ImageWidgets::SImageViewport::FDrawSettings& InDrawSettings);
+
+	/** Returns the Draw settings for the viewport. */ 
+	UE::ImageWidgets::SImageViewport::FDrawSettings GetDrawSettings() const;
 private:
 	/** Draws the node texture in the viewport */
 	void DrawTexture(const FTextureResource* TextureResource, FCanvas* Canvas, const FDrawProperties::FPlacement& TilePlacementInfo,
@@ -71,6 +77,9 @@ private:
 
 	/** Node meta data.  */
 	BufferDescriptor NodeDescriptor;
+
+	/** Current draw settings for the viewport. */ 
+	UE::ImageWidgets::SImageViewport::FDrawSettings DrawSettings;
 };
 
 /**
@@ -91,6 +100,7 @@ public:
 
 	virtual ~STG_NodePreviewWidget() override;
 
+	UE::ImageWidgets::SImageViewport::FDrawSettings GetDrawSettings() const;
 	void Construct(const FArguments& InArgs);
 
 	/** Notify the preview about the node selection having changed. */
