Fixed issue: Rendering of the output breaks after making a specific setup in the Graph

We dont need to create mips for single pixel textures. We are generally keeping mips off for everything. This single pixel mip enabled texture was creating a ripple effect and was turning on the mips for other texture descriptors. We started getting empty white mips which was causing artefacts in the viewport. This was also causing the white thumbs (mip bool was turned on for them too).

diff --git a/Engine/Plugins/Experimental/TextureGraph/Source/TextureGraphEngine/Transform/Expressions/T_FlatColorTexture.cpp b/Engine/Plugins/Experimental/TextureGraph/Source/TextureGraphEngine/Transform/Expressions/T_FlatColorTexture.cpp
--- a/Engine/Plugins/Experimental/TextureGraph/Source/TextureGraphEngine/Transform/Expressions/T_FlatColorTexture.cpp
+++ b/Engine/Plugins/Experimental/TextureGraph/Source/TextureGraphEngine/Transform/Expressions/T_FlatColorTexture.cpp
@@ -21,7 +21,7 @@ BufferDescriptor T_FlatColorTexture::GetFlatColorDesc(FString name, BufferFormat
 	Desc.ItemsPerPoint = 4;
 	Desc.Width = 1;
 	Desc.Height = 1;
-	Desc.bMipMaps = true;
+	Desc.bMipMaps = false;
 	
 	return Desc;
 }
@@ -40,7 +40,7 @@ TiledBlobPtr T_FlatColorTexture::Create(MixUpdateCyclePtr InCycle, BufferDescrip
 	Desc.Height = InCycle->GetMix()->GetNumYTiles();
 	Desc.Format = BufferFormat::Byte;
 	Desc.ItemsPerPoint = 4;
-	Desc.bMipMaps = true;
+	Desc.bMipMaps = false;
 	
 	BufferDescriptor OutputDesc = BufferDescriptor::Combine(Desc, DesiredOutputDesc);
 	if (OutputDesc.Format == BufferFormat::Byte) // Byte size automatically enable sRGB
