Fixed Deselected Output is unexpectedly exported in Texture Graph Exporter

diff --git a/Engine/Plugins/Experimental/TextureGraph/Source/TextureGraph/Private/Expressions/Output/TG_Expression_Output.cpp b/Engine/Plugins/Experimental/TextureGraph/Source/TextureGraph/Private/Expressions/Output/TG_Expression_Output.cpp
--- a/Engine/Plugins/Experimental/TextureGraph/Source/TextureGraph/Private/Expressions/Output/TG_Expression_Output.cpp
+++ b/Engine/Plugins/Experimental/TextureGraph/Source/TextureGraph/Private/Expressions/Output/TG_Expression_Output.cpp
@@ -151,12 +151,7 @@ void UTG_Expression_Output::InitializeOutputSettings()
 	Settings->FromString(OutputSettings.ToString());
 }
 
-void UTG_Expression_Output::SetExport(bool bExport)
+void UTG_Expression_Output::SetShouldExport(bool InShouldExport)
 {
-	OutputSettings.bExport = bExport;
-	UTG_Pin* Settings = GetParentNode()->GetPin(GET_MEMBER_NAME_CHECKED(UTG_Expression_Output, OutputSettings));
-
-	// We are calling this instead of SetValue() to avoid rendering trigger.
-	// TODO: Better solution is that we should have NoInvalidate meta in UProperty or come up with a better solution.
-	Settings->EditSelfVar()->SetValueFromStr(OutputSettings.ToString());
+	bShouldExport = InShouldExport;
 }
\ No newline at end of file
diff --git a/Engine/Plugins/Experimental/TextureGraph/Source/TextureGraph/Private/TG_HelperFunctions.cpp b/Engine/Plugins/Experimental/TextureGraph/Source/TextureGraph/Private/TG_HelperFunctions.cpp
--- a/Engine/Plugins/Experimental/TextureGraph/Source/TextureGraph/Private/TG_HelperFunctions.cpp
+++ b/Engine/Plugins/Experimental/TextureGraph/Source/TextureGraph/Private/TG_HelperFunctions.cpp
@@ -139,7 +139,7 @@ JobBatchPtr FTG_HelperFunctions::InitExportBatch(UTextureGraph* InTextureGraph,
 			{
 				FTG_OutputSettings& OutputSetting = TargetExpression->OutputSettings;
 
-				if (OutputSetting.bExport || ExportAllOutputs)
+				if (TargetExpression->GetShouldExport() || ExportAllOutputs)
 				{
 					FTG_EvaluationContext EvaluationContext;
 					EvaluationContext.Cycle = Batch->GetCycle();
diff --git a/Engine/Plugins/Experimental/TextureGraph/Source/TextureGraph/Public/Expressions/Output/TG_Expression_Output.h b/Engine/Plugins/Experimental/TextureGraph/Source/TextureGraph/Public/Expressions/Output/TG_Expression_Output.h
--- a/Engine/Plugins/Experimental/TextureGraph/Source/TextureGraph/Public/Expressions/Output/TG_Expression_Output.h
+++ b/Engine/Plugins/Experimental/TextureGraph/Source/TextureGraph/Public/Expressions/Output/TG_Expression_Output.h
@@ -33,6 +33,9 @@ public:
 	UPROPERTY(EditAnywhere, Category = NoCategory, meta = (TGType = "TG_Input", CollapsableChildProperties,ShowOnlyInnerProperties, FullyExpand, NoResetToDefault, PinDisplayName = "Settings") )
 	FTG_OutputSettings OutputSettings;
 
+	UPROPERTY()
+	bool bShouldExport = true;
+
 	virtual FTG_Name GetDefaultName() const override { return TEXT("Output");}
 	virtual void SetTitleName(FName NewName) override;
 	virtual FName GetTitleName() const override;
@@ -41,6 +44,7 @@ public:
 	void UpdateBufferDescriptorValues();
 	void InitializeOutputSettings();
 
-	//Sets the Export bool in OutputSettings pin
-	void SetExport(bool bExport);
+	void SetShouldExport(bool InShouldExport);
+
+	bool GetShouldExport() const { return bShouldExport; }
 };
diff --git a/Engine/Plugins/Experimental/TextureGraph/Source/TextureGraph/Public/TG_OutputSettings.h b/Engine/Plugins/Experimental/TextureGraph/Source/TextureGraph/Public/TG_OutputSettings.h
--- a/Engine/Plugins/Experimental/TextureGraph/Source/TextureGraph/Public/TG_OutputSettings.h
+++ b/Engine/Plugins/Experimental/TextureGraph/Source/TextureGraph/Public/TG_OutputSettings.h
@@ -83,9 +83,6 @@ struct TEXTUREGRAPH_API FTG_OutputSettings
 	UPROPERTY(BlueprintReadWrite, EditAnywhere, Category = "Advanced", DisplayName = "sRGB", Meta = (NoResetToDefault, EditCondition = "TexturePresetType == ETG_TexturePresetType::None"))
 		bool bSRGB = false;
 
-	UPROPERTY()
-	bool bExport = true;
-
 	FString GetFullOutputName() { return  FString::Format(TEXT("{0}"), { BaseName.ToString()});}
 
 	bool operator==(const FTG_OutputSettings& Other) const
diff --git a/Engine/Plugins/Experimental/TextureGraph/Source/TextureGraphEditor/Private/Customizations/TG_OutputSettingsCustomization.h b/Engine/Plugins/Experimental/TextureGraph/Source/TextureGraphEditor/Private/Customizations/TG_OutputSettingsCustomization.h
--- a/Engine/Plugins/Experimental/TextureGraph/Source/TextureGraphEditor/Private/Customizations/TG_OutputSettingsCustomization.h
+++ b/Engine/Plugins/Experimental/TextureGraph/Source/TextureGraphEditor/Private/Customizations/TG_OutputSettingsCustomization.h
@@ -255,12 +255,12 @@ public:
 		TAttribute<bool> OutputSettingsEditCondition = TAttribute<bool>::Create(
 			[this, OutputExpression]()
 			{
-				return OutputExpression->OutputSettings.bExport;
+				return OutputExpression->GetShouldExport();
 			});
 		
 		Row->EditCondition(OutputSettingsEditCondition, FOnBooleanValueChanged::CreateLambda([this, OutputExpression](bool NewValue)
 			{
-				OutputExpression->SetExport(NewValue);
+				OutputExpression->SetShouldExport(NewValue);
 			}));
 	}
 };
diff --git a/Engine/Plugins/Experimental/TextureGraph/Source/TextureGraphEditor/Private/STG_OutputSelectionDlg.cpp b/Engine/Plugins/Experimental/TextureGraph/Source/TextureGraphEditor/Private/STG_OutputSelectionDlg.cpp
--- a/Engine/Plugins/Experimental/TextureGraph/Source/TextureGraphEditor/Private/STG_OutputSelectionDlg.cpp
+++ b/Engine/Plugins/Experimental/TextureGraph/Source/TextureGraphEditor/Private/STG_OutputSelectionDlg.cpp
@@ -138,7 +138,7 @@ void STG_OutputSelectionDlg::AddExportItems()
 					.Name(FText::FromString(OutputName.ToString()))
 					.ThumbnailWidget(ThumbnailWidget)
 					.OnOutputSelectionChanged(this, &STG_OutputSelectionDlg::OnOutputSelectionChanged)
-					.bIsSelected(OutputSetting.bExport)
+					.bIsSelected(TargetExpression->GetShouldExport())
 				];
 				ScrollBox->AddSlot()
 				[
@@ -177,9 +177,13 @@ void STG_OutputSelectionDlg::OnOutputSelectionChanged(const FString ItemName, EC
 		UTG_Expression_Output* TargetExpression = Cast<UTG_Expression_Output>(Node->GetExpression());
 		if (TargetExpression)
 		{
-			if (TargetExpression->OutputSettings.OutputName == ItemName)
+			UTG_Pin* Pin = Node->GetPin(GET_MEMBER_NAME_CHECKED(UTG_Expression_Output, Output));
+
+			FName OutputName = Pin->GetAliasName();
+
+			if (OutputName == ItemName)
 			{
-				TargetExpression->SetExport( NewState == ECheckBoxState::Checked);
+				TargetExpression->SetShouldExport( NewState == ECheckBoxState::Checked);
 			}
 		}
 	});
