Convert Virtual Texture to non virtual texture in texture graph.

diff --git a/Engine/Plugins/Experimental/TextureGraph/Source/TextureGraphEngine/2D/Tex.cpp b/Engine/Plugins/Experimental/TextureGraph/Source/TextureGraphEngine/2D/Tex.cpp
--- a/Engine/Plugins/Experimental/TextureGraph/Source/TextureGraphEngine/2D/Tex.cpp
+++ b/Engine/Plugins/Experimental/TextureGraph/Source/TextureGraphEngine/2D/Tex.cpp
@@ -1263,6 +1263,16 @@ bool Tex::LoadAsset(FSoftObjectPath& SoftPath, const DesiredImageProperties* Pro
 	check(Obj);
 
 	Texture = static_cast<UTexture2D*>(Obj);
+
+	//For now in case of virtual texture we are making a duplicate of texture and turning off the VirtualTextureStreaming
+	//so that it can be loaded like a normal texture 2D. This is a workaround until we find a way to reliably convert the virtual texture to render target
+	if (Texture->IsCurrentlyVirtualTextured())
+	{
+		Texture = (UTexture2D*)StaticDuplicateObject(Obj, GetTransientPackage(), NAME_None, RF_Transient, UTexture2D::StaticClass());
+		Texture->Modify();
+		Texture->VirtualTextureStreaming = false;
+	}
+
 	Desc = TexDescriptor(Texture);
 
 	// override descriptor based on source properties
