Fixed issue: The Editor crashes with a thumbnail-related error after using the Subtract node

diff --git a/Engine/Plugins/Experimental/TextureGraph/Source/TextureGraphEditor/Private/EdGraph/STG_NodeThumbnail.cpp b/Engine/Plugins/Experimental/TextureGraph/Source/TextureGraphEditor/Private/EdGraph/STG_NodeThumbnail.cpp
--- a/Engine/Plugins/Experimental/TextureGraph/Source/TextureGraphEditor/Private/EdGraph/STG_NodeThumbnail.cpp
+++ b/Engine/Plugins/Experimental/TextureGraph/Source/TextureGraphEditor/Private/EdGraph/STG_NodeThumbnail.cpp
@@ -54,8 +54,9 @@ UTexture* STG_NodeThumbnail::GetTextureFromBlob(TiledBlobPtr InBlob)
 	{
 		check(InBlob->IsFinalised());
 		DeviceBufferPtr Buffer = InBlob->GetBufferRef().GetPtr();
-		check(Buffer);
 
+		checkSlow(Buffer);
+		
 		auto FXBuffer = std::static_pointer_cast<DeviceBuffer_FX>(Buffer);
 		// Got the description, now need to retreive the true live buffer and the texture
 		{
diff --git a/Engine/Plugins/Experimental/TextureGraph/Source/TextureGraphEditor/Private/STG_NodePreview.cpp b/Engine/Plugins/Experimental/TextureGraph/Source/TextureGraphEditor/Private/STG_NodePreview.cpp
--- a/Engine/Plugins/Experimental/TextureGraph/Source/TextureGraphEditor/Private/STG_NodePreview.cpp
+++ b/Engine/Plugins/Experimental/TextureGraph/Source/TextureGraphEditor/Private/STG_NodePreview.cpp
@@ -313,53 +313,43 @@ void STG_NodePreviewWidget::NodeDeleted(const UTG_Node* Node)
 
 void STG_NodePreviewWidget::Update() const
 {
-	const UTG_Node* PreviewNode = LockedNode ? LockedNode : SelectedNode;
-	BlobPtr Blob;
-	
-	// Get the Variant if preview node is valid or assign a nullptr if not.
-	const FTG_Variant* Variant = nullptr;
-	TArray<FTG_Variant> OutVariants;
-	TArray<FName>* OutNames = nullptr;
-
-	if(PreviewNode)
-	{
-		PreviewNode->GetAllOutputValues(OutVariants, OutNames);
-
-		if (!OutVariants.IsEmpty())
-		{
-			Variant = &OutVariants[0];
-		}	
-	}
-
 	// Lambda for actually updating the preview.
-	auto UpdatePreview = [this, Variant]
+	auto UpdatePreview = [this]
 	{
-		if(Variant)
+		FTG_Variant Variant;
+		GetOutputVariantFromNode(Variant);
+
+		if(Variant.IsTexture())
 		{
-			if(Variant->IsTexture())
-			{
-				NodeViewer->SetTexture(Variant->GetTexture().RasterBlob);
-			}
-			else if(Variant->IsColor())
+			if(Variant.GetTexture() && Variant.GetTexture()->IsValid())
 			{
-				NodeViewer->SetTexture(nullptr, Variant->GetColor());
+				NodeViewer->SetTexture(Variant.GetTexture().RasterBlob);
 			}
 			else
 			{
-				NodeViewer->SetTexture(nullptr);	
+				NodeViewer->SetTexture(nullptr);
 			}
 		}
+		else if(Variant.IsColor())
+		{
+			NodeViewer->SetTexture(nullptr, Variant.GetColor());
+		}
 		else
 		{
-			NodeViewer->SetTexture(nullptr);
+			NodeViewer->SetTexture(nullptr);	
 		}
 		
 		Viewport->ResetZoom(NodeViewer->GetCurrentImageInfo().Size);
 	};
 
-	if(Variant && Variant->IsTexture())
+	FTG_Variant Variant;
+
+	GetOutputVariantFromNode(Variant);
+	BlobPtr Blob;
+	
+	if(Variant && Variant.IsTexture() && Variant.GetTexture())
 	{
-		Blob = Variant->GetTexture().RasterBlob;
+		Blob = Variant.GetTexture().RasterBlob;
 
 		if(Blob)
 		{
@@ -393,6 +383,28 @@ void STG_NodePreviewWidget::Update() const
 	[[maybe_unused]] bool Result = OnNodeBlobChanged.ExecuteIfBound(Blob);
 }
 
+bool STG_NodePreviewWidget::GetOutputVariantFromNode(FTG_Variant& OutVariant) const
+{
+	const UTG_Node* PreviewNode = LockedNode ? LockedNode : SelectedNode;
+	
+	// Get the Variant if preview node is valid or assign a nullptr if not.
+	TArray<FTG_Variant> OutVariants;
+	TArray<FName>* OutNames = nullptr;
+
+	if(PreviewNode)
+	{
+		PreviewNode->GetAllOutputValues(OutVariants, OutNames);
+
+		if (!OutVariants.IsEmpty())
+		{
+			OutVariant = OutVariants[0];
+			return true;
+		}	
+	}
+
+	return false;
+}
+
 FReply STG_NodePreviewWidget::OnKeyDown(const FGeometry& InGeometry, const FKeyEvent& InKeyEvent)
 {
 	if (CommandList.IsValid() && CommandList->ProcessCommandBindings(InKeyEvent))
diff --git a/Engine/Plugins/Experimental/TextureGraph/Source/TextureGraphEditor/Private/STG_NodePreview.h b/Engine/Plugins/Experimental/TextureGraph/Source/TextureGraphEditor/Private/STG_NodePreview.h
--- a/Engine/Plugins/Experimental/TextureGraph/Source/TextureGraphEditor/Private/STG_NodePreview.h
+++ b/Engine/Plugins/Experimental/TextureGraph/Source/TextureGraphEditor/Private/STG_NodePreview.h
@@ -7,6 +7,8 @@
 #include "Data/Blob.h"
 #include "Widgets/SCompoundWidget.h"
 
+struct FTG_Variant;
+
 namespace UE::ImageWidgets
 {
 	class SImageViewport;
@@ -112,6 +114,8 @@ public:
 	/** Trigger an update of the node preview after the displayed contents changed. */
 	void Update() const;
 
+	bool GetOutputVariantFromNode(FTG_Variant& OutVariant) const; 
+
 	// SWidget overrides - begin
 	virtual FReply OnKeyDown(const FGeometry& InGeometry, const FKeyEvent& InKeyEvent) override;
 	// SWidget overrides - end
