Fix for Crash when using TexturePath node and adding a path for jpg image

diff --git a/Engine/Plugins/Experimental/TextureGraph/Source/TextureGraphEngine/2D/Tex.cpp b/Engine/Plugins/Experimental/TextureGraph/Source/TextureGraphEngine/2D/Tex.cpp
--- a/Engine/Plugins/Experimental/TextureGraph/Source/TextureGraphEngine/2D/Tex.cpp
+++ b/Engine/Plugins/Experimental/TextureGraph/Source/TextureGraphEngine/2D/Tex.cpp
@@ -767,7 +767,7 @@ bool UncompressJpeg(const ERGBFormat Format, int32 BitDepth, int32 Width, int32
 	check(compressedData.Num());
 
 	UncompressedData.Reset(Width * Height * NumChannels);
-	UncompressedData.AddUninitialized(Width * Width * NumChannels);
+	UncompressedData.AddUninitialized(Width * Height * NumChannels);
 
 	int Flag = TJFLAG_PROGRESSIVE;
 	if (tjDecompress2(Decompressor, compressedData.GetData(), compressedData.Num(), UncompressedData.GetData(), Width, 0, Height, TJPixelFormat, Flag) != 0)
