Job Batch getting stuck

diff --git a/Engine/Plugins/Experimental/TextureGraph/Source/TextureGraphEngine/Export/T_ExportAsUAsset.cpp b/Engine/Plugins/Experimental/TextureGraph/Source/TextureGraphEngine/Export/T_ExportAsUAsset.cpp
--- a/Engine/Plugins/Experimental/TextureGraph/Source/TextureGraphEngine/Export/T_ExportAsUAsset.cpp
+++ b/Engine/Plugins/Experimental/TextureGraph/Source/TextureGraphEngine/Export/T_ExportAsUAsset.cpp
@@ -45,6 +45,7 @@ cti::continuable<int32> Job_ExportAsUAsset::PreExecAsync(ENamedThreads::Type exe
 	{
 		bIsCulled = true;
 		bIsDone = true;
+		MarkJobDone();
 
 		return cti::make_ready_continuable(0);
 	}
diff --git a/Engine/Plugins/Experimental/TextureGraph/Source/TextureGraphEngine/Transform/Utility/T_LoadStaticResource.cpp b/Engine/Plugins/Experimental/TextureGraph/Source/TextureGraphEngine/Transform/Utility/T_LoadStaticResource.cpp
--- a/Engine/Plugins/Experimental/TextureGraph/Source/TextureGraphEngine/Transform/Utility/T_LoadStaticResource.cpp
+++ b/Engine/Plugins/Experimental/TextureGraph/Source/TextureGraphEngine/Transform/Utility/T_LoadStaticResource.cpp
@@ -33,6 +33,7 @@ cti::continuable<int32> Job_LoadStaticImageResource::PreExecAsync(ENamedThreads:
 
 		bIsCulled = true;
 		bIsDone = true;
+		MarkJobDone();
 
 		return cti::make_ready_continuable(0);
 	}
