Fixed issue: Texture Graph Editor preview viewport doesn't always load HDRI map.

We should use the default values of the scene settings.

diff --git a/Engine/Plugins/Experimental/TextureGraph/Source/TextureGraphEditor/Private/STG_EditorViewport.cpp b/Engine/Plugins/Experimental/TextureGraph/Source/TextureGraphEditor/Private/STG_EditorViewport.cpp
--- a/Engine/Plugins/Experimental/TextureGraph/Source/TextureGraphEditor/Private/STG_EditorViewport.cpp
+++ b/Engine/Plugins/Experimental/TextureGraph/Source/TextureGraphEditor/Private/STG_EditorViewport.cpp
@@ -226,16 +226,9 @@ void FTG_EditorViewportClient::FocusViewportOnBounds(const FBoxSphereBounds Boun
 
 void STG_EditorViewport::Construct(const FArguments& InArgs)
 {
-	//_previewScene = new FAdvancedPreviewScene(FPreviewScene::ConstructionValues(), 
-
 	TextureGraphPtr = InArgs._InTextureGraph;
 
 	PreviewScene = MakeShareable(new FAdvancedPreviewScene(FPreviewScene::ConstructionValues()));
-	/*if (StaticMesh)
-	{
-		_previewScene->SetFloorOffset(-StaticMesh->GetExtendedBounds().Origin.Z + StaticMesh->GetExtendedBounds().BoxExtent.Z);
-	}*/
-	PreviewScene->SetFloorVisibility(false, false);
 
 	// restore last used feature level
 	UWorld* World = PreviewScene->GetWorld();
@@ -256,23 +249,24 @@ void STG_EditorViewport::Construct(const FArguments& InArgs)
 
 	SEditorViewport::Construct(SEditorViewport::FArguments());
 
-	/*_previewMeshComponent = NewObject<UStaticMeshComponent>(GetTransientPackage(), NAME_None, RF_Transient);*/
-	
 	PreviewMeshComponent = nullptr;
 	
-	// InitPreviewMesh();
-
 	ERHIFeatureLevel::Type FeatureLevel = GEditor->PreviewPlatform.GetEffectivePreviewFeatureLevel();
 	if (FeatureLevel <= ERHIFeatureLevel::ES3_1)
 	{
 		PreviewMeshComponent->SetMobility(EComponentMobility::Static);
 	}
 
+	UAssetViewerSettings* Settings = UAssetViewerSettings::Get();
+	const int32 ProfileIndex = PreviewScene->GetCurrentProfileIndex();
+	if (Settings->Profiles.IsValidIndex(ProfileIndex))
+	{
+		PreviewScene->SetEnvironmentVisibility(Settings->Profiles[ProfileIndex].bShowEnvironment, true);
+	}
+
 	FCoreUObjectDelegates::OnObjectPropertyChanged.AddRaw(this, &STG_EditorViewport::OnObjectPropertyChanged);
 	
 	GenerateRendermodeToolbar();
-	
-	TogglePreviewBackground();
 }
 
 STG_EditorViewport::STG_EditorViewport()
