Reconciled offline work

diff --git a/Engine/Plugins/Experimental/TextureGraph/Source/TextureGraph/Private/Expressions/TG_Expression_MaterialBase.cpp b/Engine/Plugins/Experimental/TextureGraph/Source/TextureGraph/Private/Expressions/TG_Expression_MaterialBase.cpp
--- a/Engine/Plugins/Experimental/TextureGraph/Source/TextureGraph/Private/Expressions/TG_Expression_MaterialBase.cpp
+++ b/Engine/Plugins/Experimental/TextureGraph/Source/TextureGraph/Private/Expressions/TG_Expression_MaterialBase.cpp
@@ -142,8 +142,7 @@ void UTG_Expression_MaterialBase::Evaluate(FTG_EvaluationContext* InContext)
 	if (GetMaterial() && MaterialInstance)
 	{
 		FString AssetName = GetMaterial()->GetName();
-
-		const auto RenderMaterial = std::make_shared<RenderMaterial_BP>(AssetName, MaterialInstance->GetMaterial(), MaterialInstance);
+		const auto RenderMaterial = std::make_shared<RenderMaterial_BP>(AssetName, GetMaterial(), nullptr);
 		Result = CreateRenderMaterialJob(InContext, RenderMaterial, Output.GetBufferDescriptor(), GetRenderedAttributeId());
 	}
 
@@ -199,7 +198,7 @@ TiledBlobPtr UTG_Expression_MaterialBase::CreateRenderMaterialJob(FTG_Evaluation
 
 	Desc.DefaultValue = FLinearColor::Black;
 
-	LinkMaterialParameters(InContext, MaterialJob, InRenderMaterial->GetMaterial(), Desc);
+	LinkMaterialParameters(InContext, MaterialJob, GetMaterial(), Desc);
 
 	const TiledBlob_PromisePtr MaterialResult = std::static_pointer_cast<TiledBlob_Promise>(MaterialJob->InitResult(InRenderMaterial->GetName(), &Desc));
 	MaterialJob->AddArg(WithUnbounded(ARG_BOOL(TiledMode, "TiledMode")));
@@ -220,7 +219,7 @@ TiledBlobPtr UTG_Expression_MaterialBase::CreateRenderMaterialJob(FTG_Evaluation
 
 }
 
-void UTG_Expression_MaterialBase::LinkMaterialParameters(FTG_EvaluationContext* InContext, JobUPtr& InMaterialJob, const UMaterial* InMaterial, BufferDescriptor InDescriptor)
+void UTG_Expression_MaterialBase::LinkMaterialParameters(FTG_EvaluationContext* InContext, JobUPtr& InMaterialJob, const UMaterialInterface* InMaterial, BufferDescriptor InDescriptor)
 {
 	if (InMaterialJob)
 	{	
diff --git a/Engine/Plugins/Experimental/TextureGraph/Source/TextureGraph/Public/Expressions/TG_Expression_MaterialBase.h b/Engine/Plugins/Experimental/TextureGraph/Source/TextureGraph/Public/Expressions/TG_Expression_MaterialBase.h
--- a/Engine/Plugins/Experimental/TextureGraph/Source/TextureGraph/Public/Expressions/TG_Expression_MaterialBase.h
+++ b/Engine/Plugins/Experimental/TextureGraph/Source/TextureGraph/Public/Expressions/TG_Expression_MaterialBase.h
@@ -74,7 +74,7 @@ protected:
 
 	TiledBlobPtr CreateRenderMaterialJob(FTG_EvaluationContext* InContext, const FString& InName, const FString& InMaterialPath, const BufferDescriptor& InDescriptor, EDrawMaterialAttributeTarget InDrawMaterialAttributeTarget);
 	TiledBlobPtr CreateRenderMaterialJob(FTG_EvaluationContext* InContext, const RenderMaterial_BPPtr& InRenderMaterial, const BufferDescriptor& InDescriptor, EDrawMaterialAttributeTarget InDrawMaterialAttributeTarget);
-	void LinkMaterialParameters(FTG_EvaluationContext* InContext, JobUPtr& InMaterialJob, const UMaterial* InMaterial, BufferDescriptor InDescriptor);
+	void LinkMaterialParameters(FTG_EvaluationContext* InContext, JobUPtr& InMaterialJob, const UMaterialInterface* InMaterial, BufferDescriptor InDescriptor);
 
 	virtual void CopyVarGeneric(const FTG_Argument& Arg, FTG_Var* InVar, bool CopyVarToArg);
 	virtual void SetMaterialInternal(UMaterialInterface* InMaterial);
diff --git a/Engine/Plugins/Experimental/TextureGraph/Source/TextureGraphEngine/FxMat/RenderMaterial_BP.cpp b/Engine/Plugins/Experimental/TextureGraph/Source/TextureGraphEngine/FxMat/RenderMaterial_BP.cpp
--- a/Engine/Plugins/Experimental/TextureGraph/Source/TextureGraphEngine/FxMat/RenderMaterial_BP.cpp
+++ b/Engine/Plugins/Experimental/TextureGraph/Source/TextureGraphEngine/FxMat/RenderMaterial_BP.cpp
@@ -21,7 +21,7 @@ bool RenderMaterial_BP::ValidateMaterialCompatible(UMaterialInterface* InMateria
 	return FxMaterial_QuadDrawMaterial::ValidateMaterial(InMaterial);
 }
 
-RenderMaterial_BP::RenderMaterial_BP(FString Name, UMaterial* InMaterial, UMaterialInstanceDynamic* InMaterialInstance)
+RenderMaterial_BP::RenderMaterial_BP(FString Name, UMaterialInterface* InMaterial, UMaterialInstanceDynamic* InMaterialInstance)
 	: RenderMaterial(!Name.IsEmpty() ? Name : InMaterial->GetName())
 	, Material(InMaterial)
 	, MaterialInstance(TStrongObjectPtr<UMaterialInstanceDynamic>(InMaterialInstance ? InMaterialInstance : UMaterialInstanceDynamic::Create(InMaterial, nullptr)))
diff --git a/Engine/Plugins/Experimental/TextureGraph/Source/TextureGraphEngine/FxMat/RenderMaterial_BP.h b/Engine/Plugins/Experimental/TextureGraph/Source/TextureGraphEngine/FxMat/RenderMaterial_BP.h
--- a/Engine/Plugins/Experimental/TextureGraph/Source/TextureGraphEngine/FxMat/RenderMaterial_BP.h
+++ b/Engine/Plugins/Experimental/TextureGraph/Source/TextureGraphEngine/FxMat/RenderMaterial_BP.h
@@ -12,8 +12,9 @@ typedef std::shared_ptr< FxMaterial> FxMaterialPtr;
 class TEXTUREGRAPHENGINE_API RenderMaterial_BP : public RenderMaterial
 {
 protected:
-	UMaterial*						Material = nullptr;			/// The base material that is part 
-	TStrongObjectPtr<UMaterialInstanceDynamic> 
+
+	UMaterialInterface*				Material = nullptr;			/// The base material that is used for this job 
+	TStrongObjectPtr<UMaterialInstanceDynamic>
 									MaterialInstance;			/// An instance of the material
 	CHashPtr						HashValue;					/// The hash for this material
 	bool							RequestMaterialValidation = true;		/// 
@@ -27,7 +28,7 @@ protected:
 									FVector2D CoordinatePosition, FVector2D CoordinateSize=FVector2D::UnitVector, float Rotation=0.f,
 									FVector2D PivotPoint=FVector2D(0.5f,0.5f)) const;
 public:
-									RenderMaterial_BP(FString Name, UMaterial* InMaterial, UMaterialInstanceDynamic* InMaterialInstance = nullptr);
+									RenderMaterial_BP(FString Name, UMaterialInterface* InMaterial, UMaterialInstanceDynamic* InMaterialInstance = nullptr);
 
 	virtual							~RenderMaterial_BP() override;
 
@@ -63,7 +64,7 @@ public:
 	//////////////////////////////////////////////////////////////////////////
 	/// Inline functions
 	////////////////////////////////////////////////////////////////////////// 
-	FORCEINLINE UMaterial*			GetMaterial() { return Material; }
+	FORCEINLINE UMaterialInterface*			GetMaterial() { return Material; }
 
 	FORCEINLINE UMaterialInstanceDynamic* Instance() { return MaterialInstance.Get(); }
 	FORCEINLINE const UMaterialInstanceDynamic* Instance() const { return MaterialInstance.Get(); }
diff --git a/Engine/Plugins/Experimental/TextureGraph/Source/TextureGraphEngine/FxMat/RenderMaterial_BP_NoTile.h b/Engine/Plugins/Experimental/TextureGraph/Source/TextureGraphEngine/FxMat/RenderMaterial_BP_NoTile.h
--- a/Engine/Plugins/Experimental/TextureGraph/Source/TextureGraphEngine/FxMat/RenderMaterial_BP_NoTile.h
+++ b/Engine/Plugins/Experimental/TextureGraph/Source/TextureGraphEngine/FxMat/RenderMaterial_BP_NoTile.h
@@ -11,7 +11,7 @@ class TEXTUREGRAPHENGINE_API RenderMaterial_BP_NoTile : public RenderMaterial_BP
 {
 
 public:
-	RenderMaterial_BP_NoTile(FString InName, UMaterial* InMaterial, UMaterialInstanceDynamic* InMaterialInstance = nullptr) : RenderMaterial_BP(InName, InMaterial, InMaterialInstance) {}
+	RenderMaterial_BP_NoTile(FString InName, UMaterialInterface* InMaterial, UMaterialInstanceDynamic* InMaterialInstance = nullptr) : RenderMaterial_BP(InName, InMaterial, InMaterialInstance) {}
 	virtual bool					CanHandleTiles() const override { return false; };
 };
 
@@ -23,7 +23,7 @@ private:
 	FName							MemberName = "TileInfo";
 
 public:
-	RenderMaterial_BP_TileArgs(FString InName, UMaterial* InMaterial, UMaterialInstanceDynamic* InMaterialInstance = nullptr) : RenderMaterial_BP(InName, InMaterial, InMaterialInstance) {}
+	RenderMaterial_BP_TileArgs(FString InName, UMaterialInterface* InMaterial, UMaterialInstanceDynamic* InMaterialInstance = nullptr) : RenderMaterial_BP(InName, InMaterial, InMaterialInstance) {}
 	void							AddTileArgs(TransformArgs& Args);
 	virtual AsyncPrepareResult		PrepareResources(const TransformArgs& Args) override;
 	virtual bool					CanHandleTiles() const override { return true; };
diff --git a/Engine/Plugins/Experimental/TextureGraph/Source/TextureGraphEngine/FxMat/RenderMaterial_Thumbnail.cpp b/Engine/Plugins/Experimental/TextureGraph/Source/TextureGraphEngine/FxMat/RenderMaterial_Thumbnail.cpp
--- a/Engine/Plugins/Experimental/TextureGraph/Source/TextureGraphEngine/FxMat/RenderMaterial_Thumbnail.cpp
+++ b/Engine/Plugins/Experimental/TextureGraph/Source/TextureGraphEngine/FxMat/RenderMaterial_Thumbnail.cpp
@@ -14,7 +14,7 @@
 
 ////////////////////////////////////////////////////////////////////////// RenderMaterial_Thumbnail //////////////////////////////////////////////////////////////////////////
 
-RenderMaterial_Thumbnail::RenderMaterial_Thumbnail(FString InName, UMaterial* InMaterial, UMaterialInstanceDynamic* instance /*= nullptr*/) : RenderMaterial_BP(InName, InMaterial, instance)
+RenderMaterial_Thumbnail::RenderMaterial_Thumbnail(FString InName, UMaterialInterface* InMaterial, UMaterialInstanceDynamic* instance /*= nullptr*/) : RenderMaterial_BP(InName, InMaterial, instance)
 {
 }
 
diff --git a/Engine/Plugins/Experimental/TextureGraph/Source/TextureGraphEngine/FxMat/RenderMaterial_Thumbnail.h b/Engine/Plugins/Experimental/TextureGraph/Source/TextureGraphEngine/FxMat/RenderMaterial_Thumbnail.h
--- a/Engine/Plugins/Experimental/TextureGraph/Source/TextureGraphEngine/FxMat/RenderMaterial_Thumbnail.h
+++ b/Engine/Plugins/Experimental/TextureGraph/Source/TextureGraphEngine/FxMat/RenderMaterial_Thumbnail.h
@@ -17,7 +17,7 @@ private:
 	ERHIFeatureLevel::Type FeatureLevel;
 
 public:
-	RenderMaterial_Thumbnail(FString InName, UMaterial* InMaterial, UMaterialInstanceDynamic* instance = nullptr);
+	RenderMaterial_Thumbnail(FString InName, UMaterialInterface* InMaterial, UMaterialInstanceDynamic* instance = nullptr);
 
 	virtual void BlitTo(FRHICommandListImmediate& RHI, UTextureRenderTarget2D* DstRT, const RenderMesh* MeshObj, int32 TargetId) const override;
 	virtual bool CanHandleTiles() const override { return true; };
