Fix for Undoing/redoing converting node to Parameter/Constant doesn't properly handle the parameter in the Parameters tab

diff --git a/Engine/Plugins/Experimental/TextureGraph/Source/TextureGraph/Private/TG_Graph.cpp b/Engine/Plugins/Experimental/TextureGraph/Source/TextureGraph/Private/TG_Graph.cpp
--- a/Engine/Plugins/Experimental/TextureGraph/Source/TextureGraph/Private/TG_Graph.cpp
+++ b/Engine/Plugins/Experimental/TextureGraph/Source/TextureGraph/Private/TG_Graph.cpp
@@ -154,6 +154,7 @@ void UTG_Graph::RegenerateNode(UTG_Node* InNode)
 	// Call in the UObject Modify to snapshot the current state for do/undo feature
 	// Done BEFORE doing any change to the node
 	InNode->Modify();
+	Modify();
 	
 	FTG_SignaturePtr OldSignature = InNode->Signature;
 	// If the old signature is null, rebuild one from the current PIn arguments
