Texture Graph fix for stretched pins in node UI

diff --git a/Engine/Plugins/Experimental/TextureGraph/Source/TextureGraphEditor/Private/EdGraph/STG_EditorGraphNode.cpp b/Engine/Plugins/Experimental/TextureGraph/Source/TextureGraphEditor/Private/EdGraph/STG_EditorGraphNode.cpp
--- a/Engine/Plugins/Experimental/TextureGraph/Source/TextureGraphEditor/Private/EdGraph/STG_EditorGraphNode.cpp
+++ b/Engine/Plugins/Experimental/TextureGraph/Source/TextureGraphEditor/Private/EdGraph/STG_EditorGraphNode.cpp
@@ -93,6 +93,7 @@ void STG_EditorGraphNode::UpdateGraphNode()
 	RightNodeBox.Reset();
 	LeftNodeBox.Reset();
 	ParametersBox.Reset();
+	LowDetailParametersBox.Reset();
 	NodeSettingsBox.Reset();
 	NotConnectableParametersBox.Reset();
 	NotConnectableInputPinBox.Reset();
@@ -390,7 +391,9 @@ int STG_EditorGraphNode::GetContentAreaBorderThickness()
 
 TSharedRef<SWidget> STG_EditorGraphNode::CreateNodeContentArea()
 {
-	int HAdvanceContentPading = GetContentAreaBorderThickness();
+	const int HAdvanceContentPading = GetContentAreaBorderThickness();
+	const int LowDetailOverrideSettingWidth = 75;
+	const int LowDetailOverrideSettingHeight = 66;
 	// NODE CONTENT AREA
 	return SNew(SBorder)
 		.BorderImage(FAppStyle::GetBrush("NoBorder"))
@@ -438,16 +441,35 @@ TSharedRef<SWidget> STG_EditorGraphNode::CreateNodeContentArea()
 			.AutoHeight()
 			.Padding(HAdvanceContentPading, 0)
 			[
-				SAssignNew(ParametersBox, SVerticalBox)
-				.Visibility(this,&STG_EditorGraphNode::ShowParameters) //visible when when node have advance display
+				SNew(SLevelOfDetailBranchNode)
+				.UseLowDetailSlot(this, &STG_EditorGraphNode::UseLowDetailParameters)
+				.LowDetail()
+				[
+					SAssignNew(LowDetailParametersBox, SVerticalBox)
+					.Visibility(this, &STG_EditorGraphNode::ShowParameters) //visible when when node have advance display
 
-				+ SVerticalBox::Slot()
-				.AutoHeight()
-				.HAlign(HAlign_Fill)
-				.VAlign(VAlign_Center)
+					+ SVerticalBox::Slot()
+					.AutoHeight()
+					.HAlign(HAlign_Left)
+					.VAlign(VAlign_Center)
+					[
+						SNew(SSeparator)
+						.Thickness(2)
+					]
+				]
+				.HighDetail()
 				[
-					SNew(SSeparator)
-					.Thickness(2)
+					SAssignNew(ParametersBox, SVerticalBox)
+					.Visibility(this, &STG_EditorGraphNode::ShowParameters) //visible when when node have advance display
+
+					+ SVerticalBox::Slot()
+					.AutoHeight()
+					.HAlign(HAlign_Fill)
+					.VAlign(VAlign_Center)
+					[
+						SNew(SSeparator)
+						.Thickness(2)
+					]
 				]
 			]
 
@@ -463,8 +485,21 @@ TSharedRef<SWidget> STG_EditorGraphNode::CreateNodeContentArea()
 			.AutoHeight()
 			.Padding(HAdvanceContentPading, 0)
 			[
-				SAssignNew(NodeSettingsBox, SVerticalBox)
-				.Visibility(this, &STG_EditorGraphNode::ShowOverrideSettings)// Visible only when node is expanded
+				SNew(SLevelOfDetailBranchNode)
+				.UseLowDetailSlot(this, &STG_EditorGraphNode::UseLowDetailParameters)
+				.LowDetail()
+				[
+					//Create empty space when showing low lod
+					SNew(SBox)
+					.Visibility(this, &STG_EditorGraphNode::ShowOverrideSettings)
+					.WidthOverride(LowDetailOverrideSettingWidth)
+					.HeightOverride(LowDetailOverrideSettingHeight)
+				]
+				.HighDetail()
+				[
+					SAssignNew(NodeSettingsBox, SVerticalBox)
+					.Visibility(this, &STG_EditorGraphNode::ShowOverrideSettings)// Visible only when node is expanded
+				]
 			]
 		];
 }
@@ -542,6 +577,16 @@ bool STG_EditorGraphNode::UseLowDetailNodeTitles() const
 	return SGraphNode::UseLowDetailNodeTitles();
 }
 
+bool STG_EditorGraphNode::UseLowDetailParameters() const
+{
+	if (const SGraphPanel* MyOwnerPanel = GetOwnerPanel().Get())
+	{
+		return (MyOwnerPanel->GetCurrentLOD() <= EGraphRenderingLOD::LowDetail);
+	}
+
+	return false;
+}
+
 FSlateColor STG_EditorGraphNode::GetNodeTitleColor() const
 {
 	FLinearColor ReturnTitleColor = GraphNode->IsDeprecated() ? FLinearColor::Red : GetNodeObj()->GetNodeTitleColor();
@@ -731,6 +776,15 @@ void STG_EditorGraphNode::AddPin(const TSharedRef<SGraphPin>& PinToAdd)
 				[
 					PinToAdd
 				];
+
+				LowDetailParametersBox->AddSlot()
+				.AutoHeight()
+				.HAlign(HAlign_Left)
+				.VAlign(VAlign_Center)
+				.Padding(Settings->GetInputPinPadding())
+				[
+					PinToAdd
+				];
 			}
 			else
 			{
diff --git a/Engine/Plugins/Experimental/TextureGraph/Source/TextureGraphEditor/Private/EdGraph/STG_EditorGraphNode.h b/Engine/Plugins/Experimental/TextureGraph/Source/TextureGraphEditor/Private/EdGraph/STG_EditorGraphNode.h
--- a/Engine/Plugins/Experimental/TextureGraph/Source/TextureGraphEditor/Private/EdGraph/STG_EditorGraphNode.h
+++ b/Engine/Plugins/Experimental/TextureGraph/Source/TextureGraphEditor/Private/EdGraph/STG_EditorGraphNode.h
@@ -60,6 +60,8 @@ protected:
 	TSharedPtr<SVerticalBox> NotConnectableInputPinBox;
 	/** The area where Param pins and settings pins reside*/
 	TSharedPtr<SVerticalBox> ParametersBox;
+	/** The area where Param pins and settings pins reside for Low detail*/
+	TSharedPtr<SVerticalBox>  LowDetailParametersBox;
 	/** The area where non connectable Param pins and settings pins reside*/
 	TSharedPtr<SVerticalBox> NotConnectableParametersBox;
 	/** The area where node specific settings pins reside alyaws show in advance view */
@@ -102,4 +104,5 @@ private:
 	TArray<FString> GetTitleDetailTextLines() const;
 	FText GetLeftTitleText() const;
 	FText GetRightTitleText() const;
+	bool UseLowDetailParameters() const;
 };
