Fixing crash related to the FString vars used over fName

diff --git a/Engine/Plugins/Experimental/TextureGraph/Source/TextureGraph/Private/TG_Pin.cpp b/Engine/Plugins/Experimental/TextureGraph/Source/TextureGraph/Private/TG_Pin.cpp
--- a/Engine/Plugins/Experimental/TextureGraph/Source/TextureGraph/Private/TG_Pin.cpp
+++ b/Engine/Plugins/Experimental/TextureGraph/Source/TextureGraph/Private/TG_Pin.cpp
@@ -635,19 +635,6 @@ bool UTG_Pin::SetValue(FString Value)
 	return false;
 }
 
-bool UTG_Pin::SetValue(FName Value)
-{
-	if (IsArgString())
-	{
-		Modify();
-		EditSelfVar()->EditAs<FString>() = Value.GetPlainNameString();
-		NotifyPinSelfVarChanged();
-
-		return true;
-	}
-	return false;
-}
-
 bool UTG_Pin::SetValue(const FLinearColor& Value)
 {
 	if (IsArgColor())
diff --git a/Engine/Plugins/Experimental/TextureGraph/Source/TextureGraph/Private/TG_Signature.cpp b/Engine/Plugins/Experimental/TextureGraph/Source/TextureGraph/Private/TG_Signature.cpp
--- a/Engine/Plugins/Experimental/TextureGraph/Source/TextureGraph/Private/TG_Signature.cpp
+++ b/Engine/Plugins/Experimental/TextureGraph/Source/TextureGraph/Private/TG_Signature.cpp
@@ -21,7 +21,6 @@ bool FTG_Argument::IsBool() const
 bool FTG_Argument::IsString() const
 {
 	return CPPTypeName == TEXT("FString") || 
-		CPPTypeName == TEXT("FName") || 
 		CPPTypeName == TEXT("FText");
 }
 
diff --git a/Engine/Plugins/Experimental/TextureGraph/Source/TextureGraph/Public/TG_Pin.h b/Engine/Plugins/Experimental/TextureGraph/Source/TextureGraph/Public/TG_Pin.h
--- a/Engine/Plugins/Experimental/TextureGraph/Source/TextureGraph/Public/TG_Pin.h
+++ b/Engine/Plugins/Experimental/TextureGraph/Source/TextureGraph/Public/TG_Pin.h
@@ -183,7 +183,7 @@ public:
 	// bool getter only on a boolean argument
 	bool GetValue(bool& OutValue) const;
 
-	// bool getter only on a boolean argument
+	// String getter only on a FString argument
 	bool GetValue(FString& OutValue) const;
 
 	// Color or Vector getter only on a vector or color argument
@@ -213,16 +213,17 @@ public:
 	// This is the same as the selfvar setter
 	void				FromString(const FString& InValueStr, bool bIsTweaking = false);
 
-	// Setters for Scalar, Color, Vector arguments
+	// Setters for Scalar, Color, Vector, bool and String  arguments
+	// These check that the arg is of the correct type
 	bool SetValue(float Value);
 	bool SetValue(double Value) { return SetValue((float)Value); }
 	bool SetValue(uint32 Value) { return SetValue((float)Value); }
 	bool SetValue(int32 Value) { return SetValue((float)Value); }
 	bool SetValue(bool bValue);
-	bool SetValue(FString Value);
-	bool SetValue(FName Value);
+	bool SetValue(FString Value); // FName pin is assigned by the templated version, not FString
 	bool SetValue(const FLinearColor& Value);
 	bool SetValue(const FVector4f& Value);
+
 	bool SetValue(const float* Value, size_t Count);
 	bool SetValue(FTG_Texture&) { return false; } // this is not a possible case
 	bool SetValue(const FTG_TextureDescriptor& Value); 
