Fixed issues:
- 3D Preview Settings not updating after deleting second Output node in Texture Graph
- Duplicating a disconnected output node results in a copy with the same name

diff --git a/Engine/Plugins/Experimental/TextureGraph/Source/TextureGraph/Private/TG_Graph.cpp b/Engine/Plugins/Experimental/TextureGraph/Source/TextureGraph/Private/TG_Graph.cpp
--- a/Engine/Plugins/Experimental/TextureGraph/Source/TextureGraph/Private/TG_Graph.cpp
+++ b/Engine/Plugins/Experimental/TextureGraph/Source/TextureGraph/Private/TG_Graph.cpp
@@ -349,8 +349,12 @@ void UTG_Graph::KillPin(FTG_Id InPinId)
 		}
 	}
 
-	// Remove the pin form the param just in case
-	Params.Remove(Pin->GetAliasName());
+	// Only update for Param Pin's. 
+	if(Pin->IsParam())
+	{
+		// Remove the pin form the param just in case
+		Params.Remove(Pin->GetAliasName());
+	}
 }
 
 void UTG_Graph::RemoveNode(UTG_Node* InNode)
diff --git a/Engine/Plugins/Experimental/TextureGraph/Source/TextureGraphEditor/Private/TG_Editor.cpp b/Engine/Plugins/Experimental/TextureGraph/Source/TextureGraphEditor/Private/TG_Editor.cpp
--- a/Engine/Plugins/Experimental/TextureGraph/Source/TextureGraphEditor/Private/TG_Editor.cpp
+++ b/Engine/Plugins/Experimental/TextureGraph/Source/TextureGraphEditor/Private/TG_Editor.cpp
@@ -1851,7 +1851,7 @@ void FTG_Editor::OnNodeRemoved(UTG_Node* InNode, FName Title)
 				const FName FirstOutputName;
 				EditedTextureGraph->Graph()->ForEachParams([&](const UTG_Pin* Pin, uint32 index)
 					{
-						if(Pin->IsOutput() && (Pin->IsArgTexture() || Pin->IsArgColor()))
+						if(Pin->IsOutput())
 						{
 							if (FirstOutputName.IsNone())
 							{
