TG_Expression_Dot should ignore alpha in dot calculation

diff --git a/Engine/Plugins/Experimental/TextureGraph/Shaders/Expressions/Expression_Maths_TwoInputs.usf b/Engine/Plugins/Experimental/TextureGraph/Shaders/Expressions/Expression_Maths_TwoInputs.usf
--- a/Engine/Plugins/Experimental/TextureGraph/Shaders/Expressions/Expression_Maths_TwoInputs.usf
+++ b/Engine/Plugins/Experimental/TextureGraph/Shaders/Expressions/Expression_Maths_TwoInputs.usf
@@ -35,10 +35,10 @@ float4 FSH_Subtract(float2 uv : TEXCOORD0) : SV_Target0
 
 float4 FSH_Dot(float2 uv : TEXCOORD0) : SV_Target0
 {
-	float4 op1 = Operand1.Sample(SamplerStates_Wrap, uv);
-	float4 op2 = Operand2.Sample(SamplerStates_Wrap, uv);
+	float3 op1 = Operand1.Sample(SamplerStates_Wrap, uv).rgb;
+	float3 op2 = Operand2.Sample(SamplerStates_Wrap, uv).rgb;
 	float dp = dot(op1, op2);
-	return float4(dp, dp, dp, dp);
+	return float4(dp, dp, dp, 1);
 }
 
 float4 FSH_Cross(float2 uv : TEXCOORD0) : SV_Target0
diff --git a/Engine/Plugins/Experimental/TextureGraph/Source/TextureGraph/Private/Expressions/Maths/TG_Expression_Maths_TwoInputs.cpp b/Engine/Plugins/Experimental/TextureGraph/Source/TextureGraph/Private/Expressions/Maths/TG_Expression_Maths_TwoInputs.cpp
--- a/Engine/Plugins/Experimental/TextureGraph/Source/TextureGraph/Private/Expressions/Maths/TG_Expression_Maths_TwoInputs.cpp
+++ b/Engine/Plugins/Experimental/TextureGraph/Source/TextureGraph/Private/Expressions/Maths/TG_Expression_Maths_TwoInputs.cpp
@@ -75,7 +75,7 @@ FVector4f UTG_Expression_Dot::EvaluateVector_WithValue(FTG_EvaluationContext* In
 {
 	check(Value && Count == 2);
 	float Result = FVector3f::DotProduct(Value[0], Value[1]);
-	return FVector4f{ Result, Result, Result, Result };
+	return FVector4f{ Result, Result, Result, 1 };
 }
 
 FTG_Texture	UTG_Expression_Dot::EvaluateTexture(FTG_EvaluationContext* InContext)
