Shipping and packaging fixes

diff --git a/Engine/Plugins/Experimental/TextureGraph/Source/TextureGraph/TextureGraph.Build.cs b/Engine/Plugins/Experimental/TextureGraph/Source/TextureGraph/TextureGraph.Build.cs
--- a/Engine/Plugins/Experimental/TextureGraph/Source/TextureGraph/TextureGraph.Build.cs
+++ b/Engine/Plugins/Experimental/TextureGraph/Source/TextureGraph/TextureGraph.Build.cs
@@ -7,7 +7,13 @@ public class TextureGraph : ModuleRules
 	public TextureGraph(ReadOnlyTargetRules Target) : base(Target)
 	{
 		PCHUsage = ModuleRules.PCHUsageMode.UseExplicitOrSharedPCHs;
-		
+		bUseRTTI = true;
+		bEnableExceptions = true;
+		PublicIncludePathModuleNames.AddRange(
+			new string[]
+			{
+				"TextureGraphEngine"
+			});
 		PublicIncludePaths.AddRange(
 			new string[] {
 				// ... add public include paths required here ...
