File Name input in Texture Graph Editor missing validation for name length

diff --git a/Engine/Plugins/Experimental/TextureGraph/Source/TextureGraph/Private/Expressions/Input/TG_Expression_OutputSettings.cpp b/Engine/Plugins/Experimental/TextureGraph/Source/TextureGraph/Private/Expressions/Input/TG_Expression_OutputSettings.cpp
--- a/Engine/Plugins/Experimental/TextureGraph/Source/TextureGraph/Private/Expressions/Input/TG_Expression_OutputSettings.cpp
+++ b/Engine/Plugins/Experimental/TextureGraph/Source/TextureGraph/Private/Expressions/Input/TG_Expression_OutputSettings.cpp
@@ -63,27 +63,15 @@ void UTG_Expression_OutputSettings::Evaluate(FTG_EvaluationContext* InContext)
 
 bool UTG_Expression_OutputSettings::Validate(MixUpdateCyclePtr Cycle)
 {
-	FText Reason;
-	FTG_HelperFunctions::IsFileNameValid(Settings.BaseName, Reason);
-	
-	auto InputPin = GetParentNode()->GetPin("Input");
-	if (!InputPin->IsConnected() && !Reason.IsEmpty())
-	{
-		FString Message = "Invalid file name. " + Reason.ToString();
-		UMixInterface* ParentMix = Cast<UMixInterface>(GetOutermostObject());
-		auto ErrorType = static_cast<int32>(ETextureGraphErrorType::NODE_WARNING);
-		TextureGraphEngine::GetErrorReporter(ParentMix)->ReportWarning(ErrorType, Message, GetParentNode());
-	}
+	FString Errors;
 
-	FText PathReason;
-	FTG_HelperFunctions::IsFolderPathValid(Settings.FolderPath.ToString(), PathReason);
+	const UTG_Pin* OutputSettingsPin = GetParentNode()->GetPin("Settings");
 
-	if (!InputPin->IsConnected() && !PathReason.IsEmpty())
+	if(!OutputSettingsPin->IsConnected() && !Settings.Validate(Errors))
 	{
-		FString Message = "Invalid folder path. " + PathReason.ToString();
 		UMixInterface* ParentMix = Cast<UMixInterface>(GetOutermostObject());
 		auto ErrorType = static_cast<int32>(ETextureGraphErrorType::NODE_WARNING);
-		TextureGraphEngine::GetErrorReporter(ParentMix)->ReportWarning(ErrorType, Message, GetParentNode());
+		TextureGraphEngine::GetErrorReporter(ParentMix)->ReportWarning(ErrorType, Errors, GetParentNode());
 	}
 
 	return Super::Validate(Cycle);
diff --git a/Engine/Plugins/Experimental/TextureGraph/Source/TextureGraph/Private/Expressions/Output/TG_Expression_Output.cpp b/Engine/Plugins/Experimental/TextureGraph/Source/TextureGraph/Private/Expressions/Output/TG_Expression_Output.cpp
--- a/Engine/Plugins/Experimental/TextureGraph/Source/TextureGraph/Private/Expressions/Output/TG_Expression_Output.cpp
+++ b/Engine/Plugins/Experimental/TextureGraph/Source/TextureGraph/Private/Expressions/Output/TG_Expression_Output.cpp
@@ -97,26 +97,15 @@ void UTG_Expression_Output::Evaluate(FTG_EvaluationContext* InContext)
 
 bool UTG_Expression_Output::Validate(MixUpdateCyclePtr	Cycle)
 {
-	FText Reason;
-	FTG_HelperFunctions::IsFileNameValid(OutputSettings.BaseName, Reason);
-	auto OutputSettingsPin = GetParentNode()->GetPin("OutputSettings");
-	if (!OutputSettingsPin->IsConnected() && !Reason.IsEmpty())
-	{
-		FString Message = "Invalid file name. " + Reason.ToString();
-		UMixInterface* ParentMix = Cast<UMixInterface>(GetOutermostObject());
-		auto ErrorType = static_cast<int32>(ETextureGraphErrorType::NODE_WARNING);
-		TextureGraphEngine::GetErrorReporter(ParentMix)->ReportWarning(ErrorType, Message, GetParentNode());
-	}
-
-	FText PathReason;
-	FTG_HelperFunctions::IsFolderPathValid(OutputSettings.FolderPath.ToString(), PathReason);
+	FString Errors;
+	
+	const UTG_Pin* OutputSettingsPin = GetParentNode()->GetPin("OutputSettings");
 
-	if (!OutputSettingsPin->IsConnected() && !PathReason.IsEmpty())
+	if(!OutputSettingsPin->IsConnected() && !OutputSettings.Validate(Errors))
 	{
-		FString Message = "Invalid folder path. " + PathReason.ToString();
 		UMixInterface* ParentMix = Cast<UMixInterface>(GetOutermostObject());
 		auto ErrorType = static_cast<int32>(ETextureGraphErrorType::NODE_WARNING);
-		TextureGraphEngine::GetErrorReporter(ParentMix)->ReportWarning(ErrorType, Message, GetParentNode());
+		TextureGraphEngine::GetErrorReporter(ParentMix)->ReportWarning(ErrorType, Errors, GetParentNode());
 	}
 
 	return Super::Validate(Cycle);
diff --git a/Engine/Plugins/Experimental/TextureGraph/Source/TextureGraph/Private/TG_HelperFunctions.cpp b/Engine/Plugins/Experimental/TextureGraph/Source/TextureGraph/Private/TG_HelperFunctions.cpp
--- a/Engine/Plugins/Experimental/TextureGraph/Source/TextureGraph/Private/TG_HelperFunctions.cpp
+++ b/Engine/Plugins/Experimental/TextureGraph/Source/TextureGraph/Private/TG_HelperFunctions.cpp
@@ -10,32 +10,10 @@
 #include "Expressions/Output/TG_Expression_Output.h"
 #include "Expressions/TG_Expression.h"
 #include "Model/Mix/MixManager.h"
-#include "Expressions/TG_Expression.h"
 #include "Transform/Expressions/T_FlatColorTexture.h"
 #include "Logging/MessageLog.h"
-#include "Misc/UObjectToken.h"
 #include "Misc/FileHelper.h"
-
-bool FTG_HelperFunctions::IsFileNameValid(FName FileName, FText& Reason)
-{
-	FileName.IsValidObjectName(Reason);
-	FFileHelper::IsFilenameValidForSaving(FileName.ToString(), Reason);
-
-	return Reason.IsEmpty();
-}
-
-bool FTG_HelperFunctions::IsFolderPathValid(FString FolderPath, FText& Reason)
-{
-	FName::IsValidXName(FolderPath, INVALID_OBJECTPATH_CHARACTERS INVALID_LONGPACKAGE_CHARACTERS, &Reason);
-	
-	if (!Reason.IsEmpty())
-	{
-		Reason = FText::FromString(Reason.ToString().Replace(TEXT("Name"), TEXT("Path")));
-	}
-
-	FPaths::ValidatePath(FolderPath, &Reason);
-	return Reason.IsEmpty();
-}
+#include "TextureGraph.h"
 
 TArray<BlobPtr> FTG_HelperFunctions::GetTexturedOutputs(const UTG_Node* Node, FTG_EvaluationContext* TextureConversionContext /*= nullptr*/)
 
@@ -137,7 +115,7 @@ JobBatchPtr FTG_HelperFunctions::InitExportBatch(UTextureGraph* InTextureGraph,
 			UTG_Expression_Output* TargetExpression = Cast<UTG_Expression_Output>(Node->GetExpression());
 			if (TargetExpression)
 			{
-				FTG_OutputSettings& OutputSetting = TargetExpression->OutputSettings;
+				FTG_OutputSettings& OutputSettings = TargetExpression->OutputSettings;
 
 				if (TargetExpression->GetShouldExport() || ExportAllOutputs)
 				{
@@ -146,25 +124,21 @@ JobBatchPtr FTG_HelperFunctions::InitExportBatch(UTextureGraph* InTextureGraph,
 
 					auto ExportBlobs = FTG_HelperFunctions::GetTexturedOutputs(Node, &EvaluationContext);
 
-					FString Path = OutputSetting.FolderPath.ToString();
+					FString Path = OutputSettings.FolderPath.ToString();
 					if (OverrideExportPath)
 					{
 						Path = ExportPath;
 					}
 
-					FString FileName = AssetName.IsEmpty() ? OutputSetting.GetFullOutputName() : AssetName;
-
-					FText Reason;
-					bool IsNameValid = IsFileNameValid(FName(*FileName), Reason);
-
-					FText PathReason;
-					bool IsPathValid = TextureExporter::IsPackageNameValid(Path, FileName);
-					IsFolderPathValid(Path, PathReason);
-					IsPathValid &= PathReason.IsEmpty();
+					FString FileName = AssetName.IsEmpty() ? OutputSettings.GetFullOutputName() : AssetName;
 
+					FString PathErrors;
+					bool IsPathValid = TextureExporter::IsFilePathValid(FName(*FileName), FName(*Path), PathErrors);
+					bool IsPackageValid = TextureExporter::IsPackageNameValid(Path, FileName);
+					
 					bool HasOutputs = ExportBlobs.Num() > 0;
 
-					if (HasOutputs && IsNameValid && IsPathValid)
+					if (HasOutputs && IsPathValid && IsPackageValid)
 					{
 						TiledBlobPtr Output = std::static_pointer_cast<TiledBlob>(ExportBlobs[0]);//Dealing with one output per Node for now
 						FExportMapSettings MapSettings = TextureExporter::GetExportSettingsForTarget(TargetExportSettings, std::static_pointer_cast<TiledBlob>(Output), *FileName);
@@ -172,11 +146,11 @@ JobBatchPtr FTG_HelperFunctions::InitExportBatch(UTextureGraph* InTextureGraph,
 						MapSettings.Path = Path;
 						MapSettings.UseOverridePath = OverrideExportPath;
 						MapSettings.OverwriteTextures = OverwriteTextures;
-						MapSettings.LODGroup = OutputSetting.LODGroup;
-						MapSettings.Compression = OutputSetting.Compression;
-						MapSettings.IsSRGB = OutputSetting.bSRGB;
-						MapSettings.Width = (int32)OutputSetting.Width;
-						MapSettings.Height = (int32)OutputSetting.Height;
+						MapSettings.LODGroup = OutputSettings.LODGroup;
+						MapSettings.Compression = OutputSettings.Compression;
+						MapSettings.IsSRGB = OutputSettings.bSRGB;
+						MapSettings.Width = (int32)OutputSettings.Width;
+						MapSettings.Height = (int32)OutputSettings.Height;
 						MapSettings.bSave = bSave;
 						TargetExportSettings.ExportPreset.push_back(std::pair<FName, FExportMapSettings>{ MapSettings.Name, MapSettings });
 						AnyValidExport = true;
@@ -186,21 +160,18 @@ JobBatchPtr FTG_HelperFunctions::InitExportBatch(UTextureGraph* InTextureGraph,
 						//Log Error to Error System
 						if (!HasOutputs)
 						{
-							ErrorMessage += FString::Format(TEXT("Texture Export Error : No valid output found for OutputSetting {0}"), { OutputSetting.OutputName.ToString() });
+							ErrorMessage += FString::Format(TEXT("Texture Export Error : No valid output found for OutputSetting {0}"), { OutputSettings.OutputName.ToString() });
 							ErrorMessage += "\n";
 						}
 						if (!IsPathValid)
 						{
-							FString InvalidFolderPath = "Invalid folder path. ";
-							FString PathReasonString = PathReason.IsEmpty() ? InvalidFolderPath : InvalidFolderPath + PathReason.ToString();
-							ErrorMessage += FString::Format(TEXT("Texture Export Error : {0} for OutputSetting {1}"), { PathReasonString, OutputSetting.OutputName.ToString() });
+							ErrorMessage += FString::Format(TEXT("Texture Export Error : {0} OutputSettings Node: {1}"), { PathErrors, OutputSettings.OutputName.ToString() });
 							ErrorMessage += "\n";
 						}
-						if (!IsNameValid)
+						if (!IsPackageValid)
 						{
-							FString InvalidFileName = "Invalid File Name. ";
-							FString ReasonString = Reason.IsEmpty() ? InvalidFileName : InvalidFileName + Reason.ToString();
-							ErrorMessage += FString::Format(TEXT("Texture Export Error : {0} for OutputSetting {1}"), { ReasonString , OutputSetting.OutputName.ToString()});
+							FString Error = "Invalid Package name";
+							ErrorMessage += FString::Format(TEXT("Texture Export Error : {0} OutputSettings Node: {1}"), { Error, OutputSettings.OutputName.ToString() });
 							ErrorMessage += "\n";
 						}
 					}
diff --git a/Engine/Plugins/Experimental/TextureGraph/Source/TextureGraph/Public/TG_HelperFunctions.h b/Engine/Plugins/Experimental/TextureGraph/Source/TextureGraph/Public/TG_HelperFunctions.h
--- a/Engine/Plugins/Experimental/TextureGraph/Source/TextureGraph/Public/TG_HelperFunctions.h
+++ b/Engine/Plugins/Experimental/TextureGraph/Source/TextureGraph/Public/TG_HelperFunctions.h
@@ -4,20 +4,18 @@
 
 #include "CoreMinimal.h"
 #include "Job/JobBatch.h"
-#include "TextureGraph.h"
 #include "Export/TextureExporter.h"
 #include "Data/Blob.h"
 #include "TG_Graph.h"
 #include "TG_Node.h"
 #include "2D/TextureHelper.h"
 
+class UTextureGraph;
+
 class TEXTUREGRAPH_API FTG_HelperFunctions
 {
 
 public:
-	static bool IsFileNameValid(FName FileName, FText& Reason);
-	static bool IsFolderPathValid(FString FolderPath, FText& Reason);
-
 	static void	InitTargets(UTextureGraph* InTextureGraph);
 	static TArray<BlobPtr> GetTexturedOutputs(const UTG_Node* Node, FTG_EvaluationContext* TextureConversionContext = nullptr);
 
diff --git a/Engine/Plugins/Experimental/TextureGraph/Source/TextureGraph/Public/TG_OutputSettings.h b/Engine/Plugins/Experimental/TextureGraph/Source/TextureGraph/Public/TG_OutputSettings.h
--- a/Engine/Plugins/Experimental/TextureGraph/Source/TextureGraph/Public/TG_OutputSettings.h
+++ b/Engine/Plugins/Experimental/TextureGraph/Source/TextureGraph/Public/TG_OutputSettings.h
@@ -5,6 +5,7 @@
 #include "TG_SystemTypes.h"
 #include "Model/Mix/MixSettings.h"
 #include "Misc/OutputDeviceNull.h"
+#include "TG_HelperFunctions.h"
 #include "TG_OutputSettings.generated.h"
 
 class UTG_Expression_Output;
@@ -40,48 +41,48 @@ public:
 };
 
 USTRUCT(BlueprintType)
-struct TEXTUREGRAPH_API FTG_OutputSettings 
+struct TEXTUREGRAPH_API FTG_OutputSettings
 {
 	GENERATED_USTRUCT_BODY()
 
 	// Export name of the textured asset.
 	UPROPERTY(BlueprintReadWrite,EditAnywhere, Category = "Basic", DisplayName = "File Name", Meta = (NoResetToDefault))
-		FName BaseName;
+	FName BaseName;
 
 	UPROPERTY()
-		FName OutputName;
+	FName OutputName;
 
 	// Export path for the textured asset.
 	UPROPERTY(BlueprintReadWrite,EditAnywhere, Category = "Basic", DisplayName = "Path", Meta = (NoResetToDefault))
-		FName FolderPath;
+	FName FolderPath;
 
 	// Width of the texture in pixels. Auto means system will detect automatically based on other images
 	UPROPERTY(BlueprintReadWrite,EditAnywhere, Category = "Advanced", Meta = (NoResetToDefault))
-		EResolution Width = EResolution::Auto;
+	EResolution Width = EResolution::Auto;
 
 	// Height of the texture in pixels. Auto means system will detect automatically based on other images
 	UPROPERTY(BlueprintReadWrite,EditAnywhere, Category = "Advanced", Meta = (NoResetToDefault))
-		EResolution Height = EResolution::Auto;
+	EResolution Height = EResolution::Auto;
 
 	// List of available texture formats. Auto means system will detect automatically based on the input
 	UPROPERTY(BlueprintReadWrite,EditAnywhere, Category = "Advanced", DisplayName = "Texture Format", Meta = (NoResetToDefault))
-		ETG_TextureFormat TextureFormat = ETG_TextureFormat::BGRA8;
+	ETG_TextureFormat TextureFormat = ETG_TextureFormat::BGRA8;
 
 	// List of available texture presets available for export. 
 	UPROPERTY(BlueprintReadWrite, EditAnywhere, Category = "Advanced", DisplayName = "Texture Type", Meta = (NoResetToDefault))
-		ETG_TexturePresetType TexturePresetType = ETG_TexturePresetType::None;
+	ETG_TexturePresetType TexturePresetType = ETG_TexturePresetType::None;
 
 	// The Level of detail group of the texture
 	UPROPERTY(BlueprintReadWrite, EditAnywhere, Category = "Advanced", DisplayName = "LOD Texture Group", Meta = (NoResetToDefault, EditCondition = "TexturePresetType == ETG_TexturePresetType::None"))
-		TEnumAsByte<enum TextureGroup> LODGroup = TextureGroup::TEXTUREGROUP_World;
+	TEnumAsByte<enum TextureGroup> LODGroup = TextureGroup::TEXTUREGROUP_World;
 
 	// Compression methods available for exporting textured asset.
 	UPROPERTY(BlueprintReadWrite, EditAnywhere, Category = "Advanced", DisplayName = "Compression", Meta = (NoResetToDefault, EditCondition = "TexturePresetType == ETG_TexturePresetType::None") )
-		TEnumAsByte <enum TextureCompressionSettings> Compression = TextureCompressionSettings::TC_Default;
+	TEnumAsByte <enum TextureCompressionSettings> Compression = TextureCompressionSettings::TC_Default;
 
 	// Adjust the color space of exporting textured asset. Can be in Linear or Gamma color space. 
 	UPROPERTY(BlueprintReadWrite, EditAnywhere, Category = "Advanced", DisplayName = "sRGB", Meta = (NoResetToDefault, EditCondition = "TexturePresetType == ETG_TexturePresetType::None"))
-		bool bSRGB = false;
+	bool bSRGB = false;
 
 	FString GetFullOutputName() { return  FString::Format(TEXT("{0}"), { BaseName.ToString()});}
 
@@ -109,4 +110,9 @@ struct TEXTUREGRAPH_API FTG_OutputSettings
 		TextureCompressionSettings InCompression = TextureCompressionSettings::TC_Default, TextureGroup InLodGroup = TextureGroup::TEXTUREGROUP_World, bool InbSRGB = false);
 
 	void OnSetTexturePresetType(ETG_TexturePresetType Type);
+
+	bool Validate(FString& OutErrors)
+	{
+		return TextureExporter::IsFilePathValid(BaseName, FolderPath, OutErrors);
+	}
 };
diff --git a/Engine/Plugins/Experimental/TextureGraph/Source/TextureGraphEngine/Export/TextureExporter.cpp b/Engine/Plugins/Experimental/TextureGraph/Source/TextureGraphEngine/Export/TextureExporter.cpp
--- a/Engine/Plugins/Experimental/TextureGraph/Source/TextureGraphEngine/Export/TextureExporter.cpp
+++ b/Engine/Plugins/Experimental/TextureGraph/Source/TextureGraphEngine/Export/TextureExporter.cpp
@@ -19,14 +19,73 @@
 #if WITH_EDITOR
 #include "FileHelpers.h"
 #include "AssetToolsModule.h"
+#include "AssetViewUtils.h"
 #endif
-#include <Misc/PackageName.h> 
+#include <Misc/PackageName.h>
+#include "Misc/FileHelper.h"
 #include "TextureCompiler.h"
+
 DEFINE_LOG_CATEGORY(ExportLogs);
 bool TextureExporter::IsPackageNameValid(FString Path, FString AssetName)
 {
 	return FPackageName::TryConvertLongPackageNameToFilename(Path, AssetName);
 }
+
+bool TextureExporter::IsFileNameValid(FName FileName, FText& Reason)
+{
+	FileName.IsValidObjectName(Reason);
+	FFileHelper::IsFilenameValidForSaving(FileName.ToString(), Reason);
+
+	return Reason.IsEmpty();
+}
+
+bool TextureExporter::IsFolderPathValid(FString FolderPath, FText& Reason)
+{
+	FName::IsValidXName(FolderPath, INVALID_OBJECTPATH_CHARACTERS INVALID_LONGPACKAGE_CHARACTERS, &Reason);
+	
+	if (!Reason.IsEmpty())
+	{
+		Reason = FText::FromString(Reason.ToString().Replace(TEXT("Name"), TEXT("Path")));
+	}
+
+	FPaths::ValidatePath(FolderPath, &Reason);
+	return Reason.IsEmpty();
+}
+
+bool TextureExporter::IsFilePathValid(const FName InFileName, const FName InFolderPath, FString& OutErrors)
+{
+	FText FileNameError;
+	IsFileNameValid(InFileName, FileNameError);
+		
+	if(!FileNameError.IsEmpty())
+	{
+		FString Message = "Invalid file name. " + FileNameError.ToString();
+		OutErrors += Message;
+	}
+		
+	FText FullPathError;
+	const FString AssetName = FString::Printf(TEXT("/%s.%s"), *InFileName.ToString(), *InFileName.ToString());
+	const FString FullPath = FPaths::Combine(InFolderPath.ToString(), AssetName);
+
+#if WITH_EDITOR
+	if(!AssetViewUtils::IsValidObjectPathForCreate(FullPath, FullPathError, true))
+	{
+		FString Message = "Invalid file path. " + FullPathError.ToString();
+		OutErrors += Message;
+	}
+#endif
+	
+	FText PathValidError;
+	IsFolderPathValid(InFolderPath.ToString(), PathValidError);
+		
+	if(!PathValidError.IsEmpty())
+	{
+		FString Message = "Invalid folder path Error: " + PathValidError.ToString();
+		OutErrors += Message;
+	}
+
+	return OutErrors.IsEmpty();
+}
 AsyncBool TextureExporter::ExportAllMapsAsUAsset(UMixInterface* MixObj, const FString& OutputFolder, const FExportSettings& ExportSettings)
 {
 	FInvalidationDetails Details;
diff --git a/Engine/Plugins/Experimental/TextureGraph/Source/TextureGraphEngine/Export/TextureExporter.h b/Engine/Plugins/Experimental/TextureGraph/Source/TextureGraphEngine/Export/TextureExporter.h
--- a/Engine/Plugins/Experimental/TextureGraph/Source/TextureGraphEngine/Export/TextureExporter.h
+++ b/Engine/Plugins/Experimental/TextureGraph/Source/TextureGraphEngine/Export/TextureExporter.h
@@ -100,4 +100,7 @@ struct TEXTUREGRAPHENGINE_API TextureExporter
 	static void						ExportAsUAsset(UMixInterface* MixObj, FName AssetName, FString CurrentPath, const FExportSettings& Setting);
 	static AsyncInt					ExportRawAsUAsset(RawBufferPtr RawObj,const FExportMapSettings& Setting, const FString& CompletePath, const FName& TextureName);
 	static bool						IsPackageNameValid(FString Path, FString AssetName);
+	static bool 					IsFileNameValid(FName FileName, FText& Reason);
+	static bool 					IsFolderPathValid(FString FolderPath, FText& Reason);
+	static bool 					IsFilePathValid(const FName InFileName, const FName InFolderPath, FString& OutErrors);
 };
diff --git a/Engine/Plugins/Experimental/TextureGraph/Source/TextureGraphEngine/TextureGraphEngine.Build.cs b/Engine/Plugins/Experimental/TextureGraph/Source/TextureGraphEngine/TextureGraphEngine.Build.cs
--- a/Engine/Plugins/Experimental/TextureGraph/Source/TextureGraphEngine/TextureGraphEngine.Build.cs
+++ b/Engine/Plugins/Experimental/TextureGraph/Source/TextureGraphEngine/TextureGraphEngine.Build.cs
@@ -77,6 +77,7 @@ public class TextureGraphEngine : ModuleRules
 			PublicDependencyModuleNames.AddRange(new string[]
             {
                 "UnrealEd",
+                "AssetTools"
 			});
 		}
 
