Fix CIS non-unity issue due to removal of TextureGraph.h from TG_HelperFunctions.h

diff --git a/Engine/Plugins/Experimental/TextureGraph/Source/TextureGraphEditor/Private/STG_OutputSelectionDlg.cpp b/Engine/Plugins/Experimental/TextureGraph/Source/TextureGraphEditor/Private/STG_OutputSelectionDlg.cpp
--- a/Engine/Plugins/Experimental/TextureGraph/Source/TextureGraphEditor/Private/STG_OutputSelectionDlg.cpp
+++ b/Engine/Plugins/Experimental/TextureGraph/Source/TextureGraphEditor/Private/STG_OutputSelectionDlg.cpp
@@ -10,6 +10,7 @@
 #include "Editor.h"
 #include "EDGraph/STG_NodeThumbnail.h"
 #include "STG_OutputSelector.h"
+#include "TextureGraph.h"
 #include "TG_HelperFunctions.h"
 #include "EdGraph/TG_EdGraphNode.h"
 #include "Widgets/Colors/SColorBlock.h"
