For Material Node, we now detect when the referenced Material is modified and saved in another editor window
the referenced Material is duplicated to be used in the context of the MaterialNode of TG during evaluation.

diff --git a/Engine/Plugins/Experimental/TextureGraph/Source/TextureGraph/Private/Expressions/Input/TG_Expression_Material.cpp b/Engine/Plugins/Experimental/TextureGraph/Source/TextureGraph/Private/Expressions/Input/TG_Expression_Material.cpp
--- a/Engine/Plugins/Experimental/TextureGraph/Source/TextureGraph/Private/Expressions/Input/TG_Expression_Material.cpp
+++ b/Engine/Plugins/Experimental/TextureGraph/Source/TextureGraph/Private/Expressions/Input/TG_Expression_Material.cpp
@@ -1,6 +1,8 @@
 // Copyright Epic Games, Inc. All Rights Reserved.
 
 #include "Expressions/Input/TG_Expression_Material.h"
+#include "UObject/ObjectSaveContext.h"
+#include "Misc/TransactionObjectEvent.h"
 
 #if WITH_EDITOR
 void UTG_Expression_Material::PostEditChangeProperty(FPropertyChangedEvent& PropertyChangedEvent)
@@ -32,10 +34,42 @@ void UTG_Expression_Material::PostEditUndo()
 	Super::PostEditUndo();
 }
 
+void UTG_Expression_Material::OnReferencedObjectPreSave(UObject* Object, FObjectPreSaveContext SaveContext)
+{
+	// every editor should check if your texture graph is dependent on the object being saved
+	UMaterialInterface* MaterialBeingSaved = Cast<UMaterialInterface>(Object);
+	// if object being saved is our linked TextureGraph, we re-create
+	if (MaterialBeingSaved && Material && (MaterialBeingSaved == Material))
+	{
+
+		SetMaterialInternal(MaterialBeingSaved);
+	}
+}
+
+#endif
+
+UTG_Expression_Material::UTG_Expression_Material()
+{
+#if WITH_EDITOR
+	// listener for UObject saves, so we can synchronise when linked TextureGraphs get updated 
+	PreSaveHandle = FCoreUObjectDelegates::OnObjectPreSave.AddUObject(this, &UTG_Expression_Material::OnReferencedObjectPreSave);
+#endif
+}
+UTG_Expression_Material::~UTG_Expression_Material()
+{
+#if WITH_EDITOR
+	FCoreUObjectDelegates::OnObjectPreSave.Remove(PreSaveHandle);
 #endif
+}
 
 void UTG_Expression_Material::Initialize()
 {
+	// If the referenced material is valid, then we need to recreate a dubplicate
+	if (Material)
+	{
+		MaterialCopy = DuplicateObject(Material, this);
+	}
+
 	Super::Initialize();
 	SetRenderedAttribute(RenderedAttribute); // reassign the RenderedAttribute to make sure it is the correct one
 }
@@ -45,18 +79,21 @@ void UTG_Expression_Material::SetMaterialInternal(UMaterialInterface* InMaterial
 	if (!InMaterial)
 	{
 		Material = nullptr;
+		MaterialCopy = nullptr;
 	}
 	else if (InMaterial->IsA<UMaterial>())
 	{
 		// create a material instance
 		Material = InMaterial;
+		MaterialCopy = DuplicateObject(Material, this);
 	}
 	else if (InMaterial->IsA<UMaterialInstance>())
 	{
 		Material = InMaterial;
+		MaterialCopy = DuplicateObject(Material, this);
 	}
 
-	Super::SetMaterialInternal(InMaterial);
+	Super::SetMaterialInternal(MaterialCopy);
 
 	SetRenderedAttribute(RenderedAttribute);
 }
@@ -72,7 +109,7 @@ void UTG_Expression_Material::SetMaterial(UMaterialInterface* InMaterial)
 	if (InMaterial == Material)
 	{
 		// Just check that the MaterialInstance is valid, if not reassign below
-		if ((!Material) || (Material && MaterialInstance))
+		if ((!Material) || (Material && MaterialCopy && MaterialInstance))
 			return;
 	}
 
diff --git a/Engine/Plugins/Experimental/TextureGraph/Source/TextureGraph/Private/Expressions/TG_Expression_MaterialBase.cpp b/Engine/Plugins/Experimental/TextureGraph/Source/TextureGraph/Private/Expressions/TG_Expression_MaterialBase.cpp
--- a/Engine/Plugins/Experimental/TextureGraph/Source/TextureGraph/Private/Expressions/TG_Expression_MaterialBase.cpp
+++ b/Engine/Plugins/Experimental/TextureGraph/Source/TextureGraph/Private/Expressions/TG_Expression_MaterialBase.cpp
@@ -107,9 +107,9 @@ void UTG_Expression_MaterialBase::GenerateMaterialAttributeOptions()
 	// Detect the set of available material properties for rendering
 	AvailableMaterialAttributeIds.Empty();
 	AvailableMaterialAttributeNames.Empty();
-	if (MaterialInstance)
+	if (GetMaterial())
 	{
-		UMaterial* RefMaterial = MaterialInstance->GetMaterial();
+		UMaterial* RefMaterial = GetMaterial()->GetMaterial();
 		for (int i = MP_EmissiveColor; i < MP_MAX; ++i)
 		{
 			if (RefMaterial->IsPropertyConnected(EMaterialProperty(i)))
@@ -139,7 +139,7 @@ void UTG_Expression_MaterialBase::Evaluate(FTG_EvaluationContext* InContext)
 	/// Set it to false always
 	TiledMode = true;
 
-	if (GetMaterial() && MaterialInstance)
+	if (GetMaterial())
 	{
 		FString AssetName = GetMaterial()->GetName();
 		const auto RenderMaterial = std::make_shared<RenderMaterial_BP>(AssetName, GetMaterial(), nullptr);
@@ -303,6 +303,7 @@ void UTG_Expression_MaterialBase::LinkMaterialParameters(FTG_EvaluationContext*
 			{
 				auto ParamValue = Var->EditAs<bool>();
 				InMaterialJob->AddArg(ARG_INT(ParamValue, TCHAR_TO_UTF8(*ParameterInfo.Name.ToString())));
+
 			}
 		}
 #endif
@@ -449,5 +450,4 @@ void UTG_Expression_MaterialBase::CopyVarGeneric(const FTG_Argument& Arg, FTG_Va
 				}*/
 		}
 	}
-	
 }
\ No newline at end of file
diff --git a/Engine/Plugins/Experimental/TextureGraph/Source/TextureGraph/Public/Expressions/Input/TG_Expression_Material.h b/Engine/Plugins/Experimental/TextureGraph/Source/TextureGraph/Public/Expressions/Input/TG_Expression_Material.h
--- a/Engine/Plugins/Experimental/TextureGraph/Source/TextureGraph/Public/Expressions/Input/TG_Expression_Material.h
+++ b/Engine/Plugins/Experimental/TextureGraph/Source/TextureGraph/Public/Expressions/Input/TG_Expression_Material.h
@@ -17,6 +17,8 @@ class TEXTUREGRAPH_API UTG_Expression_Material : public UTG_Expression_MaterialB
 {
 	GENERATED_BODY()
 public:
+	UTG_Expression_Material();
+	virtual ~UTG_Expression_Material();
 #if WITH_EDITOR
 	virtual void PostEditChangeProperty(FPropertyChangedEvent& PropertyChangedEvent) override;
 	virtual void PostEditUndo() override;
@@ -43,6 +45,9 @@ public:
 	virtual FText GetTooltipText() const override { return FText::FromString(TEXT("Renders a material into a quad and makes it available. It is automatically exposed as a graph input parameter.")); } 
 
 protected:
+	UPROPERTY(Transient)
+	TObjectPtr<UMaterialInterface> MaterialCopy;
+
 	virtual void SetMaterialInternal(UMaterialInterface* InMaterial) override;
 
 	UPROPERTY(EditAnywhere, BlueprintReadWrite, Category = "")
@@ -58,7 +63,13 @@ public:
 	virtual FName GetCategory() const override { return TG_Category::Input;}
 	
 protected:
-	virtual TObjectPtr<UMaterialInterface> GetMaterial() const override { return Material;};
+	virtual TObjectPtr<UMaterialInterface> GetMaterial() const override { return MaterialCopy;};
 	virtual EDrawMaterialAttributeTarget GetRenderedAttributeId()  override;
+
+#if WITH_EDITOR // Listener for referenced material being saved to update the integration in TG 
+	FDelegateHandle PreSaveHandle; 
+	void OnReferencedObjectPreSave(UObject* Object, FObjectPreSaveContext SaveContext);
+#endif
+
 };
 
