Texture Path do not clear message

diff --git a/Engine/Plugins/Experimental/TextureGraph/Source/TextureGraph/Private/Expressions/Input/TG_Expression_TexturePath.cpp b/Engine/Plugins/Experimental/TextureGraph/Source/TextureGraph/Private/Expressions/Input/TG_Expression_TexturePath.cpp
--- a/Engine/Plugins/Experimental/TextureGraph/Source/TextureGraph/Private/Expressions/Input/TG_Expression_TexturePath.cpp
+++ b/Engine/Plugins/Experimental/TextureGraph/Source/TextureGraph/Private/Expressions/Input/TG_Expression_TexturePath.cpp
@@ -61,7 +61,9 @@ void UTG_Expression_TexturePath::Evaluate(FTG_EvaluationContext* InContext)
 	Super::Evaluate(InContext);
 
 	FString LocalPath;
-	if (!Path.IsEmpty() && ValidateInputPath(LocalPath))
+	bool bValidation = ValidateInputPath(LocalPath);
+
+	if (!Path.IsEmpty() && bValidation)
 	{
 		UStaticImageResource* StaticImageResource = UStaticImageResource::CreateNew<UStaticImageResource>();
 		StaticImageResource->SetAssetUUID(LocalPath);
@@ -79,21 +81,36 @@ void UTG_Expression_TexturePath::Evaluate(FTG_EvaluationContext* InContext)
 	{
 		Output = FTG_Texture::GetBlack();
 	}
+
+	//For the connected pin we will report error here in eveluate because it does not have the updated value during validation.
+	UTG_Pin* PathPin = GetParentNode()->GetPin(GET_MEMBER_NAME_CHECKED(UTG_Expression_TexturePath, Path));
+
+	if (PathPin->IsConnected() && !bValidation)
+	{
+		ReportError(InContext->Cycle);
+	}
 }
 
 bool UTG_Expression_TexturePath::Validate(MixUpdateCyclePtr Cycle)
 {
 	FString LocalPath;
-	if (!ValidateInputPath(LocalPath)) 
+	UTG_Pin* PathPin = GetParentNode()->GetPin(GET_MEMBER_NAME_CHECKED(UTG_Expression_TexturePath, Path));
+
+	if (!PathPin->IsConnected() && !ValidateInputPath(LocalPath))
 	{
-		auto ErrorType = static_cast<int32>(ETextureGraphErrorType::NODE_WARNING);
-		TextureGraphEngine::GetErrorReporter(Cycle->GetMix())->ReportWarning(ErrorType, FString::Printf(TEXT("Input Path <%s> is not a valid local path"), *Path), GetParentNode());
+		ReportError(Cycle);
 		return true;
 	}
 
 	return true;
 }
 
+void UTG_Expression_TexturePath::ReportError(MixUpdateCyclePtr Cycle)
+{
+	auto ErrorType = static_cast<int32>(ETextureGraphErrorType::NODE_WARNING);
+	TextureGraphEngine::GetErrorReporter(Cycle->GetMix())->ReportWarning(ErrorType, FString::Printf(TEXT("Input Path <%s> is not a valid local path"), *Path), GetParentNode());
+}
+
 void UTG_Expression_TexturePath::SetTitleName(FName NewName)
 {
 	GetParentNode()->GetPin(GET_MEMBER_NAME_CHECKED(UTG_Expression_TexturePath, Path))->SetAliasName(NewName);
diff --git a/Engine/Plugins/Experimental/TextureGraph/Source/TextureGraph/Public/Expressions/Input/TG_Expression_TexturePath.h b/Engine/Plugins/Experimental/TextureGraph/Source/TextureGraph/Public/Expressions/Input/TG_Expression_TexturePath.h
--- a/Engine/Plugins/Experimental/TextureGraph/Source/TextureGraph/Public/Expressions/Input/TG_Expression_TexturePath.h
+++ b/Engine/Plugins/Experimental/TextureGraph/Source/TextureGraph/Public/Expressions/Input/TG_Expression_TexturePath.h
@@ -24,6 +24,8 @@ protected:
 	// Validate the input path, returning the actual path to use
 	// empty if the input path is NOT valid
 	bool ValidateInputPath(FString& ValidatedPath) const;
+
+	void ReportError(MixUpdateCyclePtr Cycle);
 public:
 
 	virtual void Evaluate(FTG_EvaluationContext* InContext) override;
