Node Preview not rendering correctly

#tests 66c2ffd4db75eac44d782843

diff --git a/Engine/Plugins/Experimental/TextureGraph/Source/TextureGraphEngine/Job/Scheduler.cpp b/Engine/Plugins/Experimental/TextureGraph/Source/TextureGraphEngine/Job/Scheduler.cpp
--- a/Engine/Plugins/Experimental/TextureGraph/Source/TextureGraphEngine/Job/Scheduler.cpp
+++ b/Engine/Plugins/Experimental/TextureGraph/Source/TextureGraphEngine/Job/Scheduler.cpp
@@ -175,6 +175,9 @@ void Scheduler::Update(float dt)
 
 				ObserverSource->BatchJobsDone(PreviousBatch);
 
+				/// Mark the current batch as done
+				PreviousBatch->GetCycle()->GetDetails().BroadcastOnDone();
+
 				UE_LOG(LogBatch, Verbose, TEXT("Scheduler Observer::BatchJobsDone finished for Batch: %llu"), PreviousBatch ? PreviousBatch->GetBatchId() : -1);
 			})
 			.then([this]()
@@ -186,9 +189,6 @@ void Scheduler::Update(float dt)
 
 					ObserverSource->BatchDone(CurrentBatch);
 
-					/// Mark the current batch as done
-					CurrentBatch->GetCycle()->GetDetails().BroadcastOnDone();
-
 					UE_LOG(LogBatch, Verbose, TEXT("Scheduler Observer::BatchDone finished for Batch: %llu"), PreviousBatch ? PreviousBatch->GetBatchId() : -1);
 				});
 	}
