Fixed issue: The name of the Bool node is inconsistent across UI

diff --git a/Engine/Plugins/Experimental/TextureGraph/Source/TextureGraph/Private/Expressions/Input/TG_Expression_Bool.cpp b/Engine/Plugins/Experimental/TextureGraph/Source/TextureGraph/Private/Expressions/Input/TG_Expression_Bool.cpp
--- a/Engine/Plugins/Experimental/TextureGraph/Source/TextureGraph/Private/Expressions/Input/TG_Expression_Bool.cpp
+++ b/Engine/Plugins/Experimental/TextureGraph/Source/TextureGraph/Private/Expressions/Input/TG_Expression_Bool.cpp
@@ -10,15 +10,15 @@ void UTG_Expression_Bool::Evaluate(FTG_EvaluationContext* InContext)
 
 	// The Value is updated either as an input or as a setting and then becomes the output for this expression
 	// The pin out is named "ValueOut"
-	ValueOut = bValue;
+	ValueOut = Bool;
 }
 
 void UTG_Expression_Bool::SetTitleName(FName NewName)
 {
-	GetParentNode()->GetInputPin(GET_MEMBER_NAME_CHECKED(UTG_Expression_Bool, bValue))->SetAliasName(NewName);
+	GetParentNode()->GetInputPin(GET_MEMBER_NAME_CHECKED(UTG_Expression_Bool, Bool))->SetAliasName(NewName);
 }
 
 FName UTG_Expression_Bool::GetTitleName() const
 {
-	return GetParentNode()->GetInputPin(GET_MEMBER_NAME_CHECKED(UTG_Expression_Bool, bValue))->GetAliasName();
+	return GetParentNode()->GetInputPin(GET_MEMBER_NAME_CHECKED(UTG_Expression_Bool, Bool))->GetAliasName();
 }
\ No newline at end of file
diff --git a/Engine/Plugins/Experimental/TextureGraph/Source/TextureGraph/Public/Expressions/Input/TG_Expression_Bool.h b/Engine/Plugins/Experimental/TextureGraph/Source/TextureGraph/Public/Expressions/Input/TG_Expression_Bool.h
--- a/Engine/Plugins/Experimental/TextureGraph/Source/TextureGraph/Public/Expressions/Input/TG_Expression_Bool.h
+++ b/Engine/Plugins/Experimental/TextureGraph/Source/TextureGraph/Public/Expressions/Input/TG_Expression_Bool.h
@@ -16,9 +16,8 @@ public:
 
 	virtual void Evaluate(FTG_EvaluationContext* InContext) override;
 
-	// The 4 dimensional vector X, Y, Z and W.
 	UPROPERTY(EditAnywhere, Category = NoCategory, meta = (TGType = "TG_InputParam", PinDisplayName = ""))
-    bool bValue = false;
+    bool Bool = false;
     
 	// The output of the node, which is the color value
     UPROPERTY(meta = (TGType = "TG_Output", PinDisplayName = ""))
