Math nodes don't show scalar output in the Node Preview

#tests 66c4833690633641b848bfe4

diff --git a/Engine/Plugins/Experimental/TextureGraph/Source/TextureGraphEditor/Private/STG_NodePreview.cpp b/Engine/Plugins/Experimental/TextureGraph/Source/TextureGraphEditor/Private/STG_NodePreview.cpp
--- a/Engine/Plugins/Experimental/TextureGraph/Source/TextureGraphEditor/Private/STG_NodePreview.cpp
+++ b/Engine/Plugins/Experimental/TextureGraph/Source/TextureGraphEditor/Private/STG_NodePreview.cpp
@@ -84,7 +84,7 @@ FText FNodeViewer::GetFormatLabelText() const
 		                     FText::FromString(NodeDescriptor.bIsSRGB ? "(sRGB)" : "(Linear)"));
 	}
 
-	return FText();
+	return LabelText;
 }
 
 bool FNodeViewer::IsSingleChannel() const
@@ -110,6 +110,11 @@ void FNodeViewer::SetTexture(const BlobPtr& InBlob, const FLinearColor InClearCo
 	}
 }
 
+void FNodeViewer::SetLabelText(FText InText)
+{
+	LabelText = InText;
+}
+
 void FNodeViewer::SetRGBA(bool bR, bool bG, bool bB, bool bA)
 {
 	bRGBA[0] = bR;
@@ -317,7 +322,7 @@ void STG_NodePreviewWidget::Update() const
 	auto UpdatePreview = [this]
 	{
 		FTG_Variant Variant;
-		GetOutputVariantFromNode(Variant);
+		const bool bValidVariant = GetOutputVariantFromNode(Variant);
 
 		if(Variant.IsTexture())
 		{
@@ -336,9 +341,11 @@ void STG_NodePreviewWidget::Update() const
 		}
 		else
 		{
-			NodeViewer->SetTexture(nullptr);	
+			NodeViewer->SetTexture(nullptr);
 		}
 		
+		NodeViewer->SetLabelText(GetLabelText(Variant, bValidVariant));
+
 		Viewport->ResetZoom(NodeViewer->GetCurrentImageInfo().Size);
 	};
 
@@ -405,6 +412,40 @@ bool STG_NodePreviewWidget::GetOutputVariantFromNode(FTG_Variant& OutVariant) co
 	return false;
 }
 
+FText STG_NodePreviewWidget::GetLabelText(FTG_Variant& InVariant, bool bValidVariant) const
+{
+	const UTG_Node* PreviewNode = LockedNode ? LockedNode : SelectedNode;
+	const bool bInValidTexture = InVariant.IsTexture() && InVariant.GetTexture() && !InVariant.GetTexture()->IsValid();
+	const FText PreviewNotAvailable = FText::FromString("Node preview is not available");
+
+	if (!PreviewNode)
+	{
+		return FText::FromString("Select a node to preview");
+	}
+	else if (!bValidVariant)
+	{
+		return PreviewNotAvailable;
+	}
+	else if (InVariant.IsColor())
+	{
+		return FText::FromString("Color " + InVariant.GetColor().ToFColor(false).ToString());
+	}
+	else if (InVariant.IsVector())
+	{
+		return FText::FromString("Vector (" + InVariant.GetVector().ToString() + ")");
+	}
+	else if (InVariant.IsScalar())
+	{
+		return FText::FromString(FString::Printf(TEXT("Scalar (%0.3f)"), InVariant.GetScalar()));
+	}
+	else if (bInValidTexture)
+	{
+		return FText::FromString("Texture is not valid");
+	}
+
+	return PreviewNotAvailable;
+}
+
 FReply STG_NodePreviewWidget::OnKeyDown(const FGeometry& InGeometry, const FKeyEvent& InKeyEvent)
 {
 	if (CommandList.IsValid() && CommandList->ProcessCommandBindings(InKeyEvent))
diff --git a/Engine/Plugins/Experimental/TextureGraph/Source/TextureGraphEditor/Private/STG_NodePreview.h b/Engine/Plugins/Experimental/TextureGraph/Source/TextureGraphEditor/Private/STG_NodePreview.h
--- a/Engine/Plugins/Experimental/TextureGraph/Source/TextureGraphEditor/Private/STG_NodePreview.h
+++ b/Engine/Plugins/Experimental/TextureGraph/Source/TextureGraphEditor/Private/STG_NodePreview.h
@@ -45,6 +45,9 @@ public:
 	/** Sets the node buffer to a given node. */
 	void SetTexture(const BlobPtr& InBlob, FLinearColor InClearColor = FLinearColor(0.1, 0.1, 0.1, 1));
 
+	/** Sets the label text. */
+	void SetLabelText(FText InText);
+
 	/** Toggles RGBA components for display. */
 	void SetRGBA(bool bR, bool bG, bool bB, bool bA);
 
@@ -81,6 +84,9 @@ private:
 
 	/** Current draw settings for the viewport. */ 
 	UE::ImageWidgets::SImageViewport::FDrawSettings DrawSettings;
+
+	/** Label text for adding node description.*/
+	FText LabelText;
 };
 
 /**
@@ -115,6 +121,9 @@ public:
 
 	bool GetOutputVariantFromNode(FTG_Variant& OutVariant) const; 
 
+	/** Get the label text for the node preview*/
+	FText GetLabelText(FTG_Variant& InVariant, bool bValidVariant) const;
+
 	// SWidget overrides - begin
 	virtual FReply OnKeyDown(const FGeometry& InGeometry, const FKeyEvent& InKeyEvent) override;
 	// SWidget overrides - end
