Replacing an old texture with a new one during export does not clear the old texture cache
Adding timestamp as an argument to the job so if for whatever reason a file is modified it can create a new job for it otherwise the job was getting culled.

#tests 66c5f3c48f1999308f7e4fd5

diff --git a/Engine/Plugins/Experimental/TextureGraph/Source/TextureGraphEngine/Model/StaticImageResource.cpp b/Engine/Plugins/Experimental/TextureGraph/Source/TextureGraphEngine/Model/StaticImageResource.cpp
--- a/Engine/Plugins/Experimental/TextureGraph/Source/TextureGraphEngine/Model/StaticImageResource.cpp
+++ b/Engine/Plugins/Experimental/TextureGraph/Source/TextureGraphEngine/Model/StaticImageResource.cpp
@@ -68,9 +68,12 @@ TiledBlobPtr UStaticImageResource::GetBlob(MixUpdateCyclePtr Cycle, BufferDescri
 	if (BlobObj)
 		return BlobObj;
 
+	FDateTime LastModified = GetAssetTimeStamp();
+
 	auto JobObj = std::make_unique<Job_LoadStaticImageResource>(Cycle->GetMix(), this, TargetId);
 
 	JobObj->AddArg(ARG_STRING(AssetUUID, "AssetUUID")); /// Assett UUID contains the Name of the asset megascan surface AND the Name of channel
+	JobObj->AddArg(ARG_STRING(LastModified.ToString(), "TimeStamp"));
 
 	FString Name = "[StaticImageResource]-" + AssetUUID;
 	DesiredDesc.Format = BufferFormat::LateBound;
@@ -80,3 +83,25 @@ TiledBlobPtr UStaticImageResource::GetBlob(MixUpdateCyclePtr Cycle, BufferDescri
 
 	return BlobObj;
 }
+
+FDateTime UStaticImageResource::GetAssetTimeStamp()
+{
+	FString FullPath = AssetUUID;
+
+	if (!IsFileSystem())
+	{
+		const FString PackageName = FPackageName::ObjectPathToPackageName(AssetUUID);
+		const FString FileName = FPackageName::LongPackageNameToFilename(PackageName, FPackageName::GetAssetPackageExtension());
+		FullPath = IFileManager::Get().ConvertToAbsolutePathForExternalAppForRead(*FileName);
+	}
+	
+	FDateTime LastModified;
+	IFileManager& FileManager = IFileManager::Get();
+
+	if (FileManager.FileExists(*FullPath))
+	{
+		LastModified = FileManager.GetTimeStamp(*FullPath);
+	}
+
+	return LastModified;
+}
diff --git a/Engine/Plugins/Experimental/TextureGraph/Source/TextureGraphEngine/Model/StaticImageResource.h b/Engine/Plugins/Experimental/TextureGraph/Source/TextureGraphEngine/Model/StaticImageResource.h
--- a/Engine/Plugins/Experimental/TextureGraph/Source/TextureGraphEngine/Model/StaticImageResource.h
+++ b/Engine/Plugins/Experimental/TextureGraph/Source/TextureGraphEngine/Model/StaticImageResource.h
@@ -19,24 +19,31 @@ public:
 private:
 	friend class Job_LoadStaticImageResource;
 	
+	/*Unique id for the asset within the entire system*/
 	UPROPERTY()
-	FString 						AssetUUID;						/// Unique id for the asset within the entire system
-	TiledBlobPtr					BlobObj;						/// The blob that represents the data for this source
-	bool							bIsFilesystem = false;			/// Is loading directly from the filesystem
+	FString AssetUUID;
 
-	virtual AsyncTiledBlobRef		Load(MixUpdateCyclePtr Cycle);
+	/*The blob that represents the data for this source*/
+	TiledBlobPtr BlobObj;
+
+	/*Is loading directly from the filesystem*/
+	bool bIsFilesystem = false;
+
+	virtual AsyncTiledBlobRef Load(MixUpdateCyclePtr Cycle);
+
+	FDateTime GetAssetTimeStamp();
 	
 public:
-	virtual							~UStaticImageResource() override;
+	virtual ~UStaticImageResource() override;
 	
-	virtual TiledBlobPtr			GetBlob(MixUpdateCyclePtr Cycle, BufferDescriptor DesiredDesc, int32 TargetId);
+	virtual TiledBlobPtr GetBlob(MixUpdateCyclePtr Cycle, BufferDescriptor DesiredDesc, int32 TargetId);
 
 	//////////////////////////////////////////////////////////////////////////
 	/// Inline functions
 	//////////////////////////////////////////////////////////////////////////
-	FORCEINLINE const FString&		GetAssetUUID() const { return AssetUUID; }
-	FORCEINLINE void				SetAssetUUID(const FString& UUID) { AssetUUID = UUID; }
-	FORCEINLINE void				SetIsFileSystem(bool bInIsFileSystem) { bIsFilesystem = bInIsFileSystem; }
-	FORCEINLINE bool				IsFileSystem() const { return bIsFilesystem; }
+	FORCEINLINE const FString& GetAssetUUID() const { return AssetUUID; }
+	FORCEINLINE void SetAssetUUID(const FString& UUID) { AssetUUID = UUID; }
+	FORCEINLINE void SetIsFileSystem(bool bInIsFileSystem) { bIsFilesystem = bInIsFileSystem; }
+	FORCEINLINE bool IsFileSystem() const { return bIsFilesystem; }
 };
 
