Reconciled offline work

diff --git a/Engine/Plugins/Experimental/TextureGraph/Source/TextureGraph/Private/Expressions/TG_Expression_MaterialBase.cpp b/Engine/Plugins/Experimental/TextureGraph/Source/TextureGraph/Private/Expressions/TG_Expression_MaterialBase.cpp
--- a/Engine/Plugins/Experimental/TextureGraph/Source/TextureGraph/Private/Expressions/TG_Expression_MaterialBase.cpp
+++ b/Engine/Plugins/Experimental/TextureGraph/Source/TextureGraph/Private/Expressions/TG_Expression_MaterialBase.cpp
@@ -386,7 +386,7 @@ FTG_SignaturePtr UTG_Expression_MaterialBase::BuildSignatureDynamically() const
 
 void UTG_Expression_MaterialBase::CopyVarGeneric(const FTG_Argument& Arg, FTG_Var* InVar, bool CopyVarToArg)
 {
-	if(MaterialInstance)
+	if (MaterialInstance)
 	{
 		// Try to find the MatParam matching the arg
 		if (Arg.CPPTypeName.IsEqual(TEXT("bool")))
@@ -394,7 +394,9 @@ void UTG_Expression_MaterialBase::CopyVarGeneric(const FTG_Argument& Arg, FTG_Va
 	#if WITH_EDITORONLY_DATA
 			if (CopyVarToArg)
 			{
-				MaterialInstance->SetStaticSwitchParameterValueEditorOnly(Arg.GetName(), InVar->GetAs<bool>());
+				// Disable assigning static bool values to MID for now see UE-209533 & UE-219306
+				// We need to use a different solution for Material with Static Switch, we can't use MID for this, evetnually MIC editor only...
+				//MaterialInstance->SetStaticSwitchParameterValueEditorOnly(Arg.GetName(), InVar->GetAs<bool>());
 			} else
 			{
 				FMaterialParameterInfo ParameterInfo(Arg.GetName());
diff --git a/Engine/Plugins/Experimental/TextureGraph/Source/TextureGraph/Public/Expressions/Input/TG_Expression_Material.h b/Engine/Plugins/Experimental/TextureGraph/Source/TextureGraph/Public/Expressions/Input/TG_Expression_Material.h
--- a/Engine/Plugins/Experimental/TextureGraph/Source/TextureGraph/Public/Expressions/Input/TG_Expression_Material.h
+++ b/Engine/Plugins/Experimental/TextureGraph/Source/TextureGraph/Public/Expressions/Input/TG_Expression_Material.h
@@ -45,8 +45,9 @@ public:
 	virtual FText GetTooltipText() const override { return FText::FromString(TEXT("Renders a material into a quad and makes it available. It is automatically exposed as a graph input parameter.")); } 
 
 protected:
-	UPROPERTY(Transient)
-	TObjectPtr<UMaterialInterface> MaterialCopy;
+	// Transient and per instance Data, recreated on every new instance from the reference material
+	UPROPERTY(Transient, DuplicateTransient)
+	TObjectPtr<UMaterialInterface> MaterialCopy = nullptr;
 
 	virtual void SetMaterialInternal(UMaterialInterface* InMaterial) override;
 
diff --git a/Engine/Plugins/Experimental/TextureGraph/Source/TextureGraph/Public/Expressions/TG_Expression_MaterialBase.h b/Engine/Plugins/Experimental/TextureGraph/Source/TextureGraph/Public/Expressions/TG_Expression_MaterialBase.h
--- a/Engine/Plugins/Experimental/TextureGraph/Source/TextureGraph/Public/Expressions/TG_Expression_MaterialBase.h
+++ b/Engine/Plugins/Experimental/TextureGraph/Source/TextureGraph/Public/Expressions/TG_Expression_MaterialBase.h
@@ -62,8 +62,9 @@ public:
 	const TArray<FName>& GetAvailableMaterialAttributeNames() const { return AvailableMaterialAttributeNames; }
 
 protected:
-	UPROPERTY(Transient)
-	TObjectPtr<UMaterialInstanceDynamic> MaterialInstance; // A local Instance Material is recreated for the Material member
+	// A local per Instance Material is recreated from the reference material assigned through SetMaterialInternal
+	UPROPERTY(Transient, DuplicateTransient)
+	TObjectPtr<UMaterialInstanceDynamic> MaterialInstance = nullptr;
 
 	TArray<EDrawMaterialAttributeTarget>	AvailableMaterialAttributeIds; // The set of material properties available for rendering
 	TArray<FName>							AvailableMaterialAttributeNames; // same with the attribute names
diff --git a/Engine/Plugins/Experimental/TextureGraph/Source/TextureGraphEditor/Private/EdGraph/TG_EdGraph.cpp b/Engine/Plugins/Experimental/TextureGraph/Source/TextureGraphEditor/Private/EdGraph/TG_EdGraph.cpp
--- a/Engine/Plugins/Experimental/TextureGraph/Source/TextureGraphEditor/Private/EdGraph/TG_EdGraph.cpp
+++ b/Engine/Plugins/Experimental/TextureGraph/Source/TextureGraphEditor/Private/EdGraph/TG_EdGraph.cpp
@@ -183,7 +183,14 @@ void UTG_EdGraph::GraphChanged(UTG_Graph* InGraph, UTG_Node* InNode, bool Tweaki
 void UTG_EdGraph::OnNodeSignatureChanged(UTG_Node* InNode)
 {
 	UTG_EdGraphNode* EdGraphNode = GetViewModelNode(InNode->GetId());
-	EdGraphNode->ReconstructNode();
+	if (EdGraphNode)
+	{
+		EdGraphNode->ReconstructNode();
+	}
+	else
+	{
+		UE_LOG(LogTemp, Warning, TEXT("UTG_EdGraph::OnNodeSignatureChanged TG_Node not found in the EdNodes? this should not happen."))
+	}
 }
 
 void UTG_EdGraph::OnNodePostEvaluation(UTG_Node* InNode, const FTG_EvaluationContext* Context)
