TextureGraph Fixed node preview not showing Texture2D

#tests 66c60308ea10c9589c8a38bb

diff --git a/Engine/Plugins/Experimental/TextureGraph/Source/TextureGraphEditor/Private/STG_NodePreview.cpp b/Engine/Plugins/Experimental/TextureGraph/Source/TextureGraphEditor/Private/STG_NodePreview.cpp
--- a/Engine/Plugins/Experimental/TextureGraph/Source/TextureGraphEditor/Private/STG_NodePreview.cpp
+++ b/Engine/Plugins/Experimental/TextureGraph/Source/TextureGraphEditor/Private/STG_NodePreview.cpp
@@ -14,6 +14,7 @@
 #include "Framework/MultiBox/MultiBoxBuilder.h"
 #include "Widgets/Text/STextBlock.h"
 #include "Widgets/Input/SCheckBox.h"
+#include "Engine/Texture2D.h"
 
 #define LOCTEXT_NAMESPACE "STG_NodePreview"
 
@@ -21,7 +22,7 @@ UE::ImageWidgets::IImageViewer::FImageInfo FNodeViewer::GetCurrentImageInfo() co
 {
 	if (NodeTexture)
 	{
-		return {{}, {NodeTexture->SizeX, NodeTexture->SizeY}, NodeTexture->GetNumMips(), true};
+		return {{}, {(int32)NodeTexture->GetSurfaceWidth(), (int32)NodeTexture->GetSurfaceHeight()}, GetNodeTextureNumMips(), true};
 	}
 
 	return {{}, FIntPoint::ZeroValue, 0, false};
@@ -173,9 +174,9 @@ ESimpleElementBlendMode FNodeViewer::GetBlendMode() const
 	return static_cast<ESimpleElementBlendMode>(Result);
 }
 
-UTextureRenderTarget2D* FNodeViewer::GetTextureFromBuffer(const DeviceBufferPtr& Buffer) const
+UTexture* FNodeViewer::GetTextureFromBuffer(const DeviceBufferPtr& Buffer) const
 {
-	UTextureRenderTarget2D* OutTexture = nullptr;
+	UTexture* OutTexture = nullptr;
 
 	if (const std::shared_ptr<DeviceBuffer_FX> FXBuffer = std::static_pointer_cast<DeviceBuffer_FX>(Buffer))
 	{
@@ -188,9 +189,13 @@ UTextureRenderTarget2D* FNodeViewer::GetTextureFromBuffer(const DeviceBufferPtr&
 			{
 				OutTexture = RenderTarget;
 			}
+			else if (UTexture2D* Texture2D = Cast<UTexture2D>(BufferTexture))
+			{
+				OutTexture = Texture2D;
+			}
 			else
 			{
-				UE_LOG(LogTemp, Warning, TEXT("Texture is not a UTextureRenderTarget2D."));
+				UE_LOG(LogTemp, Warning, TEXT("Texture is not a UTexture2D | UTextureRenderTarget2D."));
 			}
 		}
 		else
@@ -206,6 +211,22 @@ UTextureRenderTarget2D* FNodeViewer::GetTextureFromBuffer(const DeviceBufferPtr&
 	return OutTexture;
 }
 
+int32 FNodeViewer::GetNodeTextureNumMips() const
+{
+	//Enable this when we can switch mip levels
+#if 0
+	if (UTexture2D* Texture2D = Cast<UTexture2D>(NodeTexture))
+	{
+		return Texture2D->GetNumMips();
+	}
+	else if(UTextureRenderTarget2D* RenderTarget = Cast<UTextureRenderTarget2D>(NodeTexture))
+	{
+		return RenderTarget->GetNumMips();
+	}
+#endif
+	return 0;
+}
+
 bool FNodeViewer::IsSRGB() const
 {
 	return NodeDescriptor.bIsSRGB;
diff --git a/Engine/Plugins/Experimental/TextureGraph/Source/TextureGraphEditor/Private/STG_NodePreview.h b/Engine/Plugins/Experimental/TextureGraph/Source/TextureGraphEditor/Private/STG_NodePreview.h
--- a/Engine/Plugins/Experimental/TextureGraph/Source/TextureGraphEditor/Private/STG_NodePreview.h
+++ b/Engine/Plugins/Experimental/TextureGraph/Source/TextureGraphEditor/Private/STG_NodePreview.h
@@ -64,8 +64,11 @@ private:
 	/** Determines the blend mode based on the node texture and the RGBA toggles. */
 	ESimpleElementBlendMode GetBlendMode() const;
 
-	/** Retrieves the render target texture  from a given node buffer. */
-	UTextureRenderTarget2D* GetTextureFromBuffer(const DeviceBufferPtr& Buffer) const;
+	/** Retrieves the texture  from a given node buffer. */
+	UTexture* GetTextureFromBuffer(const DeviceBufferPtr& Buffer) const;
+
+	/** Returns the num mips for node texture*/
+	int32 GetNodeTextureNumMips() const;
 
 	/** Indicates that a node is in sRGB format. */
 	bool IsSRGB() const;
@@ -73,8 +76,8 @@ private:
 	/** The Blob for this node. */
 	BlobPtr CurrentBlob;
 
-	/** The render target texture for the node. */
-	UTextureRenderTarget2D* NodeTexture = nullptr;
+	/** The texture for the node. */
+	UTexture* NodeTexture = nullptr;
 
 	/** Toggles for enabling RGBA components for drawing. */
 	bool bRGBA[4] = {true, true, true, true};
