Resetting node settings in the Details tab only affects the UI in the Details tab
Property change events needed to be called when the property itself is changed in the variant property customization

diff --git a/Engine/Plugins/Experimental/TextureGraph/Source/TextureGraphEditor/Private/Customizations/TG_VariantCustomization.h b/Engine/Plugins/Experimental/TextureGraph/Source/TextureGraphEditor/Private/Customizations/TG_VariantCustomization.h
--- a/Engine/Plugins/Experimental/TextureGraph/Source/TextureGraphEditor/Private/Customizations/TG_VariantCustomization.h
+++ b/Engine/Plugins/Experimental/TextureGraph/Source/TextureGraphEditor/Private/Customizations/TG_VariantCustomization.h
@@ -57,17 +57,14 @@ public:
 					// Handle->SetInstanceMetaData(TEXT("TGType"), TEXT("TG_Output"));
 					TSharedPtr<IPropertyUtilities> PropertyUtilities;
 					PropertyUtilities = CustomizationUtils.GetPropertyUtilities();
-					Handle->SetOnChildPropertyValuePreChange(FSimpleDelegate::CreateLambda([=, this]()
-					{
-						ParentExpression->Modify();
-					}));
-					Handle->SetOnChildPropertyValueChangedWithData(TDelegate<void(const FPropertyChangedEvent&)>::CreateLambda([=, this](const FPropertyChangedEvent& InEvent)
+
+					TDelegate<void(const FPropertyChangedEvent&)> OnPropertyChangeWithData = TDelegate<void(const FPropertyChangedEvent&)>::CreateLambda([=, this](const FPropertyChangedEvent& InEvent)
 					{
 						if (PropertyUtilities.IsValid())
 						{
 							void* StructData = nullptr;
 							const FPropertyAccess::Result Result = Handle->GetValueData(StructData);
-							if(Result == FPropertyAccess::Success)
+							if (Result == FPropertyAccess::Success)
 							{
 								check(StructData);
 								FTG_TextureDescriptor* Descriptor = static_cast<FTG_TextureDescriptor*>(StructData);
@@ -79,7 +76,20 @@ public:
 								PropertyUtilities->RequestRefresh();
 							}
 						}
-					}));
+					});
+
+					FSimpleDelegate OnPropertyValuePreChange = FSimpleDelegate::CreateLambda([=, this]()
+					{
+						ParentExpression->Modify();
+					});
+
+					//when property itself is modified like when reset button is clicked
+					Handle->SetOnPropertyValuePreChange(OnPropertyValuePreChange);
+					Handle->SetOnPropertyValueChangedWithData(OnPropertyChangeWithData);
+
+					//when child properties are modified
+					Handle->SetOnChildPropertyValuePreChange(OnPropertyValuePreChange);
+					Handle->SetOnChildPropertyValueChangedWithData(OnPropertyChangeWithData);
 				}
 			}
 		}
