Fixed issue: Deleting TextureGraph asset after deleting TextureGraph node it was assigned to crashes the Editor

Fix/How: This happens when we force delete the reference in the undo history. In that case, the Null texture graph is being set. We have already removed the pins from the graph list. I have added null checks to avoid crashes.

diff --git a/Engine/Plugins/Experimental/TextureGraph/Source/TextureGraph/Private/TG_Graph.cpp b/Engine/Plugins/Experimental/TextureGraph/Source/TextureGraph/Private/TG_Graph.cpp
--- a/Engine/Plugins/Experimental/TextureGraph/Source/TextureGraph/Private/TG_Graph.cpp
+++ b/Engine/Plugins/Experimental/TextureGraph/Source/TextureGraph/Private/TG_Graph.cpp
@@ -219,7 +219,8 @@ void UTG_Graph::RegenerateNode(UTG_Node* InNode)
 	{
 		// The pin is still valid so kill it
 		UTG_Pin* Pin = InNode->Pins[i].Get();
-		if (Pin)
+		
+		if (Pin && IsValidPin(Pin->GetId()))
 		{
 			KillPin(Pin->GetId()); // This also remove the Pin from the Params
 		}
diff --git a/Engine/Plugins/Experimental/TextureGraph/Source/TextureGraph/Private/TG_Node.cpp b/Engine/Plugins/Experimental/TextureGraph/Source/TextureGraph/Private/TG_Node.cpp
--- a/Engine/Plugins/Experimental/TextureGraph/Source/TextureGraph/Private/TG_Node.cpp
+++ b/Engine/Plugins/Experimental/TextureGraph/Source/TextureGraph/Private/TG_Node.cpp
@@ -147,8 +147,12 @@ void UTG_Node::OnExpressionChangedWithoutVar(const FPropertyChangedEvent& Proper
 		if (PinId.IsValid())
 		{
 			auto ModifiedPin = GetPin(PinId);
-			auto Arg = ModifiedPin->GetArgument();
-			ModifiedPin->EditSelfVar()->CopyFrom(GetExpression(), Arg);
+
+			if(ModifiedPin)
+			{
+				auto Arg = ModifiedPin->GetArgument();
+				ModifiedPin->EditSelfVar()->CopyFrom(GetExpression(), Arg);
+			}
 		}
 	}
 
diff --git a/Engine/Plugins/Experimental/TextureGraph/Source/TextureGraph/Public/TG_Graph.h b/Engine/Plugins/Experimental/TextureGraph/Source/TextureGraph/Public/TG_Graph.h
--- a/Engine/Plugins/Experimental/TextureGraph/Source/TextureGraph/Public/TG_Graph.h
+++ b/Engine/Plugins/Experimental/TextureGraph/Source/TextureGraph/Public/TG_Graph.h
@@ -83,7 +83,7 @@ private:
 	void NotifyNodePostEvaluate(UTG_Node* InNode, const FTG_EvaluationContext* InContext);
 
 	bool IsValidNode(FTG_Id NodeId) const { return (NodeId != FTG_Id::INVALID && Nodes.IsValidIndex(NodeId.NodeIdx())); }
-	bool IsValidPin(FTG_Id PinId) const { return IsValidNode(PinId) && Nodes[PinId.NodeIdx()]->Pins.IsValidIndex(PinId.PinIdx()); }
+	bool IsValidPin(FTG_Id PinId) const { return IsValidNode(PinId) && Nodes[PinId.NodeIdx()] && Nodes[PinId.NodeIdx()]->Pins.IsValidIndex(PinId.PinIdx()); }
 	
 protected:
 	friend class UTG_Node;
