Fixed issue: Settings in the Preview Scene Settings disappear after reopening the asset

Issue: Preview scene settings init depend upon the 3D viewport. When we rearrange the tabs layout, there is the possibility that the viewport is not initiated yet.

Fix: Moved the preview settings init to layout changed. This way, we always have a valid viewport.

diff --git a/Engine/Plugins/Experimental/TextureGraph/Source/TextureGraphEditor/Private/TG_Editor.cpp b/Engine/Plugins/Experimental/TextureGraph/Source/TextureGraphEditor/Private/TG_Editor.cpp
--- a/Engine/Plugins/Experimental/TextureGraph/Source/TextureGraphEditor/Private/TG_Editor.cpp
+++ b/Engine/Plugins/Experimental/TextureGraph/Source/TextureGraphEditor/Private/TG_Editor.cpp
@@ -173,7 +173,6 @@ void FTG_Editor::InitEditor(const EToolkitMode::Type Mode, const TSharedPtr< cla
 {
 	TG_Parameters = NewObject<UTG_Parameters>();
 	OriginalTextureGraph = InTextureGraph;
-
 	// We create a duplicate TextureGraph. All the modifications will be done in it.
 	// All the modifications are applied to original TextureGraph when asset is saved.
 	// Propagate all object flags except for RF_Standalone, otherwise the preview material won't GC once
@@ -325,18 +324,17 @@ void FTG_Editor::InitEditor(const EToolkitMode::Type Mode, const TSharedPtr< cla
 				)
 			)
 		);
-
-
 	FAssetEditorToolkit::InitAssetEditor(Mode, InitToolkitHost, TG_EditorAppIdentifier, StandaloneDefaultLayout, /*bCreateDefaultToolbar*/ true, /*bCreateDefaultStandaloneMenu*/ true, InTextureGraph);
 
 	RegenerateMenusAndToolbars();
 
-	bool bViewportIsOff = !ViewportTabContent.IsValid();
-	TSharedPtr<SDockTab> ViewportTab; 
+	auto CurrentViewportTab = GetTabManager()->FindExistingLiveTab(FTG_EditorTabs::ViewportTabId);
+	bool bViewportIsOff = CurrentViewportTab == nullptr;
+	
 	// check here if 3d viewport is turned off, we need to turn it on temporarily to initialize our systems correctly
 	if (bViewportIsOff)
 	{
-		ViewportTab = GetTabManager()->TryInvokeTab(FTG_EditorTabs::ViewportTabId);	
+		CurrentViewportTab = GetTabManager()->TryInvokeTab(FTG_EditorTabs::ViewportTabId);	
 	}
 	
 	// Set the preview mesh for the material.  This call must occur after the toolbar is initialized.
@@ -345,7 +343,7 @@ void FTG_Editor::InitEditor(const EToolkitMode::Type Mode, const TSharedPtr< cla
 	if (bViewportIsOff)
 	{
 		// turn the viewporttab back off
-		ViewportTab->RequestCloseTab();
+		CurrentViewportTab->RequestCloseTab();
 	}
 
 	EditedTextureGraph->GetPackage()->ClearDirtyFlag();
@@ -927,9 +925,25 @@ FGraphAppearanceInfo FTG_Editor::GetGraphAppearance() const
 	return AppearanceInfo;
 }
 
+void FTG_Editor::BuildSubTools()
+{
+	FAdvancedPreviewSceneModule& AdvancedPreviewSceneModule = FModuleManager::LoadModuleChecked<FAdvancedPreviewSceneModule>("AdvancedPreviewScene");
+
+	TArray<FAdvancedPreviewSceneModule::FDetailDelegates> Delegates;
+	Delegates.Add({ OnPreviewSceneChangedDelegate });
+	AdvancedPreviewSettingsWidget = AdvancedPreviewSceneModule.CreateAdvancedPreviewSceneSettingsWidget(GetEditorViewport()->GetPreviewScene(), nullptr, TArray<FAdvancedPreviewSceneModule::FDetailCustomizationInfo>(),  TArray<FAdvancedPreviewSceneModule::FPropertyTypeCustomizationInfo>(), Delegates);
+
+	if (PreviewSceneSettingsDockTab.IsValid())
+	{
+		PreviewSceneSettingsDockTab.Pin()->SetContent(AdvancedPreviewSettingsWidget.ToSharedRef());
+	}
+}
+
 void FTG_Editor::OnEditorLayoutChanged()
 {
+	BuildSubTools();
 
+	OnPreviewSceneChangedDelegate.Broadcast(GetEditorViewport()->GetPreviewScene());
 }
 
 TSharedRef<SDockTab> FTG_Editor::SpawnTab_GraphEditor(const FSpawnTabArgs& Args)
@@ -960,20 +974,16 @@ TSharedRef<SDockTab> FTG_Editor::SpawnTab_Viewport(const FSpawnTabArgs& Args)
 
 	TSharedRef< SDockTab > DockableTab =
 		SNew(SDockTab);
-
 	MakeViewportFunc = [this](const FAssetEditorViewportConstructionArgs& InArgs)
-		{
-			return SNew(STG_EditorViewport)
-				.InTextureGraph(EditedTextureGraph);
-		};
-
+	{
+		return SNew(STG_EditorViewport)
+			.InTextureGraph(EditedTextureGraph);
+	};
 	// Create a new tab
 	ViewportTabContent = MakeShareable(new FEditorViewportTabContent());
 	ViewportTabContent->OnViewportTabContentLayoutChanged().AddRaw(this, &FTG_Editor::OnEditorLayoutChanged);
-
 	const FString LayoutId = FString("TG_EditorViewport");
 	ViewportTabContent->Initialize(MakeViewportFunc, DockableTab, LayoutId);
-
 	// This call must occur after the toolbar is initialized.
 	SetViewportPreviewMesh();
 	
@@ -1017,23 +1027,12 @@ TSharedRef<SDockTab> FTG_Editor::SpawnTab_PreviewSettings(const FSpawnTabArgs& A
 {
 	check(Args.GetTabId() == FTG_EditorTabs::PreviewSceneSettingsTabId);
 
-	TSharedRef<SWidget> InWidget = SNullWidget::NullWidget;
-
-	if (GetEditorViewport().IsValid())
-	{
-		FAdvancedPreviewSceneModule& AdvancedPreviewSceneModule = FModuleManager::LoadModuleChecked<FAdvancedPreviewSceneModule>("AdvancedPreviewScene");
-		InWidget = AdvancedPreviewSceneModule.CreateAdvancedPreviewSceneSettingsWidget(GetEditorViewport()->GetPreviewScene());
-	}
-
-	TSharedRef<SDockTab> SpawnedTab = SNew(SDockTab)
+	check( Args.GetTabId() == FTG_EditorTabs::PreviewSceneSettingsTabId );
+	return SAssignNew(PreviewSceneSettingsDockTab, SDockTab)
+		.Label( LOCTEXT("TG_EditorPreviewSceneSettings_TabTitle", "Preview Scene Settings") )
 		[
-			SNew(SBox)
-				[
-					InWidget
-				]
+			AdvancedPreviewSettingsWidget.IsValid() ? AdvancedPreviewSettingsWidget.ToSharedRef() : SNullWidget::NullWidget
 		];
-	
-	return SpawnedTab;
 }
 
 TSharedRef<SDockTab> FTG_Editor::SpawnTab_ParameterDefaults(const FSpawnTabArgs& Args)
diff --git a/Engine/Plugins/Experimental/TextureGraph/Source/TextureGraphEditor/Private/TG_Editor.h b/Engine/Plugins/Experimental/TextureGraph/Source/TextureGraphEditor/Private/TG_Editor.h
--- a/Engine/Plugins/Experimental/TextureGraph/Source/TextureGraphEditor/Private/TG_Editor.h
+++ b/Engine/Plugins/Experimental/TextureGraph/Source/TextureGraphEditor/Private/TG_Editor.h
@@ -3,6 +3,7 @@
 #pragma once
 
 #include "CoreMinimal.h"
+#include "AdvancedPreviewSceneModule.h"
 #include "ITG_Editor.h"
 #include "EdGraph/EdGraphPin.h"
 #include "EditorUndoClient.h"
@@ -201,6 +202,9 @@ private:
 	/** Gets the current TG_ Graph's appearance */
 	FGraphAppearanceInfo							GetGraphAppearance() const;
 
+	/** Builds the sub-tools that are a part of the texture graph editor. */
+	void											BuildSubTools();
+	
 	/** Called when the Viewport Layout has changed. */
 	void											OnEditorLayoutChanged();
 
@@ -314,6 +318,8 @@ private:
 	/** Storage for our viewport creation function that will be passed to the viewport layout system*/
 	AssetEditorViewportFactoryFunction				MakeViewportFunc;
 
+	FAdvancedPreviewSceneModule::FOnPreviewSceneChanged OnPreviewSceneChangedDelegate;
+
 	// Tracking the active viewports in this editor.
 	TSharedPtr<class FEditorViewportTabContent>		ViewportTabContent;
 
@@ -321,11 +327,15 @@ private:
 	TWeakPtr<SDockTab>								SpawnedDetailsTab;	
 	TWeakPtr<SDockTab>								NodeHistogramTab;	
 	TWeakPtr<SDockTab>								PaletteTab;	
-
+	TWeakPtr<SDockTab>								PreviewSceneSettingsDockTab;
+	
 	/** Stats log, with the log listing that it reflects */
 	TSharedPtr<class SWidget>						ErrorsWidget;
 	TSharedPtr<class IMessageLogListing>			ErrorsListing;
 
+	/** Scene preview settings widget */
+	TSharedPtr<SWidget>								AdvancedPreviewSettingsWidget;
+
 	/** Hashed error code use to refresh the error displaying widget when this will change */
 	FString											ErrorHash;
 
