Fixed issue: Texture Graph with Material node is rendered in very low res in Texture Graph Exporter

Fix/How: When a texture graph asset containing materials is loaded, there is a chance that the textures referenced in the material are not yet loaded. To fix that, we get all the referenced textures and wait for them to load.

diff --git a/Engine/Plugins/Experimental/TextureGraph/Source/TextureGraph/Private/Expressions/TG_Expression_MaterialBase.cpp b/Engine/Plugins/Experimental/TextureGraph/Source/TextureGraph/Private/Expressions/TG_Expression_MaterialBase.cpp
--- a/Engine/Plugins/Experimental/TextureGraph/Source/TextureGraph/Private/Expressions/TG_Expression_MaterialBase.cpp
+++ b/Engine/Plugins/Experimental/TextureGraph/Source/TextureGraph/Private/Expressions/TG_Expression_MaterialBase.cpp
@@ -12,6 +12,10 @@
 #include "FxMat/MaterialManager.h"
 #include "Materials/MaterialAttributeDefinitionMap.h"
 
+#if WITH_EDITOR
+#include "TextureCompiler.h"
+#endif
+
 
 EDrawMaterialAttributeTarget UTG_Expression_MaterialBase::ConvertEMaterialPropertyToEDrawMaterialAttributeTarget(EMaterialProperty InMaterialProperty)
 {
@@ -173,6 +177,26 @@ TiledBlobPtr UTG_Expression_MaterialBase::CreateRenderMaterialJob(FTG_Evaluation
 	InRenderMaterial->Instance()->EnsureIsComplete();
 	InRenderMaterial->Instance()->SetForceMipLevelsToBeResident(true, true, -1);
 
+	TArray<UTexture*> ReferencedTextures;
+	
+	InRenderMaterial->Instance()->GetUsedTextures(ReferencedTextures, EMaterialQualityLevel::Num, false, ERHIFeatureLevel::Num, true);
+
+	for (int32 i = 0; i < ReferencedTextures.Num(); ++i)
+	{
+		UTexture* ReferencedTexture = Cast<UTexture>(ReferencedTextures[i]);
+
+		if(ReferencedTexture)
+		{
+			ReferencedTexture->SetForceMipLevelsToBeResident(30);
+			ReferencedTexture->WaitForStreaming();
+
+#if WITH_EDITOR
+			const bool IsCompiling = FTextureCompilingManager::Get().IsCompilingTexture(ReferencedTexture); 
+			checkSlow(!IsCompiling);
+#endif
+		}
+	}
+	
 	JobUPtr MaterialJob = std::make_unique<Job>(InContext->Cycle->GetMix(), InContext->TargetId, std::static_pointer_cast<BlobTransform>(InRenderMaterial), GetParentNode());
 
 	FLinearColor PSControl;
