UE-220413 fix. Crash if the result tiled blob is not a promise.

#tests 66d0ab2a51460c13dc0d41f5

diff --git a/Engine/Plugins/Experimental/TextureGraph/Source/TextureGraphEngine/Job/Job.cpp b/Engine/Plugins/Experimental/TextureGraph/Source/TextureGraphEngine/Job/Job.cpp
--- a/Engine/Plugins/Experimental/TextureGraph/Source/TextureGraphEngine/Job/Job.cpp
+++ b/Engine/Plugins/Experimental/TextureGraph/Source/TextureGraphEngine/Job/Job.cpp
@@ -923,18 +923,19 @@ AsyncTransformResultPtr Job::ExecTransform(JobRunInfo InRunInfo, BlobTransformPt
 
 	if (TransformObj->GeneratesData())
 	{
-		check(Result && Result->IsPromise());
-		TiledBlob_PromisePtr result = GetResultPromise();
-
 		TransArgs.Target = GetResultRef();
 
-		if (TileX >= 0 && TileY >= 0 && result->TiledTarget())
+		if (TileX >= 0 && TileY >= 0 && Result->TiledTarget())
 		{
-			BlobRef TileResult = result->GetTile(TileX, TileY);
+			BlobRef TileResult = Result->GetTile(TileX, TileY);
 			
 			/// This should've been prepared in the PrepareResources function
 			check(TileResult);
 
+			/// If the tile result has already been finalised then we don't do anything to it
+			if (TileResult->IsFinalised())
+				return cti::make_ready_continuable(std::make_shared<TransformResult>());
+
 			TransArgs.Target = TileResult.get();
 
 			/// Set the job result preemptively
@@ -1102,8 +1103,6 @@ void Job::AddResultToBlobber()
 #endif 
 
 		Result = TextureGraphEngine::GetBlobber()->AddTiledResult(TempHash, ResultOrg, CacheOpt);
-		//ResultOrg = nullptr;
-
 		check(Result && Result->IsTiled());
 	}
 }
