Fixed issue: Incorrect colors on exported textures from Texture Graph with different resolutions

When Output and Input desc don't match, it just ignores input's Desc sRGB value and uses Output desc value (which is always false)

Right now, the sRGB settings cannot be changed in the output node, so we can safely apply the input sRGB setting here.

diff --git a/Engine/Plugins/Experimental/TextureGraph/Source/TextureGraph/Private/Expressions/Output/TG_Expression_Output.cpp b/Engine/Plugins/Experimental/TextureGraph/Source/TextureGraph/Private/Expressions/Output/TG_Expression_Output.cpp
--- a/Engine/Plugins/Experimental/TextureGraph/Source/TextureGraph/Private/Expressions/Output/TG_Expression_Output.cpp
+++ b/Engine/Plugins/Experimental/TextureGraph/Source/TextureGraph/Private/Expressions/Output/TG_Expression_Output.cpp
@@ -36,6 +36,10 @@ void UTG_Expression_Output::Evaluate(FTG_EvaluationContext* InContext)
 			BufferDescriptor OutputDesc = Output.EditTexture().GetBufferDescriptor();
 			const BufferDescriptor& InputDesc = Source.GetTexture()->GetDescriptor();
 
+			// When Output and Input desc don't match, it just ignores input's Desc sRGB value and uses Output desc value (which is currently always false) 
+			// Right now, the sRGB settings cannot be changed in the output node, so we can safely apply the input sRGB setting here.
+			OutputDesc.bIsSRGB = InputDesc.bIsSRGB;
+			
 			/// If the descriptors are not the same
 			if (OutputDesc != InputDesc)
 			{
