Fixed issue: The color node doesn't apply color on the OK button press in the dialog box

The issue occurred because the editor lost focus when we opened the color picker from the node UI. Viewport invalidation is happening from the tick function, which doesn't work until the editor is focused again by clicking on it. To fix that, we just need to invalidate the viewport whenever rendering is completed.

diff --git a/Engine/Plugins/Experimental/TextureGraph/Source/TextureGraphEditor/Private/TG_Editor.cpp b/Engine/Plugins/Experimental/TextureGraph/Source/TextureGraphEditor/Private/TG_Editor.cpp
--- a/Engine/Plugins/Experimental/TextureGraph/Source/TextureGraphEditor/Private/TG_Editor.cpp
+++ b/Engine/Plugins/Experimental/TextureGraph/Source/TextureGraphEditor/Private/TG_Editor.cpp
@@ -634,6 +634,7 @@ void FTG_Editor::OnRenderingDone(UMixInterface* TextureGraph, const FInvalidatio
 	if (TextureGraph != nullptr)
 	{
 		RefreshNodePreview(GraphEditorWidget->GetSelectedNodes(), Details, true);
+		RefreshViewport();
 	}
 }
 
@@ -1798,7 +1799,6 @@ void FTG_Editor::OnGraphChanged(UTG_Graph* InGraph, UTG_Node* InNode, bool Tweak
 		UpdateMixSettings();
 	}
 	RefreshViewport();
-	// OutputView->ForceRefresh();
 }
 
 void FTG_Editor::UpdateMixSettings()
