Fixed issue: Blend node clamp issue

Removed unnecessary saturate calls. Added a clamp property in the blend expression which is by default true.

diff --git a/Engine/Plugins/Experimental/TextureGraph/Shaders/Expressions/Expression_Blend.usf b/Engine/Plugins/Experimental/TextureGraph/Shaders/Expressions/Expression_Blend.usf
--- a/Engine/Plugins/Experimental/TextureGraph/Shaders/Expressions/Expression_Blend.usf
+++ b/Engine/Plugins/Experimental/TextureGraph/Shaders/Expressions/Expression_Blend.usf
@@ -31,6 +31,11 @@ float4 FSH_BlendNormal(in float2 uv : TEXCOORD0) : SV_Target0
 #else
 	float4 FinalColor = lerp(BackgroundTex, ForegroundTex, MaskingValue);
 #endif
+
+#if CLAMP
+	FinalColor = saturate(FinalColor);
+#endif
+	
     return FinalColor;
 }
 
@@ -46,6 +51,10 @@ float4 FSH_BlendAdd(in float2 uv : TEXCOORD0) : SV_Target0
 #else
 	float4 FinalColor = lerp(BackgroundTex, BackgroundTex + ForegroundTex, MaskingValue);
 #endif
+
+#if CLAMP
+	FinalColor = saturate(FinalColor);
+#endif
 	
 	return FinalColor;
 }
@@ -58,10 +67,14 @@ float4 FSH_BlendSubtract(in float2 uv : TEXCOORD0) : SV_Target0
 	float MaskingValue = GetMask(MaskTexture, uv, Opacity);
 
 #if IGNORE_ALPHA
-	float3 FinalRGB = lerp(BackgroundTex, saturate(BackgroundTex - ForegroundTex), MaskingValue);
+	float3 FinalRGB = lerp(BackgroundTex, BackgroundTex - ForegroundTex, MaskingValue);
 	float4 FinalColor = float4(FinalRGB, BackgroundTex.a);
 #else
-	float4 FinalColor = lerp(BackgroundTex, saturate(BackgroundTex - ForegroundTex), MaskingValue);
+	float4 FinalColor = lerp(BackgroundTex, BackgroundTex - ForegroundTex, MaskingValue);
+#endif
+
+#if CLAMP
+	FinalColor = saturate(FinalColor);
 #endif
 	
 	return FinalColor;
@@ -75,10 +88,14 @@ float4 FSH_BlendMultiply(in float2 uv : TEXCOORD0) : SV_Target0
 	float MaskingValue = GetMask(MaskTexture, uv, Opacity);
 
 #if IGNORE_ALPHA
-	float3 FinalRGB = lerp(BackgroundTex, saturate(ForegroundTex * BackgroundTex), MaskingValue);
+	float3 FinalRGB = lerp(BackgroundTex, ForegroundTex * BackgroundTex, MaskingValue);
 	float4 FinalColor = float4(FinalRGB, BackgroundTex.a);
 #else
-	float4 FinalColor = lerp(BackgroundTex, saturate(ForegroundTex * BackgroundTex), MaskingValue);
+	float4 FinalColor = lerp(BackgroundTex, ForegroundTex * BackgroundTex, MaskingValue);
+#endif
+
+#if CLAMP
+	FinalColor = saturate(FinalColor);
 #endif
 	
 	return FinalColor;
@@ -99,6 +116,10 @@ float4 FSH_BlendDivide(in float2 uv : TEXCOORD0) : SV_Target0
 #else
 	float4 FinalColor = lerp(BackgroundTex, DivisionResult, MaskingValue);
 #endif
+
+#if CLAMP
+	FinalColor = saturate(FinalColor);
+#endif
 	
 	return FinalColor;
 }
@@ -116,6 +137,10 @@ float4 FSH_BlendDifference(in float2 uv : TEXCOORD0) : SV_Target0
 #else
 	float4 FinalColor = lerp(BackgroundTex, abs(BackgroundTex - ForegroundTex), MaskingValue);
 #endif
+
+#if CLAMP
+	FinalColor = saturate(FinalColor);
+#endif
 	
 	return FinalColor;
 }
@@ -133,6 +158,11 @@ float4 FSH_BlendMax(in float2 uv : TEXCOORD0) : SV_Target0
 #else
 	float4 FinalColor = lerp(BackgroundTex, max(BackgroundTex, ForegroundTex), MaskingValue);
 #endif
+
+#if CLAMP
+	FinalColor = saturate(FinalColor);
+#endif
+	
 	return FinalColor;
 }
 
@@ -149,6 +179,10 @@ float4 FSH_BlendMin(in float2 uv : TEXCOORD0) : SV_Target0
 #else
 	float4 FinalColor = lerp(BackgroundTex, min(BackgroundTex, ForegroundTex), MaskingValue);
 #endif
+
+#if CLAMP
+	FinalColor = saturate(FinalColor);
+#endif
 	
 	return FinalColor;
 }
@@ -167,6 +201,10 @@ float4 FSH_BlendStep(in float2 uv : TEXCOORD0) : SV_Target0
 #else
 	float4 FinalColor = lerp(BackgroundTex, StepValue, MaskingValue);
 #endif
+
+#if CLAMP
+	FinalColor = saturate(FinalColor);
+#endif
 	
 	return FinalColor;
 }
@@ -195,6 +233,10 @@ float4 FSH_BlendOverlay(in float2 uv : TEXCOORD0) : SV_Target0
 	float4 FinalColor = lerp(BackgroundTex, OverlayBlend(BackgroundTex, ForegroundTex), MaskingValue);
 #endif
 
+#if CLAMP
+	FinalColor = saturate(FinalColor);
+#endif
+	
 	return FinalColor;
 }
 
diff --git a/Engine/Plugins/Experimental/TextureGraph/Source/TextureGraph/Private/Expressions/Maths/TG_Expression_Blend.cpp b/Engine/Plugins/Experimental/TextureGraph/Source/TextureGraph/Private/Expressions/Maths/TG_Expression_Blend.cpp
--- a/Engine/Plugins/Experimental/TextureGraph/Source/TextureGraph/Private/Expressions/Maths/TG_Expression_Blend.cpp
+++ b/Engine/Plugins/Experimental/TextureGraph/Source/TextureGraph/Private/Expressions/Maths/TG_Expression_Blend.cpp
@@ -32,6 +32,7 @@ void UTG_Expression_Blend::Evaluate(FTG_EvaluationContext* InContext)
 	BlendSettings.Mask = Mask;
 	BlendSettings.Opacity = Opacity;
 	BlendSettings.bIgnoreAlpha = bIgnoreAlpha;
+	BlendSettings.bClamp = bClamp;
 	
 	Output = T_Blend::Create(InContext->Cycle, Output.GetBufferDescriptor(), InContext->TargetId, BlendMode, &BlendSettings);
 }
diff --git a/Engine/Plugins/Experimental/TextureGraph/Source/TextureGraph/Public/Expressions/Maths/TG_Expression_Blend.h b/Engine/Plugins/Experimental/TextureGraph/Source/TextureGraph/Public/Expressions/Maths/TG_Expression_Blend.h
--- a/Engine/Plugins/Experimental/TextureGraph/Source/TextureGraph/Public/Expressions/Maths/TG_Expression_Blend.h
+++ b/Engine/Plugins/Experimental/TextureGraph/Source/TextureGraph/Public/Expressions/Maths/TG_Expression_Blend.h
@@ -29,6 +29,9 @@ public:
 	// If true, it will ignore the alpha channel blending and keep the background alpha.  
 	UPROPERTY(EditAnywhere, Category = NoCategory, meta = (TGType = "TG_Setting", DisplayName = "Ignore Alpha"))
 	bool bIgnoreAlpha = true;
+
+	UPROPERTY(EditAnywhere, Category = NoCategory, meta = (TGType = "TG_Setting", DisplayName = "Clamp"))
+	bool bClamp = true;
 	
 	UPROPERTY(EditAnywhere, Category = NoCategory, meta = (TGType = "TG_Output", PinDisplayName = ""))
 	FTG_Texture Output;
diff --git a/Engine/Plugins/Experimental/TextureGraph/Source/TextureGraphEngine/Transform/Expressions/T_Blend.cpp b/Engine/Plugins/Experimental/TextureGraph/Source/TextureGraphEngine/Transform/Expressions/T_Blend.cpp
--- a/Engine/Plugins/Experimental/TextureGraph/Source/TextureGraphEngine/Transform/Expressions/T_Blend.cpp
+++ b/Engine/Plugins/Experimental/TextureGraph/Source/TextureGraphEngine/Transform/Expressions/T_Blend.cpp
@@ -67,6 +67,7 @@ TiledBlobPtr CreateGenericBlend(MixUpdateCyclePtr InCycle, int32 InTargetId, Buf
 {
 	FSH_BlendBase::FPermutationDomain PermutationVector;
 	PermutationVector.Set<FSH_BlendBase::FIgnoreAlpha>(InBlendSettings->bIgnoreAlpha);
+	PermutationVector.Set<FSH_BlendBase::FClamp>(InBlendSettings->bClamp);
 	
 	const RenderMaterial_FXPtr RenderMaterial = TextureGraphEngine::GetMaterialManager()->CreateMaterial_FX<VSH_Simple, FSH_Type>(InTransformName, PermutationVector);
 	check(RenderMaterial);
@@ -98,6 +99,7 @@ TiledBlobPtr CreateGenericBlend(MixUpdateCyclePtr InCycle, int32 InTargetId, Buf
 		->AddArg(ARG_BLOB(MaskTexture, "MaskTexture"))
 		->AddArg(ARG_FLOAT(InBlendSettings->Opacity, "Opacity"))
 		->AddArg(WithUnbounded(ARG_BOOL(InBlendSettings->bIgnoreAlpha, "IgnoreAlpha")))
+		->AddArg(WithUnbounded(ARG_BOOL(InBlendSettings->bClamp, "Clamp")))
 		;
 
 	const FString Name = FString::Printf(TEXT("[%llu] - Blend - %s"), InCycle->GetBatch()->GetBatchId(), *InTransformName);
diff --git a/Engine/Plugins/Experimental/TextureGraph/Source/TextureGraphEngine/Transform/Expressions/T_Blend.h b/Engine/Plugins/Experimental/TextureGraph/Source/TextureGraphEngine/Transform/Expressions/T_Blend.h
--- a/Engine/Plugins/Experimental/TextureGraph/Source/TextureGraphEngine/Transform/Expressions/T_Blend.h
+++ b/Engine/Plugins/Experimental/TextureGraph/Source/TextureGraphEngine/Transform/Expressions/T_Blend.h
@@ -26,7 +26,8 @@ public:
 	END_SHADER_PARAMETER_STRUCT()
 
 	class FIgnoreAlpha : SHADER_PERMUTATION_BOOL("IGNORE_ALPHA");
-	using FPermutationDomain = TShaderPermutationDomain<FIgnoreAlpha>;
+	class FClamp : SHADER_PERMUTATION_BOOL("CLAMP");
+	using FPermutationDomain = TShaderPermutationDomain<FIgnoreAlpha, FClamp>;
 
 public:
 	static bool ShouldCompilePermutation(const FGlobalShaderPermutationParameters& Parameters)
@@ -163,6 +164,7 @@ public:
 		TiledBlobPtr Mask;
 		float Opacity;
 		bool bIgnoreAlpha;
+		bool bClamp;
 	};
 	//////////////////////////////////////////////////////////////////////////
 	/// Static functions
