Fix cyclic dependency for unsaved texture graphs

diff --git a/Engine/Plugins/Experimental/TextureGraph/Source/TextureGraph/Private/Expressions/Input/TG_Expression_Graph.cpp b/Engine/Plugins/Experimental/TextureGraph/Source/TextureGraph/Private/Expressions/Input/TG_Expression_Graph.cpp
--- a/Engine/Plugins/Experimental/TextureGraph/Source/TextureGraph/Private/Expressions/Input/TG_Expression_Graph.cpp
+++ b/Engine/Plugins/Experimental/TextureGraph/Source/TextureGraph/Private/Expressions/Input/TG_Expression_Graph.cpp
@@ -170,19 +170,53 @@ void UTG_Expression_TextureGraph::OnTextureGraphPreSave(UObject* Object, FObject
 {
 	// every editor should check if your texture graph is dependent on the object being saved
 	UTextureGraph* GraphBeingSaved = Cast<UTextureGraph>(Object);
-	// if object being saved is our linked TextureGraph, we re-create
-	if (GraphBeingSaved && TextureGraph)
+	UTextureGraph* ParentGraph = GetTypedOuter<UTextureGraph>();
+	// if object being saved is our linked TextureGraph, 
+	// and this expression's graph is NOT the same as the graph being saved, we re-create
+	if (GraphBeingSaved && TextureGraph && 
+		GraphBeingSaved != ParentGraph )
 	{
-		UTextureGraph* ParentGraph = CastChecked<UTextureGraph>(GetOutermostObject());
-
-		// update if the graph being saved is the one we're using in this expression directly
-		// Or our parent graph is indirectly dependent on the graph being saved 
-		if (GraphBeingSaved == TextureGraph ||
-			ParentGraph->IsDependent(GraphBeingSaved))
+		
+		check (ParentGraph);
+		// if graph being saved is not ourself, and there is a cyclic dependency
+		bool bSameGraphs = GraphBeingSaved->GetOuter() == TextureGraph->GetOuter();
+		if (bSameGraphs || ParentGraph->IsDependentOn(GraphBeingSaved))
 		{
-			// Recreate the Runtime copy
-			RuntimeGraph = DuplicateObject(TextureGraph->Graph(), this);	
+			// Check if the graph saved has caused cyclic dependency on the parent graph of this expression
+			if (!bSameGraphs && ParentGraph->HasCyclicDependency())
+			{
+				if (TextureGraphEngine::GetErrorReporter(ParentGraph) != nullptr)
+				{
+					auto ErrorType = static_cast<int32>(ETextureGraphErrorType::RECURSIVE_CALL);
+	
+					TextureGraphEngine::GetErrorReporter(ParentGraph)->ReportError(ErrorType, FString::Printf(TEXT("Can not use graph as a cyclic dependency is detected")));
+				}
+				Modify();
+				//we reset graph
+				TextureGraph = nullptr;
+				RuntimeGraph = nullptr;
+#if WITH_EDITOR
+				// We need to find its property and trigger property change event manually.
+				const auto SourcePin = GetParentNode()->GetInputPin(GET_MEMBER_NAME_CHECKED(UTG_Expression_TextureGraph, TextureGraph));
 
+				check(SourcePin)
+	
+				if(SourcePin)
+				{
+					FProperty* Property = SourcePin->GetExpressionProperty();
+					NotifyExpressionChanged(FPropertyChangedEvent(Property));
+				}
+#endif
+			}
+			else
+			{
+				// Recreate the Runtime copy
+				RuntimeGraph = DuplicateObject(TextureGraph->Graph(), this);	
+#ifdef UE_BUILD_DEBUG
+				RuntimeGraph->IsRuntime = 2;
+#endif
+			}
+			
 			// Graph is reset, notify the owning node / graph to update itself
 			NotifyGraphChanged();
 		}
@@ -193,6 +227,9 @@ void UTG_Expression_TextureGraph::Initialize()
 	if (TextureGraph && !RuntimeGraph)
 	{
 		RuntimeGraph = DuplicateObject(TextureGraph->Graph(), this);
+#ifdef UE_BUILD_DEBUG
+		RuntimeGraph->IsRuntime = 2;
+#endif
 	}
 
 }
@@ -218,7 +255,7 @@ bool UTG_Expression_TextureGraph::CheckDependencies(const UTextureGraph* InTextu
 		return false;
 	}
 
-	if (InTextureGraph->IsDependent(ParentGraph))
+	if (InTextureGraph->IsDependentOn(ParentGraph))
 	{
 		auto ErrorType = static_cast<int32>(ETextureGraphErrorType::RECURSIVE_CALL);
 	
@@ -233,10 +270,8 @@ void UTG_Expression_TextureGraph::SetTextureGraphInternal(UTextureGraph* InTextu
 {
 	// validate that we're allowed to use InTextureGraph as input here
 	// if dependencies exist, we don't want to use the InTextureGraph so we ll treat it just like a nullptr.
-	// TODO: At the moment we reset TextureGraph to nullptr when a dependency check fails.
-	// Ideally we'd like to revert to the previous valid value. The problem here is that
-	// TextureGraph at this point comes pre-filled with the problematic value 
 	bool IsValidTextureGraph = CheckDependencies(InTextureGraph);
+
 	if (!IsValidTextureGraph)
 	{
 		return;
@@ -253,6 +288,9 @@ void UTG_Expression_TextureGraph::SetTextureGraphInternal(UTextureGraph* InTextu
 		// This is a valid TextureGraph
 		TextureGraph = InTextureGraph;
 		RuntimeGraph = DuplicateObject(TextureGraph->Graph(), this);
+#ifdef UE_BUILD_DEBUG
+		RuntimeGraph->IsRuntime = 2;
+#endif
 	}
 
 	// Graph is reset, notify the owning node / graph to update itself
@@ -315,11 +353,14 @@ void UTG_Expression_TextureGraph::Evaluate(FTG_EvaluationContext* InContext)
 bool UTG_Expression_TextureGraph::Validate(MixUpdateCyclePtr Cycle)
 {
 	// if TextureGraph is already in execution list then
-	if(TextureGraph &&  Cycle->ContainsMix(TextureGraph))
+	UTextureGraph* ParentGraph = CastChecked<UTextureGraph>(GetOutermostObject());
+
+	if((TextureGraph &&  Cycle->ContainsMix(TextureGraph)) ||
+		(ParentGraph && ParentGraph->HasCyclicDependency()))
 	{
 		auto ErrorType = static_cast<int32>(ETextureGraphErrorType::RECURSIVE_CALL);
 		
-		const FString ErrorMsg = FString::Format(TEXT("Recursive Texture Graph call found in '{0}' asset."), {Cycle->TopMix()->GetName()} );
+		const FString ErrorMsg = FString::Format(TEXT("Recursive Texture Graph call found in '{0}' asset."), {TextureGraph->GetName()} );
 		TextureGraphEngine::GetErrorReporter(Cycle->GetMix())->ReportError(ErrorType, ErrorMsg, GetParentNode());
 		return false;
 	}
diff --git a/Engine/Plugins/Experimental/TextureGraph/Source/TextureGraph/Private/TextureGraph.cpp b/Engine/Plugins/Experimental/TextureGraph/Source/TextureGraph/Private/TextureGraph.cpp
--- a/Engine/Plugins/Experimental/TextureGraph/Source/TextureGraph/Private/TextureGraph.cpp
+++ b/Engine/Plugins/Experimental/TextureGraph/Source/TextureGraph/Private/TextureGraph.cpp
@@ -21,7 +21,7 @@
 #include "Transform/Mix/T_UpdateTargets.h"
 
 
-bool UTextureGraph::CheckCyclicDependency(const UTextureGraph* InTextureGraph) const
+bool UTextureGraph::CheckRecursiveDependency(const UTextureGraph* InTextureGraph) const
 {
 	TArray<UTextureGraph*> DependentGraphs;
 	GatherAllDependentGraphs(DependentGraphs);
@@ -36,6 +36,10 @@ bool UTextureGraph::CheckCyclicDependency(const UTextureGraph* InTextureGraph) c
 				(Package == SecondPackage && !bIsTransientPackage);
 	});
 }
+bool UTextureGraph::HasCyclicDependency() const
+{
+	return CheckRecursiveDependency(this);
+}
 void UTextureGraph::GatherAllDependentGraphs(TArray<UTextureGraph*>& DependentGraphs) const
 {
 	Graph()->ForEachNodes(
@@ -60,7 +64,7 @@ void UTextureGraph::GatherAllDependentGraphs(TArray<UTextureGraph*>& DependentGr
 				}
 		});
 }
-bool UTextureGraph::IsDependent(const UTextureGraph* InTextureGraph) const
+bool UTextureGraph::IsDependentOn(const UTextureGraph* InTextureGraph) const
 {
 	
 	// check if we're trying to assign our own TextureGraph to this expression
@@ -71,7 +75,7 @@ bool UTextureGraph::IsDependent(const UTextureGraph* InTextureGraph) const
 	}
 			
 	// check for cyclic dependency
-	if (CheckCyclicDependency(InTextureGraph))
+	if (CheckRecursiveDependency(InTextureGraph))
 	{
 		return true;
 	}
diff --git a/Engine/Plugins/Experimental/TextureGraph/Source/TextureGraph/Public/TG_Graph.h b/Engine/Plugins/Experimental/TextureGraph/Source/TextureGraph/Public/TG_Graph.h
--- a/Engine/Plugins/Experimental/TextureGraph/Source/TextureGraph/Public/TG_Graph.h
+++ b/Engine/Plugins/Experimental/TextureGraph/Source/TextureGraph/Public/TG_Graph.h
@@ -289,8 +289,6 @@ protected:
 	void EvalInOutPins() const;
 	void EvalTraverseOrder() const;
 
-	bool IsDependentInternal(const UTG_Graph* SourceGraph, TArray<UTG_Graph*>& DependentGraphs);
-
 #if WITH_EDITORONLY_DATA
 	// Extra data to hold information that is useful only in editor (like comments)
 	UPROPERTY()
@@ -328,6 +326,13 @@ public:
 	static FString LogCall(const TArray<FTG_Id>& PinInputs, const  TArray<FTG_Id>& PinOutputs, int32 InputLogWidth = 10);
 
 	//// Testing zone
+	
+#ifdef UE_BUILD_DEBUG
+	// 0 = DiskVersion, 
+	// 1 = Runtime copy made by TG_Editor, 
+	// 2 = Runtime copy made by TG_Expression_Graph
+	int                 IsRuntime = 0;
+#endif
 private:
 };
 
diff --git a/Engine/Plugins/Experimental/TextureGraph/Source/TextureGraph/Public/TextureGraph.h b/Engine/Plugins/Experimental/TextureGraph/Source/TextureGraph/Public/TextureGraph.h
--- a/Engine/Plugins/Experimental/TextureGraph/Source/TextureGraph/Public/TextureGraph.h
+++ b/Engine/Plugins/Experimental/TextureGraph/Source/TextureGraph/Public/TextureGraph.h
@@ -18,12 +18,14 @@ protected:
 	UPROPERTY(EditAnywhere, Category="TG_GraphParams")
 	TObjectPtr<class UTG_Graph> TextureGraph;
 
-	bool CheckCyclicDependency(const UTextureGraph* InTextureGraph) const;
 	void GatherAllDependentGraphs(TArray<UTextureGraph*>& DependentGraphs) const;
 	
+	bool CheckRecursiveDependency(const UTextureGraph* InTextureGraph) const;
+	
 public:
 
-	bool IsDependent(const UTextureGraph* TextureGraph) const;
+	bool IsDependentOn(const UTextureGraph* TextureGraph) const;
+	bool HasCyclicDependency() const;
 	
 	// Construct the script giving it its name Initialize to a default one output script
 	virtual void Construct(FString Name);
diff --git a/Engine/Plugins/Experimental/TextureGraph/Source/TextureGraphEditor/Private/TG_Editor.cpp b/Engine/Plugins/Experimental/TextureGraph/Source/TextureGraphEditor/Private/TG_Editor.cpp
--- a/Engine/Plugins/Experimental/TextureGraph/Source/TextureGraphEditor/Private/TG_Editor.cpp
+++ b/Engine/Plugins/Experimental/TextureGraph/Source/TextureGraphEditor/Private/TG_Editor.cpp
@@ -178,10 +178,14 @@ void FTG_Editor::InitEditor(const EToolkitMode::Type Mode, const TSharedPtr< cla
 	// Propagate all object flags except for RF_Standalone, otherwise the preview material won't GC once
 	// the TS editor releases the reference.
 	// overwrite the original TextureGraph in place by constructing a new one with the same name
-	FObjectDuplicationParameters Params = InitStaticDuplicateObjectParams(OriginalTextureGraph, OriginalTextureGraph->GetOuter(), NAME_None,
+	FName RuntimeGraphName = FName(OriginalTextureGraph->GetFName().ToString() + TEXT("_Runtime"));
+	FObjectDuplicationParameters Params = InitStaticDuplicateObjectParams(OriginalTextureGraph, OriginalTextureGraph->GetOuter(), RuntimeGraphName,
 	~RF_Standalone, UTextureGraph::StaticClass(), EDuplicateMode::Normal, EInternalObjectFlags::None);
 	
 	EditedTextureGraph = Cast<UTextureGraph>(StaticDuplicateObjectEx(Params));
+#ifdef UE_BUILD_DEBUG
+	EditedTextureGraph->Graph()->IsRuntime = 1;
+#endif
 	FCoreUObjectDelegates::OnObjectPreSave.AddSP(this, &FTG_Editor::OnTextureGraphPreSave);
 
 	//Editor gets notified when rendering is done
@@ -1601,8 +1605,20 @@ void FTG_Editor::OnClose()
 
 bool FTG_Editor::UpdateOriginalTextureGraph()
 {
-	// TODO : We should cancel saving when TextureGraph has errors.
-
+	// We should cancel saving when TextureGraph has certain errors.
+	FTextureGraphErrorReporter* ErrorReporter = TextureGraphEngine::GetErrorReporter(EditedTextureGraph);
+	if (ErrorReporter)
+	{
+		// see if there are recursive call errors, we warn user and not let them save in that state
+		const auto CompileErrors = TextureGraphEngine::GetErrorReporter(EditedTextureGraph)->GetCompilationErrors();
+		const TArray<FTextureGraphErrorReport>* RecursiveErrors = CompileErrors.Find((int)ETextureGraphErrorType::RECURSIVE_CALL);
+		if (RecursiveErrors != nullptr && RecursiveErrors->Num() > 0)
+		{
+			const FText Message =  FText::FromString(TEXT("Cannot save as Recursive Errors found.\r\n") + (*RecursiveErrors)[0].ErrorMsg);
+			FMessageDialog::Open(EAppMsgCategory::Error, EAppMsgType::Ok, Message);
+			return false;
+		}
+	}
 	if (EditedTextureGraph->GetPackage()->IsDirty())
 	{
 		// Cache any metadata
@@ -1614,7 +1630,9 @@ bool FTG_Editor::UpdateOriginalTextureGraph()
 
 		// UObject* NewAsset = StaticDuplicateObjectEx(Params);
 		OriginalTextureGraph = Cast<UTextureGraph>(StaticDuplicateObjectEx(Params));
-
+#ifdef UE_BUILD_DEBUG
+		OriginalTextureGraph->Graph()->IsRuntime = 0;
+#endif
 		// Restore the metadata
 		if (MetaData)
 		{
