Fixing issue with material parameters having the same name.

We introduce a runtime array mapping the exposed arguments of the Node as pins to the corresponding MaterialParameter.
we make sure that the exposed arguments all have different name since this is a requirement from TG. THe table allows to have a different name for the exposed arguments of a material parameter which would already be named the same.

We also tried to clean up the Material parameter diiscovery phase that generate the signature for the Material Expression. It now follows the order of typesas in UMaterialInterface.

diff --git a/Engine/Plugins/Experimental/TextureGraph/Source/TextureGraph/Private/Expressions/TG_Expression_MaterialBase.cpp b/Engine/Plugins/Experimental/TextureGraph/Source/TextureGraph/Private/Expressions/TG_Expression_MaterialBase.cpp
--- a/Engine/Plugins/Experimental/TextureGraph/Source/TextureGraph/Private/Expressions/TG_Expression_MaterialBase.cpp
+++ b/Engine/Plugins/Experimental/TextureGraph/Source/TextureGraph/Private/Expressions/TG_Expression_MaterialBase.cpp
@@ -134,6 +134,15 @@ void UTG_Expression_MaterialBase::GenerateMaterialAttributeOptions()
 	}
 }
 
+void UTG_Expression_MaterialBase::Initialize()
+{
+	if(GetMaterial() && !MaterialInstance)
+	{
+		MaterialInstance = UMaterialInstanceDynamic::Create(GetMaterial(), this);
+	}
+	GenerateMaterialAttributeOptions(); // also populate the attributes availables
+}
+
 void UTG_Expression_MaterialBase::Evaluate(FTG_EvaluationContext* InContext)
 {
 	Super::Evaluate(InContext);
@@ -153,15 +162,6 @@ void UTG_Expression_MaterialBase::Evaluate(FTG_EvaluationContext* InContext)
 	Output = Result;
 }
 
-void UTG_Expression_MaterialBase::Initialize()
-{
-	if(GetMaterial() && !MaterialInstance)
-	{
-		MaterialInstance = UMaterialInstanceDynamic::Create(GetMaterial(), this);
-	}
-	GenerateMaterialAttributeOptions(); // also populate the attributes availables
-}
-
 TiledBlobPtr UTG_Expression_MaterialBase::CreateRenderMaterialJob(FTG_EvaluationContext* InContext, const FString& InName, const FString& InMaterialPath, const BufferDescriptor& InDescriptor, EDrawMaterialAttributeTarget InDrawMaterialAttributeTarget)
 {
 	const RenderMaterial_BPPtr RenderMaterial = TextureGraphEngine::GetMaterialManager()->CreateMaterial_BP(InName, InMaterialPath);
@@ -247,129 +247,136 @@ void UTG_Expression_MaterialBase::LinkMaterialParameters(FTG_EvaluationContext*
 {
 	if (InMaterialJob)
 	{	
-		TArray<FMaterialParameterInfo> OutParameterInfo;
-		TArray<FGuid> OutParameterIds;
-
-		InMaterial->GetAllScalarParameterInfo(OutParameterInfo, OutParameterIds);
-
-		for (auto ParameterInfo : OutParameterInfo)
+		for (auto& ArgToMatParam : ArgToMatParams)
 		{
-			FTG_Var* Var = InContext->Inputs.GetVar(ParameterInfo.Name);
-			const FTG_Argument* VarArgument = InContext->Inputs.GetVarArgument(ParameterInfo.Name);
-			
+			FTG_Var* Var = InContext->Inputs.GetVar(ArgToMatParam.ArgName);
 			if (Var && !Var->IsEmpty())
 			{
-				FName CPPType = VarArgument->GetCPPTypeName();
-
-				float ParamValue;
-				
-				if (CPPType == TEXT("int32"))
+				switch (ArgToMatParam.MatType)
+				{
+				case EMaterialParameterType::Scalar:
 				{
-					ParamValue = Var->GetAs<int32>();
+					const FTG_Argument* VarArgument = InContext->Inputs.GetVarArgument(ArgToMatParam.ArgName);
+					FName CPPType = VarArgument->GetCPPTypeName();
+					float ParamValue;
+					if (CPPType == TEXT("int32"))
+					{
+						ParamValue = Var->GetAs<int32>();
+					}
+					else if (CPPType == TEXT("uint32"))
+					{
+						ParamValue = Var->GetAs<uint32>();
+					}
+					else
+					{
+						ParamValue = Var->GetAs<float>();
+					}
+
+					InMaterialJob->AddArg(ARG_FLOAT(ParamValue, TCHAR_TO_UTF8(*ArgToMatParam.MatParamName.ToString())));
 				}
-				else if (CPPType == TEXT("uint32"))
+				break;
+				case EMaterialParameterType::Vector:
 				{
-					ParamValue = Var->GetAs<uint32>();
+					auto ParamValue = Var->EditAs<FLinearColor>();
+					InMaterialJob->AddArg(ARG_VECTOR(ParamValue, TCHAR_TO_UTF8(*ArgToMatParam.MatParamName.ToString())));
 				}
-				else
+				break;
+				case EMaterialParameterType::DoubleVector:
 				{
-					ParamValue = Var->GetAs<float>();
+					/*
+					auto ParamValue = Var->EditAs<FVector4d>();
+					Job->AddArg(ARG_VECTOR4(ParamValue, TCHAR_TO_UTF8(*ArgToMatParam.MatParamName.ToString())));
+					*/
 				}
-				
-				InMaterialJob->AddArg(ARG_FLOAT(ParamValue, TCHAR_TO_UTF8(*ParameterInfo.Name.ToString())));
-			}
-		}
-		
-		InMaterial->GetAllTextureParameterInfo(OutParameterInfo, OutParameterIds);
+				break;
+				case EMaterialParameterType::Texture:
+				{
+					FTG_Texture& ParamValue = Var->EditAs<FTG_Texture>();
 
-		for (auto ParameterInfo : OutParameterInfo)
-		{
-			FTG_Var* Var = InContext->Inputs.GetVar(ParameterInfo.Name);
+					// there could be a case where the var has an empty blob, we don't need to do anything
+					// we fallback to the default internal material parameter value for this input pin
+					if (ParamValue.RasterBlob)
+					{
+						auto CombinedBlob = T_CombineTiledBlob::Create(InContext->Cycle, ParamValue.GetBufferDescriptor(), 0, ParamValue.RasterBlob);
 
-			if (Var && !Var->IsEmpty())
-			{
-				FTG_Texture& ParamValue = Var->EditAs<FTG_Texture>();
+						auto ArgBlob = ARG_BLOB(CombinedBlob, TCHAR_TO_UTF8(*ArgToMatParam.MatParamName.ToString()));
 
-				// there could be a case where the var has an empty blob, we don't need to do anything
-				// we fallback to the default internal material parameter value for this input pin
-				if (ParamValue.RasterBlob)
-				{
-					auto CombinedBlob = T_CombineTiledBlob::Create(InContext->Cycle, ParamValue.GetBufferDescriptor(), 0, ParamValue.RasterBlob);
+						ArgBlob->SetHandleTiles(TiledMode);
 
-					auto ArgBlob = ARG_BLOB(CombinedBlob, TCHAR_TO_UTF8(*ParameterInfo.Name.ToString()));
-	
-					ArgBlob->SetHandleTiles(TiledMode);
-				
-					InMaterialJob->AddArg(ArgBlob);
+						InMaterialJob->AddArg(ArgBlob);
+					}
 				}
-			}
-		}
-
-		InMaterial->GetAllVectorParameterInfo(OutParameterInfo, OutParameterIds);
-
-		for (auto ParameterInfo : OutParameterInfo)
-		{
-			FTG_Var* Var = InContext->Inputs.GetVar(ParameterInfo.Name);
-			
-			if (Var && !Var->IsEmpty())
-			{
-				auto ParamValue = Var->EditAs<FLinearColor>();
-				InMaterialJob->AddArg(ARG_VECTOR(ParamValue, TCHAR_TO_UTF8(*ParameterInfo.Name.ToString())));
-			}
-		}
+				break;
+				case EMaterialParameterType::StaticSwitch:
+				{
 #if WITH_EDITORONLY_DATA
-		InMaterial->GetAllStaticSwitchParameterInfo(OutParameterInfo, OutParameterIds);
-		for (auto ParameterInfo : OutParameterInfo)
-		{
-			FTG_Var* Var = InContext->Inputs.GetVar(ParameterInfo.Name);
-			
-			if (Var && !Var->IsEmpty())
-			{
-				auto ParamValue = Var->EditAs<bool>();
-				InMaterialJob->AddArg(ARG_INT(ParamValue, TCHAR_TO_UTF8(*ParameterInfo.Name.ToString())));
-
+					auto ParamValue = Var->EditAs<bool>();
+					InMaterialJob->AddArg(ARG_INT(ParamValue, TCHAR_TO_UTF8(*ArgToMatParam.MatParamName.ToString())));
+#endif
+				}
+				break;
+				default:
+					break;
+				}
 			}
 		}
-#endif
-		//TODO: We need to add double Vector
-		/*InMaterial->GetAllDoubleVectorParameterInfo(OutParameterInfo, OutParameterIds);
-
-		for (auto ParameterInfo : OutParameterInfo)
-		{
-			FTG_Var* Var = InContext->Inputs.GetVar(ParameterInfo.Name);
+	}
+}
 
-			if (Var)
-			{
-				auto ParamValue = Var->EditAs<FVector4d>();
-				Job->AddArg(ARG_VECTOR4(ParamValue, TCHAR_TO_UTF8(*ParameterInfo.Name.ToString())));
-			}
-		}*/
+FName UTG_Expression_MaterialBase::CPPTypeNameFromMaterialParamType(EMaterialParameterType InMatType)
+{
+	switch (InMatType)
+	{
+	case EMaterialParameterType::Scalar:
+		return TEXT("float");
+	case EMaterialParameterType::Vector:
+		return TEXT("FLinearColor");
+	case EMaterialParameterType::DoubleVector:
+		return TEXT("FVector4");
+	case EMaterialParameterType::Texture:
+		return TEXT("FTG_Texture");
+	case EMaterialParameterType::StaticSwitch:
+		return TEXT("bool");
+	default:
+		return FName();
 	}
 }
 
-void UTG_Expression_MaterialBase::AddSignatureParam(TArray<FMaterialParameterInfo> OutParameterInfo, FName CPPTypeName,  
-                                                         FTG_Signature::FInit& SignatureInit, bool IsScalar /*= false*/) const
+void UTG_Expression_MaterialBase::AddSignatureParam(const TArray<FMaterialParameterInfo>& OutParameterInfo, const TArray<FGuid>& OutGuids,
+		EMaterialParameterType MatType, FTG_Signature::FInit& SignatureInit) const
 {
-	for (FMaterialParameterInfo parameter : OutParameterInfo)
+	for (int i = 0; i < OutParameterInfo.Num(); ++i)
 	{
+		const FMaterialParameterInfo& MatParam = OutParameterInfo[i];
+		const FGuid& Guid = OutGuids[i];
+
 		TMap<FName, FString> MetaDataMap;
 #if WITH_EDITOR
-		if(IsScalar)
+		if(MatType == EMaterialParameterType::Scalar)
 		{
 			// In case of scalar parameter, material has input range. We use that range to handle the input in the node.
 			float MinValue, MaxValue;
-
-			if(GetMaterial()->GetScalarParameterSliderMinMax(parameter.Name, MinValue, MaxValue))
+			if(GetMaterial()->GetScalarParameterSliderMinMax(MatParam.Name, MinValue, MaxValue))
 			{
 				MetaDataMap.Add("MinValue", FString::SanitizeFloat(MinValue));
 				MetaDataMap.Add("MaxValue", FString::SanitizeFloat(MaxValue));
 			}
 		}
-#endif		
-		FTG_Argument Arg{ parameter.Name, CPPTypeName, { ETG_Access::In }, MetaDataMap };
+#endif	
+
+		// Check that no other arg have the  same name in the signature already
+		// If it's the case, postfix that arg name from the point of view of the signature
+		TArray<FName> ArgNames = TG_MakeArrayOfArgumentNames(SignatureInit.Arguments);
+		FName ArgName = TG_MakeNameUniqueInCollection(MatParam.Name, ArgNames);
+	
+		// New Arg of the signature
+		FTG_Argument Arg{ ArgName, CPPTypeNameFromMaterialParamType(MatType), {ETG_Access::In}, MetaDataMap};
 		
 		Arg.SetPersistentSelfVar(); // Set the material parameter persistent selfvar in order to save the state
 		SignatureInit.Arguments.Add(Arg);
+
+		// And new entry in ArgtoMatParams
+		ArgToMatParams.Add({ ArgName, MatParam.Name, Guid, MatType });
 	}
 }
 
@@ -377,31 +384,35 @@ FTG_SignaturePtr UTG_Expression_MaterialBase::BuildSignatureDynamically() const
 {
 	FTG_Signature::FInit SignatureInit = GetSignatureInitArgsFromClass();
 
+	// the Arg to Material Param array is populated along with the signature
+	// start fresh here
+	ArgToMatParams.Empty();
+
+
 	// append
 	if (GetMaterial())
 	{
 		TArray<FMaterialParameterInfo> OutParameterInfo;
 		TArray<FGuid> OutParameterIds;
 
-		// The texture param are declared as FTG_Texture so they can be connected from the standard nodes
-		GetMaterial()->GetAllTextureParameterInfo(OutParameterInfo, OutParameterIds);
-		AddSignatureParam(OutParameterInfo, TEXT("FTG_Texture"), SignatureInit);
-		
-#if WITH_EDITORONLY_DATA
-		GetMaterial()->GetAllStaticSwitchParameterInfo(OutParameterInfo, OutParameterIds);
-		AddSignatureParam(OutParameterInfo, TEXT("bool"), SignatureInit);
-#endif
 		GetMaterial()->GetAllScalarParameterInfo(OutParameterInfo, OutParameterIds);
-		AddSignatureParam(OutParameterInfo, TEXT("float"), SignatureInit, true);
+		AddSignatureParam(OutParameterInfo, OutParameterIds, EMaterialParameterType::Scalar, SignatureInit);
 
 		GetMaterial()->GetAllVectorParameterInfo(OutParameterInfo, OutParameterIds);
-		AddSignatureParam(OutParameterInfo, TEXT("FLinearColor"), SignatureInit);
+		AddSignatureParam(OutParameterInfo, OutParameterIds, EMaterialParameterType::Vector, SignatureInit);
 
 		// TODO: We need to add Double Vector4 support.
-		/*Material->GetAllDoubleVectorParameterInfo(OutParameterInfo, OutParameterIds);
-		AddSignatureParam(OutParameterInfo, TEXT("FVector4"), SignatureInit);*/
+		//GetMaterial()->GetAllDoubleVectorParameterInfo(OutParameterInfo, OutParameterIds);
+		//AddSignatureParam(OutParameterInfo, OutParameterIds, EMaterialParameterType::DoubleVector, SignatureInit);
 
+		// The texture param are declared as FTG_Texture so they can be connected from the standard nodes
+		GetMaterial()->GetAllTextureParameterInfo(OutParameterInfo, OutParameterIds);
+		AddSignatureParam(OutParameterInfo, OutParameterIds, EMaterialParameterType::Texture, SignatureInit);
 		
+#if WITH_EDITORONLY_DATA
+		GetMaterial()->GetAllStaticSwitchParameterInfo(OutParameterInfo, OutParameterIds);
+		AddSignatureParam(OutParameterInfo, OutParameterIds, EMaterialParameterType::StaticSwitch, SignatureInit);
+#endif		
 	}
 	
 	return MakeShared<FTG_Signature>(SignatureInit);
@@ -410,26 +421,12 @@ FTG_SignaturePtr UTG_Expression_MaterialBase::BuildSignatureDynamically() const
 
 void UTG_Expression_MaterialBase::CopyVarGeneric(const FTG_Argument& Arg, FTG_Var* InVar, bool CopyVarToArg)
 {
-	if (MaterialInstance)
+	FArgToMaterialParamInfo* ArgToMatParam = ArgToMatParams.FindByKey(Arg.Name);
+	if (MaterialInstance && ArgToMatParam)
 	{
-		// Try to find the MatParam matching the arg
-		if (Arg.CPPTypeName.IsEqual(TEXT("bool")))
+		switch (ArgToMatParam->MatType)
 		{
-	#if WITH_EDITORONLY_DATA
-			if (CopyVarToArg)
-			{
-				// Disable assigning static bool values to MID for now see UE-209533 & UE-219306
-				// We need to use a different solution for Material with Static Switch, we can't use MID for this, evetnually MIC editor only...
-				//MaterialInstance->SetStaticSwitchParameterValueEditorOnly(Arg.GetName(), InVar->GetAs<bool>());
-			} else
-			{
-				FMaterialParameterInfo ParameterInfo(Arg.GetName());
-				FGuid OutParameterId;
-				MaterialInstance->GetStaticSwitchParameterValue(ParameterInfo, InVar->EditAs<bool>(), OutParameterId);
-			}
-	#endif
-		}
-		else if (Arg.CPPTypeName.IsEqual(TEXT("float")))
+		case EMaterialParameterType::Scalar:
 		{
 			if (CopyVarToArg)
 			{
@@ -441,39 +438,68 @@ void UTG_Expression_MaterialBase::CopyVarGeneric(const FTG_Argument& Arg, FTG_Va
 				MaterialInstance->GetScalarParameterValue(ParameterInfo, InVar->EditAs<float>());
 			}
 		}
-		else if (Arg.CPPTypeName.IsEqual(TEXT("FLinearColor")))
+		break;
+		case EMaterialParameterType::Vector:
 		{
 			if (CopyVarToArg)
 			{
 				MaterialInstance->SetVectorParameterValue(Arg.GetName(), InVar->GetAs<FLinearColor>());
-			} else
+			}
+			else
 			{
 				FMaterialParameterInfo ParameterInfo(Arg.GetName());
 				MaterialInstance->GetVectorParameterValue(ParameterInfo, InVar->EditAs<FLinearColor>());
 			}
 		}
-		/*else if (Arg.CPPTypeName.IsEqual(TEXT("FVector4")))
+		break;
+		case EMaterialParameterType::DoubleVector:
 		{
-			if (CopyVarToArg)
+			// TODO: Support that case later ?
+			/*if (CopyVarToArg)
 			{
-				MaterialInstance->SetDoubleVectorParameterValue(Arg.GetName(), InVar->GetAs<FVector4>());
-			} else
+				MaterialInstance->SetDoubleVectorParameterValue(Arg.GetName(), InVar->GetAs<FVector4d>());
+			}
+			else
 			{
 				FMaterialParameterInfo ParameterInfo(Arg.GetName());
-				MaterialInstance->GetDoubleVectorParameterValue(ParameterInfo, InVar->EditAs<FVector4>());
-			}
-		}*/
-		else if (Arg.CPPTypeName.IsEqual(TEXT("FTG_Texture")))
+				MaterialInstance->GetDoubleVectorParameterValue(ParameterInfo, InVar->EditAs<FVector4d>());
+			}*/
+		}
+		break;
+		case EMaterialParameterType::Texture:
 		{
 			// TODO: Support that case later ?
 			/*	if (CopyVarToArg)
-				{
-					MaterialInstance->SetScalarParameterValue(Arg.GetName(), InVar->GetAs<float>());
-				} else
-				{
-					FMaterialParameterInfo ParameterInfo(Arg.GetName());
-					MaterialInstance->GetScalarParameterValue(ParameterInfo, InVar->EditAs<float>());
-				}*/
+			{
+				MaterialInstance->SetScalarParameterValue(Arg.GetName(), InVar->GetAs<float>());
+			}
+			else
+			{
+				FMaterialParameterInfo ParameterInfo(Arg.GetName());
+				MaterialInstance->GetScalarParameterValue(ParameterInfo, InVar->EditAs<float>());
+			}*/
+		}
+		break;
+		case EMaterialParameterType::StaticSwitch:
+		{
+#if WITH_EDITORONLY_DATA
+			if (CopyVarToArg)
+			{
+				// Disable assigning static bool values to MID for now see UE-209533 & UE-219306
+				// We need to use a different solution for Material with Static Switch, we can't use MID for this, evetnually MIC editor only...
+				//MaterialInstance->SetStaticSwitchParameterValueEditorOnly(Arg.GetName(), InVar->GetAs<bool>());
+			}
+			else
+			{
+				FMaterialParameterInfo ParameterInfo(Arg.GetName());
+				FGuid OutParameterId;
+				MaterialInstance->GetStaticSwitchParameterValue(ParameterInfo, InVar->EditAs<bool>(), OutParameterId);
+			}
+#endif
+		}
+		break;
+		default:
+			break;
 		}
 	}
 }
\ No newline at end of file
diff --git a/Engine/Plugins/Experimental/TextureGraph/Source/TextureGraph/Private/TG_Signature.cpp b/Engine/Plugins/Experimental/TextureGraph/Source/TextureGraph/Private/TG_Signature.cpp
--- a/Engine/Plugins/Experimental/TextureGraph/Source/TextureGraph/Private/TG_Signature.cpp
+++ b/Engine/Plugins/Experimental/TextureGraph/Source/TextureGraph/Private/TG_Signature.cpp
@@ -84,6 +84,17 @@ FTG_Hash FTG_Argument::Hash(const FTG_Argument& Argument)
 	return TG_Hash(v);
 }
 
+TArray<FName> TG_MakeArrayOfArgumentNames(const FTG_Arguments& InArguments)
+{
+	TArray<FName> Names;
+	Names.Empty(InArguments.Num());
+	for (auto& Arg : InArguments)
+	{
+		Names.Add(Arg.Name);
+	}
+	return Names;
+}
+
 FTG_Hash FTG_ArgumentSet::Hash(const FTG_ArgumentSet& Set)
 {
 	FTG_Hash v = 0;
diff --git a/Engine/Plugins/Experimental/TextureGraph/Source/TextureGraph/Public/Expressions/TG_Expression_MaterialBase.h b/Engine/Plugins/Experimental/TextureGraph/Source/TextureGraph/Public/Expressions/TG_Expression_MaterialBase.h
--- a/Engine/Plugins/Experimental/TextureGraph/Source/TextureGraph/Public/Expressions/TG_Expression_MaterialBase.h
+++ b/Engine/Plugins/Experimental/TextureGraph/Source/TextureGraph/Public/Expressions/TG_Expression_MaterialBase.h
@@ -61,6 +61,8 @@ public:
 	const TArray<EDrawMaterialAttributeTarget>& GetAvailableMaterialAttributeIds() const { return AvailableMaterialAttributeIds; }
 	const TArray<FName>& GetAvailableMaterialAttributeNames() const { return AvailableMaterialAttributeNames; }
 
+	static FName CPPTypeNameFromMaterialParamType(EMaterialParameterType InMatType);
+
 protected:
 	// A local per Instance Material is recreated from the reference material assigned through SetMaterialInternal
 	UPROPERTY(Transient, DuplicateTransient)
@@ -69,6 +71,18 @@ protected:
 	TArray<EDrawMaterialAttributeTarget>	AvailableMaterialAttributeIds; // The set of material properties available for rendering
 	TArray<FName>							AvailableMaterialAttributeNames; // same with the attribute names
 
+	// Arg to Material Param Info array records the map of arg name to the corresponding Material Param
+	struct FArgToMaterialParamInfo
+	{
+		FName ArgName;
+		FName MatParamName;
+		FGuid MatParamGuid;
+		EMaterialParameterType MatType = EMaterialParameterType::None;
+
+		bool operator== (const FName& InArgName) const { return ArgName == InArgName; }
+	};
+	mutable TArray<FArgToMaterialParamInfo> ArgToMatParams; // mutable because it is populated in the BuildSignatureDynamically()
+
 	virtual void Initialize() override;
 
 	void GenerateMaterialAttributeOptions(); // Based on the current Material, list of material attributes available
@@ -84,7 +98,7 @@ protected:
 	virtual TObjectPtr<UMaterialInterface> GetMaterial() const { return nullptr;}
 
 private:
-	void AddSignatureParam(TArray<FMaterialParameterInfo> OutParameterInfoScalar, FName CPPTypeName, FTG_Signature::FInit& SignatureInit, bool IsScalar = false) const;
+	void AddSignatureParam(const TArray<FMaterialParameterInfo>& OutParameterInfo, const TArray<FGuid>& OutParameterIds, EMaterialParameterType MatType, FTG_Signature::FInit& SignatureInit) const;
 
 	static EDrawMaterialAttributeTarget ConvertEMaterialPropertyToEDrawMaterialAttributeTarget(EMaterialProperty InMaterialProperty);
 };
diff --git a/Engine/Plugins/Experimental/TextureGraph/Source/TextureGraph/Public/TG_Signature.h b/Engine/Plugins/Experimental/TextureGraph/Source/TextureGraph/Public/TG_Signature.h
--- a/Engine/Plugins/Experimental/TextureGraph/Source/TextureGraph/Public/TG_Signature.h
+++ b/Engine/Plugins/Experimental/TextureGraph/Source/TextureGraph/Public/TG_Signature.h
@@ -152,7 +152,7 @@ public:
 
 	bool				IsScalar() const; // float or int
 	bool				IsBool() const; // boolean
-	bool				IsString() const; // string
+	bool				IsString() const; // string  but NOT FName
 	bool				IsColor() const; // FlinearColor
 	bool				IsVector() const; // FLinearColor or ...
 	bool				IsTexture() const; // FTG_Texture or ...
@@ -169,9 +169,16 @@ public:
 	friend bool operator== (const FTG_Argument& lhs, const FTG_Argument& rhs) {
 		return (lhs.Name == rhs.Name) && (lhs.CPPTypeName == rhs.CPPTypeName) && (lhs.ArgumentType == rhs.ArgumentType);
 	}
+
+	// Comparison operator against the Arg Name, useful for FindByKey on TArray of Arguments
+	bool operator== (const FName& InArgName) {
+		return (Name == InArgName);
+	}
 };
 using FTG_Arguments = TArray<FTG_Argument>;
 
+TArray<FName> TG_MakeArrayOfArgumentNames(const FTG_Arguments& InArguments);
+
 // Argument Set:
 // A container for several Arguments which is searchable and hashable
 // The arguments all have an index (their position in the set).
