Fixed pin connections after copy/paste

diff --git a/Engine/Plugins/Experimental/TextureGraph/Source/TextureGraph/Public/TG_Node.h b/Engine/Plugins/Experimental/TextureGraph/Source/TextureGraph/Public/TG_Node.h
--- a/Engine/Plugins/Experimental/TextureGraph/Source/TextureGraph/Public/TG_Node.h
+++ b/Engine/Plugins/Experimental/TextureGraph/Source/TextureGraph/Public/TG_Node.h
@@ -111,7 +111,7 @@ public:
 	// All the pins are stored packed in the <Pins> member
 	//////////////////////////////////////////////////////////////////////////
 
-    UPROPERTY(TextExportTransient)
+    UPROPERTY()
 	TArray<TObjectPtr<UTG_Pin>> Pins;
 
 	// Iterate through all the Input OR Output pins in this node, return the number traversed
diff --git a/Engine/Plugins/Experimental/TextureGraph/Source/TextureGraphEditor/Private/EdGraph/TG_EdGraph.cpp b/Engine/Plugins/Experimental/TextureGraph/Source/TextureGraphEditor/Private/EdGraph/TG_EdGraph.cpp
--- a/Engine/Plugins/Experimental/TextureGraph/Source/TextureGraphEditor/Private/EdGraph/TG_EdGraph.cpp
+++ b/Engine/Plugins/Experimental/TextureGraph/Source/TextureGraphEditor/Private/EdGraph/TG_EdGraph.cpp
@@ -180,6 +180,26 @@ void UTG_EdGraph::GraphChanged(UTG_Graph* InGraph, UTG_Node* InNode, bool Tweaki
 	}
 }
 
+void UTG_EdGraph::FixDuplicatedNodesPinConnections(TSet<UEdGraphNode*>& PastedNodes)
+{
+	// we recreate the EdPin connections (coming from the copy/paste) to trigger the creation of TG_Node/TG_Pin connections correctly
+	for (TSet<UEdGraphNode*>::TIterator It(PastedNodes); It; ++It)
+	{
+		UEdGraphNode* EdNode = *It;
+	
+		for(UEdGraphPin* SourceEdGraphPin : EdNode->Pins)
+		{
+			if (SourceEdGraphPin->Direction == EGPD_Output)
+			{
+				for (UEdGraphPin* DestEdGraphPin : SourceEdGraphPin->LinkedTo)
+				{
+					GetSchema()->TryCreateConnection(SourceEdGraphPin, DestEdGraphPin);	
+				}
+			}
+		}
+	}
+}
+
 void UTG_EdGraph::OnNodeSignatureChanged(UTG_Node* InNode)
 {
 	UTG_EdGraphNode* EdGraphNode = GetViewModelNode(InNode->GetId());
diff --git a/Engine/Plugins/Experimental/TextureGraph/Source/TextureGraphEditor/Private/EdGraph/TG_EdGraphNode.cpp b/Engine/Plugins/Experimental/TextureGraph/Source/TextureGraphEditor/Private/EdGraph/TG_EdGraphNode.cpp
--- a/Engine/Plugins/Experimental/TextureGraph/Source/TextureGraphEditor/Private/EdGraph/TG_EdGraphNode.cpp
+++ b/Engine/Plugins/Experimental/TextureGraph/Source/TextureGraphEditor/Private/EdGraph/TG_EdGraphNode.cpp
@@ -356,7 +356,7 @@ void UTG_EdGraphNode::ReconstructNode()
 		{
 			// And also check that the new pin at that name is connected to anything, if so grab the UI connections
 			UTG_Pin* TGPin = Node->GetPin(Node->GetPinId(OldPinName));
-			if (TGPin->IsConnected())
+			if (TGPin->IsConnected() || OldPin->HasAnyConnections())
 			{
 				(*NewPin)->MovePersistentDataFromOldPin(*OldPin);
 			}
@@ -457,10 +457,13 @@ void UTG_EdGraphNode::PostPasteNode()
 	UTG_EdGraph* EdGraph = CastChecked<UTG_EdGraph>(GetGraph());
 	UTG_Graph* Graph = EdGraph->TextureGraph->Graph();
 	check(Graph);
+	
+	Node->Pins.Empty();
 	Node->Rename(nullptr, Graph, REN_DontCreateRedirectors | REN_DoNotDirty);
-
+	
 	// Our TG node is a new node that need to be taken care of and added to the graph
 	Graph->AddPostPasteNode(Node);
+
 }
 
 #if WITH_EDITOR
diff --git a/Engine/Plugins/Experimental/TextureGraph/Source/TextureGraphEditor/Private/TG_Editor.cpp b/Engine/Plugins/Experimental/TextureGraph/Source/TextureGraphEditor/Private/TG_Editor.cpp
--- a/Engine/Plugins/Experimental/TextureGraph/Source/TextureGraphEditor/Private/TG_Editor.cpp
+++ b/Engine/Plugins/Experimental/TextureGraph/Source/TextureGraphEditor/Private/TG_Editor.cpp
@@ -1396,7 +1396,8 @@ void FTG_Editor::PasteNodesHere(const FVector2D& Location, const class UEdGraph*
 	// Import the nodes
 	TSet<UEdGraphNode*> PastedNodes;
 	FEdGraphUtilities::ImportNodesFromText(ExpressionGraph, TextToImport, /*out*/ PastedNodes);
-
+	ExpressionGraph->FixDuplicatedNodesPinConnections(PastedNodes);
+	
 	//Average position of nodes so we can move them while still maintaining relative distances to each other
 	FVector2D AvgNodePosition(0.0f, 0.0f);
 
diff --git a/Engine/Plugins/Experimental/TextureGraph/Source/TextureGraphEditor/Public/EdGraph/TG_EdGraph.h b/Engine/Plugins/Experimental/TextureGraph/Source/TextureGraphEditor/Public/EdGraph/TG_EdGraph.h
--- a/Engine/Plugins/Experimental/TextureGraph/Source/TextureGraphEditor/Public/EdGraph/TG_EdGraph.h
+++ b/Engine/Plugins/Experimental/TextureGraph/Source/TextureGraphEditor/Public/EdGraph/TG_EdGraph.h
@@ -60,6 +60,7 @@ public:
 	 */
 	void OnNodeCreateThumbnail(UTG_Node* Node, const FTG_EvaluationContext* EvaluationContext);
 	
+
 	UPROPERTY()
 	TObjectPtr<UTextureGraph> TextureGraph = nullptr;
 	
@@ -78,9 +79,8 @@ private:
 	void OnNodePostEvaluation(UTG_Node* InNode, const FTG_EvaluationContext* Context);
 
 	void GraphChanged(UTG_Graph* InGraph, UTG_Node* InNode, bool Tweaking);
-
-
 	TMap<FTG_Id, TiledBlobPtr> PinThumbBlobMap;
+
 public:
-	//void UpdateParams();
+	void FixDuplicatedNodesPinConnections(TSet<UEdGraphNode*>& PastedNodes);
 };
