Address the problem of output type produced by the dot product expression
improved the common variant type api and removed from node  to make it easier to change it to specialize in sub expression class

diff --git a/Engine/Plugins/Experimental/TextureGraph/Source/TextureGraph/Private/Expressions/Maths/TG_Expression_Maths_TwoInputs.cpp b/Engine/Plugins/Experimental/TextureGraph/Source/TextureGraph/Private/Expressions/Maths/TG_Expression_Maths_TwoInputs.cpp
--- a/Engine/Plugins/Experimental/TextureGraph/Source/TextureGraph/Private/Expressions/Maths/TG_Expression_Maths_TwoInputs.cpp
+++ b/Engine/Plugins/Experimental/TextureGraph/Source/TextureGraph/Private/Expressions/Maths/TG_Expression_Maths_TwoInputs.cpp
@@ -68,7 +68,7 @@ FTG_Texture	UTG_Expression_Subtract::EvaluateTexture(FTG_EvaluationContext* InCo
 //////////////////////////////////////////////////////////////////////////
 float UTG_Expression_Dot::EvaluateScalar_WithValue(FTG_EvaluationContext* InContext, const float* const ValuePtr, size_t Count)
 {
-	return 0;
+	return  ValuePtr[0] * ValuePtr[1];
 }
 
 FVector4f UTG_Expression_Dot::EvaluateVector_WithValue(FTG_EvaluationContext* InContext, const FVector4f* const Value, size_t Count)
@@ -83,6 +83,58 @@ FTG_Texture	UTG_Expression_Dot::EvaluateTexture(FTG_EvaluationContext* InContext
 	return T_Maths_TwoInputs::CreateDot(InContext->Cycle, Output.EditTexture().GetBufferDescriptor(), InContext->TargetId, Input1.GetTexture(InContext), Input2.GetTexture(InContext));
 }
 
+FTG_Variant::EType	UTG_Expression_Dot::EvalExpressionCommonOutputVariantType() const
+{
+	// Run the default behavior for comon type and force to scalar if it is NOT a texture type
+	auto DefaultCommonType = GetCommonInputVariantType();
+	switch (DefaultCommonType)
+	{
+	case FTG_Variant::EType::Scalar:
+	case FTG_Variant::EType::Color:
+	case FTG_Variant::EType::Vector:
+		return FTG_Variant::EType::Scalar;
+	case FTG_Variant::EType::Texture:
+		return FTG_Variant::EType::Texture;
+	case FTG_Variant::EType::Invalid:
+	default:
+		return FTG_Variant::EType::Invalid; // THis shouldn't happen
+	}
+}
+
+
+void UTG_Expression_Dot::Evaluate(FTG_EvaluationContext* InContext)
+{
+	UTG_Expression::Evaluate(InContext);
+
+	// For dot product, common output variant is scalar for simple types or texture
+	switch (GetCommonInputVariantType())
+	{
+	case FTG_Variant::EType::Scalar:
+		Output.EditScalar() = EvaluateScalar(InContext);
+		break;
+	case FTG_Variant::EType::Color:
+	{
+		Output.EditScalar() = EvaluateColor(InContext).R;
+		break;
+	}
+	case FTG_Variant::EType::Vector:
+	{
+		Output.EditScalar() = EvaluateVector(InContext).X;
+		break;
+	}
+	case FTG_Variant::EType::Texture:
+	{
+		Output.EditTexture() = EvaluateTexture(InContext).RasterBlob;
+		break;
+	}
+	default:
+	{
+		Output.EditScalar() = 0.0f;
+		break;
+	}
+	}
+}
+
 //////////////////////////////////////////////////////////////////////////
 /// Cross 
 //////////////////////////////////////////////////////////////////////////
diff --git a/Engine/Plugins/Experimental/TextureGraph/Source/TextureGraph/Private/Expressions/TG_Expression.cpp b/Engine/Plugins/Experimental/TextureGraph/Source/TextureGraph/Private/Expressions/TG_Expression.cpp
--- a/Engine/Plugins/Experimental/TextureGraph/Source/TextureGraph/Private/Expressions/TG_Expression.cpp
+++ b/Engine/Plugins/Experimental/TextureGraph/Source/TextureGraph/Private/Expressions/TG_Expression.cpp
@@ -324,3 +324,42 @@ void UTG_Expression::LogEvaluation(FTG_EvaluationContext* InContext)
 
 	UE_LOG(LogTextureGraph, VeryVerbose, TEXT("%-*s %s"), UTG_Graph::LogHeaderWidth, *GetDefaultName().ToString(), *UTG_Graph::LogCall(InVarIds, OutVarIds));
 }
+
+
+FTG_Variant::EType UTG_Expression::EvalExpressionCommonInputVariantType() const
+{
+	FTG_Variant::EType Type = FTG_Variant::EType::Invalid;
+	const UTG_Node* ParentNode = GetParentNode();
+	for (auto Pin : ParentNode->Pins)
+	{
+		if (Pin->IsInput() && Pin->IsArgVariant())
+		{
+			if (Pin->IsConnected())
+			{
+				// now we know the thing connected to the pin variant exist, what is its type?
+				UTG_Pin* OtherPin = ParentNode->GetGraph()->GetPin(Pin->GetVarId());
+
+				if (OtherPin)
+				{
+					FTG_Variant::EType OtherType = FTG_Variant::GetTypeFromName(OtherPin->GetArgumentCPPTypeName());
+					Type = FTG_Variant::WhichCommonType(Type, OtherType);
+				}
+				else
+					return FTG_Variant::EType::Invalid;
+			}
+			// removed the else part here since the selfvar/default will be reset to the common type on the next evaluation
+			/*else
+				{
+				Type = FTG_Variant::WhichCommonType(Type, Pin->GetSelfVar()->GetAs<FTG_Variant>().GetType());
+				}*/
+		}
+	}
+
+	return Type;
+}
+
+
+FTG_Variant::EType UTG_Expression::EvalExpressionCommonOutputVariantType() const
+{
+	return GetCommonInputVariantType();
+}
\ No newline at end of file
diff --git a/Engine/Plugins/Experimental/TextureGraph/Source/TextureGraph/Private/Expressions/TG_Expression_Variant.cpp b/Engine/Plugins/Experimental/TextureGraph/Source/TextureGraph/Private/Expressions/TG_Expression_Variant.cpp
--- a/Engine/Plugins/Experimental/TextureGraph/Source/TextureGraph/Private/Expressions/TG_Expression_Variant.cpp
+++ b/Engine/Plugins/Experimental/TextureGraph/Source/TextureGraph/Private/Expressions/TG_Expression_Variant.cpp
@@ -93,9 +93,7 @@ void UTG_Expression_Variant::Evaluate(FTG_EvaluationContext* InContext)
 {
 	Super::Evaluate(InContext);
 
-	//CommonVariantType = EvalExpressionCommonVariantType();
-
-	switch (CommonVariantType)
+	switch (GetCommonInputVariantType())
 	{
 	case FTG_Variant::EType::Scalar:
 		Output.EditScalar() = EvaluateScalar(InContext);
diff --git a/Engine/Plugins/Experimental/TextureGraph/Source/TextureGraph/Private/TG_Graph.cpp b/Engine/Plugins/Experimental/TextureGraph/Source/TextureGraph/Private/TG_Graph.cpp
--- a/Engine/Plugins/Experimental/TextureGraph/Source/TextureGraph/Private/TG_Graph.cpp
+++ b/Engine/Plugins/Experimental/TextureGraph/Source/TextureGraph/Private/TG_Graph.cpp
@@ -273,7 +273,7 @@ void UTG_Graph::RegenerateNode(UTG_Node* InNode)
 	// Last and not least, if there is a variant pin, let's reevaluate the common variant type
 	if (bHasVariant)
 	{
-		InNode->GetExpression()->AssignCommonVariantType(InNode->EvalExpressionCommonVariantType());
+		InNode->GetExpression()->ResetCommonInputVariantType();
 	}
 }
 
diff --git a/Engine/Plugins/Experimental/TextureGraph/Source/TextureGraph/Private/TG_GraphEvaluation.cpp b/Engine/Plugins/Experimental/TextureGraph/Source/TextureGraph/Private/TG_GraphEvaluation.cpp
--- a/Engine/Plugins/Experimental/TextureGraph/Source/TextureGraph/Private/TG_GraphEvaluation.cpp
+++ b/Engine/Plugins/Experimental/TextureGraph/Source/TextureGraph/Private/TG_GraphEvaluation.cpp
@@ -42,7 +42,7 @@ void FTG_Evaluation::TransferVarToPin(UTG_Pin* InPin, FTG_EvaluationContext* Con
 
 		if (InPin->IsArgVariant() && !InPin->IsConnected())
 		{
-			auto VariantType = Context->CurrentNode->GetExpressionCommonVariantType();
+			auto VariantType = Context->CurrentNode->GetExpressionCommonInputVariantType();
 			Var->EditAs<FTG_Variant>().ResetTypeAs(VariantType);
 		}
 
@@ -171,7 +171,7 @@ void FloatToFTG_Variant_Converter(FTG_Evaluation::VarConverterInfo& Info)
 	auto Input = Info.InVar->GetAs<float>();
 	auto& Output = Info.OutVar->EditAs<FTG_Variant>();
 
-	FTG_Variant::EType VariantType = Info.Context->CurrentNode->GetExpressionCommonVariantType();
+	FTG_Variant::EType VariantType = Info.Context->CurrentNode->GetExpressionCommonInputVariantType();
 
 	switch (VariantType)
 	{
@@ -200,7 +200,7 @@ void FLinearColorToFTG_Variant_Converter(FTG_Evaluation::VarConverterInfo& Info)
 	auto Input = Info.InVar->GetAs<FLinearColor>();
 	auto& Output = Info.OutVar->EditAs<FTG_Variant>();
 
-	FTG_Variant::EType VariantType = Info.Context->CurrentNode->GetExpressionCommonVariantType();
+	FTG_Variant::EType VariantType = Info.Context->CurrentNode->GetExpressionCommonInputVariantType();
 
 	switch (VariantType)
 	{
@@ -227,7 +227,7 @@ void FVector4fToFTG_Variant_Converter(FTG_Evaluation::VarConverterInfo& Info)
 	auto Input = Info.InVar->GetAs<FVector4f>();
 	auto& Output = Info.OutVar->EditAs<FTG_Variant>();
 
-	FTG_Variant::EType VariantType = Info.Context->CurrentNode->GetExpressionCommonVariantType();
+	FTG_Variant::EType VariantType = Info.Context->CurrentNode->GetExpressionCommonInputVariantType();
 
 	switch (VariantType)
 	{
diff --git a/Engine/Plugins/Experimental/TextureGraph/Source/TextureGraph/Private/TG_Node.cpp b/Engine/Plugins/Experimental/TextureGraph/Source/TextureGraph/Private/TG_Node.cpp
--- a/Engine/Plugins/Experimental/TextureGraph/Source/TextureGraph/Private/TG_Node.cpp
+++ b/Engine/Plugins/Experimental/TextureGraph/Source/TextureGraph/Private/TG_Node.cpp
@@ -187,8 +187,7 @@ void UTG_Node::OnPinConnectionChanged(FTG_Id InPinId, FTG_Id OldPinId, FTG_Id Ne
 
 	if (ThePin->IsArgVariant())
 	{
-		FTG_Variant::EType CommonType = EvalExpressionCommonVariantType();
-		GetExpression()->NotifyCommonVariantTypeChanged(CommonType);
+		GetExpression()->NotifyCommonInputVariantTypeChanged(); // notify and update the common type and the signature accordingly
 	}
 
 	if (!ThePin->IsConnected())
@@ -208,8 +207,7 @@ void UTG_Node::OnPinConnectionUndo(FTG_Id InPinId)
 	// Reset the Pin to Default in this case
 	if (PinTo->IsArgVariant())
 	{
-		FTG_Variant::EType CommonType = FTG_Variant::EType::Scalar;
-		GetExpression()->NotifyCommonVariantTypeChanged(CommonType);
+		GetExpression()->NotifyCommonInputVariantTypeChanged(FTG_Variant::EType::Scalar); // notify and update the common type and the signature accordingly
 	}
 	else if (PinTo->IsArgTexture())
 	{
@@ -485,41 +483,9 @@ bool UTG_Node::CheckPinSignatureAgainstExpression() const
 	return ExpArguments == PinArguments;
 }
 
-FTG_Variant::EType UTG_Node::GetExpressionCommonVariantType() const
+FTG_Variant::EType UTG_Node::GetExpressionCommonInputVariantType() const
 {
-	return GetExpression()->GetCommonVariantType();
-}
-
-FTG_Variant::EType UTG_Node::EvalExpressionCommonVariantType() const
-{
-	FTG_Variant::EType Type = FTG_Variant::EType::Invalid;
-
-	for (auto Pin : Pins)
-	{
-		if (Pin->IsInput() && Pin->IsArgVariant())
-		{
-			if (Pin->IsConnected())
-			{
-				// now we know the thing connected to the pin variant exist, what is its type?
-				UTG_Pin* OtherPin = GetGraph()->GetPin(Pin->GetVarId());
-
-				if (OtherPin)
-				{
-					FTG_Variant::EType OtherType = FTG_Variant::GetTypeFromName(OtherPin->GetArgumentCPPTypeName());
-					Type = FTG_Variant::WhichCommonType(Type, OtherType);
-				}
-				else
-					return FTG_Variant::EType::Invalid;
-			}
-			// removed the else part here since the selfvar/default will be reset to the common type on the next evaluation
-			/*else
-			{
-				Type = FTG_Variant::WhichCommonType(Type, Pin->GetSelfVar()->GetAs<FTG_Variant>().GetType());
-			}*/
-		}
-	}
-
-	return Type;
+	return GetExpression()->GetCommonInputVariantType();
 }
 
 int UTG_Node::GetAllOutputValues(TArray<FTG_Variant>& OutVariants, TArray<FName>* OutNames) const
diff --git a/Engine/Plugins/Experimental/TextureGraph/Source/TextureGraph/Public/Expressions/Maths/TG_Expression_Maths_TwoInputs.h b/Engine/Plugins/Experimental/TextureGraph/Source/TextureGraph/Public/Expressions/Maths/TG_Expression_Maths_TwoInputs.h
--- a/Engine/Plugins/Experimental/TextureGraph/Source/TextureGraph/Public/Expressions/Maths/TG_Expression_Maths_TwoInputs.h
+++ b/Engine/Plugins/Experimental/TextureGraph/Source/TextureGraph/Public/Expressions/Maths/TG_Expression_Maths_TwoInputs.h
@@ -110,6 +110,10 @@ protected:
 	virtual float						EvaluateScalar_WithValue(FTG_EvaluationContext* InContext, const float* const Value, size_t Count) override;
 	virtual FVector4f					EvaluateVector_WithValue(FTG_EvaluationContext* InContext, const FVector4f* const Value, size_t Count) override;
 	virtual FTG_Texture					EvaluateTexture(FTG_EvaluationContext* InContext) override;
+
+	virtual void						Evaluate(FTG_EvaluationContext* InContext) override;
+
+	virtual FTG_Variant::EType			EvalExpressionCommonOutputVariantType() const override; // specify the logic to decide the output variant type
 };
 
 //////////////////////////////////////////////////////////////////////////
diff --git a/Engine/Plugins/Experimental/TextureGraph/Source/TextureGraph/Public/Expressions/TG_Expression.h b/Engine/Plugins/Experimental/TextureGraph/Source/TextureGraph/Public/Expressions/TG_Expression.h
--- a/Engine/Plugins/Experimental/TextureGraph/Source/TextureGraph/Public/Expressions/TG_Expression.h
+++ b/Engine/Plugins/Experimental/TextureGraph/Source/TextureGraph/Public/Expressions/TG_Expression.h
@@ -122,11 +122,31 @@ protected:
 	// Only concrete implementation for Dynamic Expression
 	virtual void NotifySignatureChanged() const;
 
+	// Variant expression API: specialized when using the macro TG_DECLARE_VARIANT_EXPRESSION: 
+	// When input connection are changing, the graph / node calls ResetCommonInputVariantType in order for the expression
+	// to reeval its CommonInputVariant type as well as it s signature if it changes 
+	virtual bool ResetCommonInputVariantType(FTG_Variant::EType InType = FTG_Variant::EType::Invalid) const { return false; }
+	virtual void NotifyCommonInputVariantTypeChanged(FTG_Variant::EType NewType = FTG_Variant::EType::Invalid) const final {
+		if (ResetCommonInputVariantType(NewType))
+			NotifySignatureChanged();
+	}
+	virtual FTG_Variant::EType GetCommonInputVariantType() const { return FTG_Variant::EType::Scalar; }
+
+	// Variant expression API:
+	// Eval the common input variant type used by the expression (only if it is variant)
+	// by looking at the various input pins connected or default types and come up with a "common input" type used for 
+	// how to interpret all the input variants
+	// Default implementation: The common variant type is found as the type that superseeds all the variant input arguments
+	//           currently fed from other nodes 
+	// This call can be modified in a specific class to get a different behavior.
+	virtual FTG_Variant::EType EvalExpressionCommonInputVariantType() const;
+
+	// Variant expression API:
+	// Similarly, this call eval the common output variant type used by the expression (only if it is variant)
+	// Default implementation: the common output variant is the same as the common input variant.
+	// this call can be modified in a specific class to get a differnet behavior, see Expression_Dot for an example.
+	virtual FTG_Variant::EType EvalExpressionCommonOutputVariantType() const;
 
-	virtual FTG_Variant::EType GetCommonVariantType() const { return FTG_Variant::EType::Scalar; }
-	virtual bool AssignCommonVariantType(FTG_Variant::EType InType) const { return false; }
-	virtual void NotifyCommonVariantTypeChanged(FTG_Variant::EType NewType) const {}
-	FTG_Variant::EType EvalExpressionCommonVariantType() const { return GetParentNode()->EvalExpressionCommonVariantType(); }
 
 	// In some cases, evaluation or change in the Expression needs to be feedback to the matching pin's value
 	// This is not needed for the standard flow of evaluation but is sometime required for coupled member.
@@ -164,29 +184,28 @@ protected:
 	public: virtual FName GetCategory() const override { return Category; } 
 
 #define TG_DECLARE_VARIANT_EXPRESSION(Category) \
-	protected: mutable FTG_Variant::EType CommonVariantType = FTG_Variant::EType::Scalar; \
+	protected: mutable FTG_Variant::EType CommonInputVariantType = FTG_Variant::EType::Scalar; \
 	virtual void Initialize() override { \
 		Super::Initialize(); \
-		CommonVariantType = EvalExpressionCommonVariantType(); \
+		CommonInputVariantType = EvalExpressionCommonInputVariantType(); \
 	} \
-	virtual FTG_Variant::EType GetCommonVariantType() const final { return CommonVariantType; } \
-	virtual bool AssignCommonVariantType(FTG_Variant::EType InType) const final { \
-		if (CommonVariantType != InType) { \
-			CommonVariantType = InType; \
+	virtual FTG_Variant::EType GetCommonInputVariantType() const final { return CommonInputVariantType; } \
+	virtual bool ResetCommonInputVariantType(FTG_Variant::EType InType = FTG_Variant::EType::Invalid) const final { \
+		if (InType == FTG_Variant::EType::Invalid) { \
+			InType = EvalExpressionCommonInputVariantType(); \
+		} \
+		if (CommonInputVariantType != InType) { \
+			CommonInputVariantType = InType; \
 			return true; \
 		} \
 		return false; \
 	} \
-	virtual void NotifyCommonVariantTypeChanged(FTG_Variant::EType InType) const final { \
-		if (AssignCommonVariantType(InType)) \
-			NotifySignatureChanged(); \
-	} \
 	protected: mutable FTG_SignaturePtr DynSignature = nullptr; \
 	virtual FTG_SignaturePtr BuildSignatureDynamically() const final { \
 		FTG_Signature::FInit SignatureInit = GetSignatureInitArgsFromClass(); \
 		for (auto& a : SignatureInit.Arguments) { \
 			if (a.IsOutput()) { \
-				a.CPPTypeName = FTG_Variant::GetArgNameFromType(CommonVariantType); \
+				a.CPPTypeName = FTG_Variant::GetArgNameFromType(EvalExpressionCommonOutputVariantType()); \
 			} \
 		} \
 		return MakeShared<FTG_Signature>(SignatureInit); \
diff --git a/Engine/Plugins/Experimental/TextureGraph/Source/TextureGraph/Public/TG_Node.h b/Engine/Plugins/Experimental/TextureGraph/Source/TextureGraph/Public/TG_Node.h
--- a/Engine/Plugins/Experimental/TextureGraph/Source/TextureGraph/Public/TG_Node.h
+++ b/Engine/Plugins/Experimental/TextureGraph/Source/TextureGraph/Public/TG_Node.h
@@ -166,9 +166,9 @@ public:
 	// along with their names optionally
 	int GetAllOutputValues(TArray<FTG_Texture>& OutTextures, TArray<FName>* OutNames = nullptr) const;
 
-	// Get the current common variant type of the expression
+	// Get the current common input variant type of the expression
 	// only valid for expression with FTG_Variant arguments
-	FTG_Variant::EType GetExpressionCommonVariantType() const;
+	FTG_Variant::EType GetExpressionCommonInputVariantType() const;
 
 	void OnPinConnectionUndo(FTG_Id InPinId);
 
@@ -201,11 +201,6 @@ protected:
 
 	void OnPinConnectionChanged(FTG_Id InPinId, FTG_Id OldPinId, FTG_Id NewPinId);
 
-	// Eval the common variant type used by the expression (only if it is variant)
-	// The common variant type is found as the type that superseeds all the variant input arguments
-	// currently fed from other nodes 
-	FTG_Variant::EType EvalExpressionCommonVariantType() const;
-
 #if WITH_EDITOR
 	virtual void PostEditChangeProperty(FPropertyChangedEvent& PropertyChangedEvent) override;
 	virtual void PostEditUndo() override;
