TextureGraphEngine: Disable clang analysis due to clang-cl crash

#rnx

diff --git a/Engine/Plugins/Experimental/TextureGraph/Source/TextureGraphEngine/TextureGraphEngine.Build.cs b/Engine/Plugins/Experimental/TextureGraph/Source/TextureGraphEngine/TextureGraphEngine.Build.cs
--- a/Engine/Plugins/Experimental/TextureGraph/Source/TextureGraphEngine/TextureGraphEngine.Build.cs
+++ b/Engine/Plugins/Experimental/TextureGraph/Source/TextureGraphEngine/TextureGraphEngine.Build.cs
@@ -34,6 +34,15 @@ public class TextureGraphEngine : ModuleRules
         bUseRTTI = true;
         bEnableExceptions = true;
 
+		// Disable Clang analysis due to a crash in clang-cl
+		if (Target.StaticAnalyzer != StaticAnalyzer.None
+			&& Target.Platform.IsInGroup(UnrealPlatformGroup.Microsoft)
+			&& Target.WindowsPlatform.Compiler.IsClang())
+		{
+			// https://developercommunity.visualstudio.com/t/clang-cl---analyze-crash-with-thread-i/10746623
+			bDisableStaticAnalysis = true;
+		}
+
 		//PublicDefinitions.Add("WITH_MALLOC_STOMP=1");
 
 		PrivateDependencyModuleNames.AddRange(new string[] 
