Arrays support.

#tests 66e42e09f77dd673757ae70e

diff --git a/Engine/Plugins/Experimental/TextureGraph/Source/TextureGraph/Private/Expressions/Arrays/TG_Expression_Array.cpp b/Engine/Plugins/Experimental/TextureGraph/Source/TextureGraph/Private/Expressions/Arrays/TG_Expression_Array.cpp
new file mode 100644
--- /dev/null
+++ b/Engine/Plugins/Experimental/TextureGraph/Source/TextureGraph/Private/Expressions/Arrays/TG_Expression_Array.cpp
@@ -0,0 +1,24 @@
+// Copyright Epic Games, Inc. All Rights Reserved.
+
+#include "Expressions/Arrays/TG_Expression_Array.h"
+
+#include "TG_Graph.h"
+#include "2D/TextureHelper.h"
+#include "Model/StaticImageResource.h"
+
+void UTG_Expression_Array4::Evaluate(FTG_EvaluationContext* InContext)
+{
+	Super::Evaluate(InContext);
+
+	Output.SetNum(4);
+
+	//Output[0] = Input1;
+	//Output[1] = Input2;
+	//Output[2] = Input3;
+	//Output[3] = Input4;
+
+	Output.Set(0, Input1);
+	Output.Set(1, Input2);
+	Output.Set(2, Input3);
+	Output.Set(3, Input4);
+}
diff --git a/Engine/Plugins/Experimental/TextureGraph/Source/TextureGraph/Private/Expressions/TG_Expression.cpp b/Engine/Plugins/Experimental/TextureGraph/Source/TextureGraph/Private/Expressions/TG_Expression.cpp
--- a/Engine/Plugins/Experimental/TextureGraph/Source/TextureGraph/Private/Expressions/TG_Expression.cpp
+++ b/Engine/Plugins/Experimental/TextureGraph/Source/TextureGraph/Private/Expressions/TG_Expression.cpp
@@ -13,6 +13,7 @@ const FName TG_Category::Procedural = "Procedural";
 const FName TG_Category::Maths = "Math";
 const FName TG_Category::Utilities = "Utilities";
 const FName TG_Category::Filter = "Filter";
+const FName TG_Category::Arrays = "Arrays";
 const FName TG_Category::Custom = "Custom";
 
 void UTG_Expression::Serialize(FArchive& Ar)
diff --git a/Engine/Plugins/Experimental/TextureGraph/Source/TextureGraph/Private/TG_GraphEvaluation.cpp b/Engine/Plugins/Experimental/TextureGraph/Source/TextureGraph/Private/TG_GraphEvaluation.cpp
--- a/Engine/Plugins/Experimental/TextureGraph/Source/TextureGraph/Private/TG_GraphEvaluation.cpp
+++ b/Engine/Plugins/Experimental/TextureGraph/Source/TextureGraph/Private/TG_GraphEvaluation.cpp
@@ -6,8 +6,36 @@
 #include "TG_Texture.h"
 #include "TG_Variant.h"
 #include "Transform/Expressions/T_FlatColorTexture.h"
+#include "TextureGraphEngine/TextureGraphEngine.h"
 
-void FTG_Evaluation::TransferVarToPin(UTG_Pin* InPin, FTG_EvaluationContext* Context)
+int FTG_Evaluation::FilterArrayInputs(FTG_EvaluationContext* InContext, const TArray<TObjectPtr<UTG_Pin>>& InPins, 
+	TArray<TObjectPtr<UTG_Pin>>& ArrayInputs, TArray<TObjectPtr<UTG_Pin>>& NonArrayPins)
+{
+	int MaxCount = -1;
+	for (const TObjectPtr<UTG_Pin>& InPin : InPins)
+	{
+		FTG_Argument Arg = InPin->GetArgument();
+
+		const auto VarId = InPin->GetVarId();
+		FTG_Var* Var = InContext->Graph->GetVar(VarId);
+
+		if (Var)
+		{
+			/// If it's input and not an array but the var is an array
+			if (Arg.IsInput() && !Arg.IsArray() && Var->IsArray())
+			{
+				ArrayInputs.Add(InPin);
+				MaxCount = FMath::Max<int>(MaxCount, Var->GetAs<FTG_VariantArray>().Num());
+			}
+			else
+				NonArrayPins.Add(InPin);
+		}
+	}
+
+	return MaxCount;
+}
+
+void FTG_Evaluation::TransferVarToPin(UTG_Pin* InPin, FTG_EvaluationContext* Context, int Index)
 {
 	FTG_Argument Arg = InPin->GetArgument();
 
@@ -28,15 +56,17 @@ void FTG_Evaluation::TransferVarToPin(UTG_Pin* InPin, FTG_EvaluationContext* Con
 			// Else COnverter is valid and we need to run it 
 			if ((*Converter))
 			{
-				FTG_Evaluation::VarConverterInfo info;
-				info.Context = Context;
-				info.InVar = Var;
-				info.OutVar = InPin->EditConvertedVar(); // Generate the converted version of the connected input into the CoonvertedVar of the Pin
-				(*Converter)(info);
+				FTG_Evaluation::VarConverterInfo Info;
+				Info.Context = Context;
+				Info.InVar = Var;
+				Info.OutVar = InPin->EditConvertedVar(); // Generate the converted version of the connected input into the CoonvertedVar of the Pin
+				Info.Index = Index;
+
+				(*Converter)(Info);
 
 				// Since the source var is converted into the converted var
 				// we need to pass the self var as the input for the pin
-				Var = info.OutVar;
+				Var = Info.OutVar;
 			}
 		}
 
@@ -48,16 +78,17 @@ void FTG_Evaluation::TransferVarToPin(UTG_Pin* InPin, FTG_EvaluationContext* Con
 
 		if (InPin->NeedsConformance())
 		{
-	 		FTG_Evaluation::VarConformerInfo info;
-	 		info.Context = Context;
-	 		info.InVar = Var;
-	 		info.OutVar = InPin->EditSelfVar();
-			bool result = (InPin->ConformerFunctor)(info);
-	 		if(result)
+	 		FTG_Evaluation::VarConformerInfo Info;
+	 		Info.Context = Context;
+	 		Info.InVar = Var;
+	 		Info.OutVar = InPin->EditSelfVar();
+			Info.Index = Index;
+			bool Result = (InPin->ConformerFunctor)(Info);
+	 		if (Result)
 			{
 				// Since the source var is conformed into the SelfVar
 	            // we need to pass the self var as the input for the pin
-	            Var = info.OutVar;
+	            Var = Info.OutVar;
 			}
 		}
 
@@ -65,6 +96,10 @@ void FTG_Evaluation::TransferVarToPin(UTG_Pin* InPin, FTG_EvaluationContext* Con
 	}
 	else if (Arg.IsOutput())
 	{
+		/// OK, if the output is an array then we just reset the VAR to NULL. 
+		/// This will get copied to anyway when the expression runs. 
+		if (Var->IsArray())
+			Var->Reset();
 		Context->Outputs.VarArguments.Add(Arg.GetName(), { Var, Arg });
 	}
 }
@@ -251,16 +286,22 @@ void FVector4fToFTG_Variant_Converter(FTG_Evaluation::VarConverterInfo& Info)
 
 void FTG_TextureToFTG_Variant_Converter(FTG_Evaluation::VarConverterInfo& Info)
 {
-	auto Input = Info.InVar->GetAs<FTG_Texture>();
-	auto& Output = Info.OutVar->EditAs<FTG_Variant>();
+	FTG_Texture Input = FTG_Texture::GetBlack();
+
+	if (Info.InVar->IsArray())
+	{
+		FTG_VariantArray VarArray = Info.InVar->GetAs<FTG_VariantArray>();
+		check(Info.Index >= 0 && Info.Index < VarArray.Num());
+		Input = VarArray.Get(Info.Index).GetTexture();
+	}
+	else
+		Input = Info.InVar->GetAs<FTG_Texture>();
+	FTG_Variant& Output = Info.OutVar->EditAs<FTG_Variant>();
 	Output.Data.Set<FTG_Texture>(Input);
 }
 
-void FTG_VariantToFloat_Converter(FTG_Evaluation::VarConverterInfo& Info)
+void FTG_VariantToFloat_Converter_Internal(float& Output, const FTG_Variant& Input)
 {
-	auto Input = Info.InVar->GetAs<FTG_Variant>();
-	auto& Output = Info.OutVar->EditAs<float>();
-
 	FTG_Variant::EType SourceType = Input.GetType();
 
 	switch (SourceType)
@@ -280,11 +321,15 @@ void FTG_VariantToFloat_Converter(FTG_Evaluation::VarConverterInfo& Info)
 	}
 }
 
-void FTG_VariantToFLinearColor_Converter(FTG_Evaluation::VarConverterInfo& Info)
+void FTG_VariantToFloat_Converter(FTG_Evaluation::VarConverterInfo& Info)
 {
 	auto Input = Info.InVar->GetAs<FTG_Variant>();
-	auto& Output = Info.OutVar->EditAs<FLinearColor>();
+	auto& Output = Info.OutVar->EditAs<float>();
+	FTG_VariantToFloat_Converter_Internal(Output, Input);
+}
 
+void FTG_VariantToFLinearColor_Converter_Internal(FLinearColor& Output, const FTG_Variant& Input)
+{
 	FTG_Variant::EType SourceType = Input.GetType();
 
 	switch (SourceType)
@@ -304,11 +349,15 @@ void FTG_VariantToFLinearColor_Converter(FTG_Evaluation::VarConverterInfo& Info)
 	}
 }
 
-void FTG_VariantToFVector4f_Converter(FTG_Evaluation::VarConverterInfo& Info)
+void FTG_VariantToFLinearColor_Converter(FTG_Evaluation::VarConverterInfo& Info)
 {
 	auto Input = Info.InVar->GetAs<FTG_Variant>();
-	auto& Output = Info.OutVar->EditAs<FVector4f>();
+	auto& Output = Info.OutVar->EditAs<FLinearColor>();
+	FTG_VariantToFLinearColor_Converter_Internal(Output, Input);
+}
 
+void FTG_VariantToFVector4f_Converter_Internal(FVector4f& Output, const FTG_Variant& Input)
+{
 	FTG_Variant::EType SourceType = Input.GetType();
 
 	switch (SourceType)
@@ -328,12 +377,15 @@ void FTG_VariantToFVector4f_Converter(FTG_Evaluation::VarConverterInfo& Info)
 	}
 }
 
-
-void FTG_VariantToFVector2f_Converter(FTG_Evaluation::VarConverterInfo& Info)
+void FTG_VariantToFVector4f_Converter(FTG_Evaluation::VarConverterInfo& Info)
 {
 	auto Input = Info.InVar->GetAs<FTG_Variant>();
-	auto& Output = Info.OutVar->EditAs<FVector2f>();
+	auto& Output = Info.OutVar->EditAs<FVector4f>();
+	FTG_VariantToFVector4f_Converter_Internal(Output, Input);
+}
 
+void FTG_VariantToFVector2f_Converter_Internal(FVector2f& Output, const FTG_Variant& Input)
+{
 	FTG_Variant::EType SourceType = Input.GetType();
 
 	switch (SourceType)
@@ -353,12 +405,28 @@ void FTG_VariantToFVector2f_Converter(FTG_Evaluation::VarConverterInfo& Info)
 	}
 }
 
+void FTG_VariantToFTG_VariantArray_Converter_Internal(FTG_VariantArray& Output, const FTG_Variant& Input, int Index)
+{
+	check(Index >= 0 && Index < Output.Num());
+	Output.Set(Index, Input);
+}
 
-void FTG_VariantToFTG_Texture_Converter(FTG_Evaluation::VarConverterInfo& Info)
+void FTG_VariantToFTG_VariantArray_Converter(FTG_Evaluation::VarConverterInfo& Info)
+{
+	FTG_Variant& Input = Info.InVar->GetAs<FTG_Variant>();
+	FTG_VariantArray& Output = Info.OutVar->EditAs<FTG_VariantArray>();
+	FTG_VariantToFTG_VariantArray_Converter_Internal(Output, Input, Info.Index);
+}
+
+void FTG_VariantToFVector2f_Converter(FTG_Evaluation::VarConverterInfo& Info)
 {
 	auto Input = Info.InVar->GetAs<FTG_Variant>();
-	auto& Output = Info.OutVar->EditAs<FTG_Texture>();
+	auto& Output = Info.OutVar->EditAs<FVector2f>();
+	FTG_VariantToFVector2f_Converter_Internal(Output, Input);
+}
 
+void FTG_VariantToFTG_Texture_Converter_Internal(FTG_Texture& Output, const FTG_Variant& Input, FTG_Evaluation::VarConverterInfo& Info)
+{
 	FTG_Variant::EType SourceType = Input.GetType();
 
 	FLinearColor Color = FLinearColor::Black;
@@ -381,10 +449,33 @@ void FTG_VariantToFTG_Texture_Converter(FTG_Evaluation::VarConverterInfo& Info)
 	}
 	BufferDescriptor Desc = GetFlatColorDesc(SourceType);
 	Output = T_FlatColorTexture::Create(Info.Context->Cycle, Desc, Color, Info.Context->TargetId);
+}
 
+void FTG_VariantToFTG_Texture_Converter(FTG_Evaluation::VarConverterInfo& Info)
+{
+	auto Input = Info.InVar->GetAs<FTG_Variant>();
+	auto& Output = Info.OutVar->EditAs<FTG_Texture>();
+	FTG_VariantToFTG_Texture_Converter_Internal(Output, Input, Info);
 }
 
+void FTG_VariantArrayToFTG_Variant_Converter(FTG_Evaluation::VarConverterInfo& Info)
+{
+	const FTG_VariantArray& InputArray = Info.InVar->GetAs<FTG_VariantArray>();
+	check(Info.Index >= 0 && Info.Index < InputArray.Num());
+	const FTG_Variant& Input = InputArray.GetArray()[Info.Index];
+	FTG_Variant& Output = Info.OutVar->EditAs<FTG_Variant>();
+	Output = Input;
+}
 
+void FTG_VariantArrayToFTG_Texture_Converter(FTG_Evaluation::VarConverterInfo& Info)
+{
+	const FTG_VariantArray& InputArray = Info.InVar->GetAs<FTG_VariantArray>();
+	check(Info.Index >= 0 && Info.Index < InputArray.Num());
+	const FTG_Variant& Input = InputArray.GetArray()[Info.Index];
+
+	FTG_Texture& Output = Info.OutVar->EditAs<FTG_Texture>();
+	FTG_VariantToFTG_Texture_Converter_Internal(Output, Input, Info);
+}
 
 //////////////////////////////////////////////////////////////////////////////////////////////////////////////////////////////////////
 #define VAR_CONVERTER(nameFrom, nameTo, function)	{ MakeConvertKey(TEXT(#nameFrom), TEXT(#nameTo)), std::bind(&function, std::placeholders::_1) }
@@ -417,6 +508,9 @@ FTG_Evaluation::ConverterMap FTG_Evaluation::DefaultConverters
 		VAR_CONVERTER(FVector4f, FTG_Variant, FVector4fToFTG_Variant_Converter),
 		VAR_CONVERTER(FTG_Texture, FTG_Variant, FTG_TextureToFTG_Variant_Converter),
 
+		VAR_CONVERTER(FTG_VariantArray, FTG_Variant, FTG_VariantArrayToFTG_Variant_Converter),
+		VAR_CONVERTER(FTG_VariantArray, FTG_Texture, FTG_VariantArrayToFTG_Texture_Converter),
+
 		VAR_CONVERTER(FTG_Variant.Scalar, float, FTG_VariantToFloat_Converter),
 
 		VAR_CONVERTER(FTG_Variant.Scalar, FLinearColor, FTG_VariantToFLinearColor_Converter),
@@ -436,6 +530,10 @@ FTG_Evaluation::ConverterMap FTG_Evaluation::DefaultConverters
 		VAR_CONVERTER(FTG_Variant.Color, FVector2f, FTG_VariantToFVector2f_Converter),
 		VAR_CONVERTER(FTG_Variant.Vector, FVector2f, FTG_VariantToFVector2f_Converter),
 
+		VAR_CONVERTER(FTG_Variant.Scalar, FTG_VariantArray, FTG_VariantToFTG_VariantArray_Converter),
+		VAR_CONVERTER(FTG_Variant.Color, FTG_VariantArray, FTG_VariantToFTG_VariantArray_Converter),
+		VAR_CONVERTER(FTG_Variant.Vector, FTG_VariantArray, FTG_VariantToFTG_VariantArray_Converter),
+		VAR_CONVERTER(FTG_Variant.Texture, FTG_VariantArray, FTG_VariantToFTG_VariantArray_Converter),
 
 		VAR_CONVERTER_NULL(FTG_Variant.Scalar, FTG_Variant),
 		VAR_CONVERTER_NULL(FTG_Variant.Color, FTG_Variant),
@@ -468,7 +566,6 @@ bool FTG_Evaluation::AreArgumentsCompatible(const FTG_Argument& ArgFrom, const F
 	return false;
 }
 
-
 void FTG_Evaluation::EvaluateGraph(UTG_Graph* InGraph, FTG_EvaluationContext* InContext)
 {
 	// Entering a new Graph scope, we build a new EvaluationContext
@@ -487,7 +584,7 @@ void FTG_Evaluation::EvaluateGraph(UTG_Graph* InGraph, FTG_EvaluationContext* In
 			auto TheVar = InGraph->GetVar(ParamId);
 			if (TheVar)
 			{
-				(*V.Value.Var).CopyTo(TheVar);
+				V.Value.Var->CopyTo(TheVar);
 				//TheVar->ShareData(*V.Value.Var);
 				EvalContext.ConnectedInputParamIds.Add(ParamId);
 			}
@@ -525,6 +622,107 @@ void FTG_Evaluation::EvaluateGraph(UTG_Graph* InGraph, FTG_EvaluationContext* In
 	}
 }
 
+void FTG_Evaluation::EvaluateNodeArray(UTG_Node* InNode, const TArray<TObjectPtr<UTG_Pin>>& ArrayInputs, 
+	const TArray<TObjectPtr<UTG_Pin>>& NonArrayPins, int MaxCount, FTG_EvaluationContext* InContext)
+{
+	check(!ArrayInputs.IsEmpty());
+	check(MaxCount > 0);
+
+	/// Non array pins only need to have transfer once
+	for (UTG_Pin* Pin : NonArrayPins)
+	{
+		FTG_Evaluation::TransferVarToPin(Pin, InContext, -1);
+	}
+
+	/// Initialize an array for each of the output pins/vars
+	TMap<FName, FTG_VariantArray> ArrayOutputs;
+	for (auto& OutputVar : InContext->Outputs.VarArguments)
+	{
+		FTG_VariantArray ArrayOutput;
+		ArrayOutput.SetNum(MaxCount);
+		ArrayOutputs.Add(OutputVar.Key, std::move(ArrayOutput));
+	}
+
+	for (int VarIndex = 0; VarIndex < MaxCount; VarIndex++)
+	{
+		/// Run through array inputs, one at a time
+		for (int ArrayIndex = 0; ArrayIndex < ArrayInputs.Num(); ArrayIndex++)
+		{
+			UTG_Pin* ArrayPin = ArrayInputs[ArrayIndex];
+			FTG_Var* Var = InContext->Graph->GetVar(ArrayPin->GetVarId());
+			check(Var && Var->IsArray());
+			const FTG_VariantArray& VarArray = Var->GetAs<FTG_VariantArray>();
+
+			/// If an array is smaller than the largest array supplied in the inputs, then we 
+			/// simply clamp it. Currently, it's the responsibility of the user to ensure that 
+			/// inputs match up in terms of their arrays (with the exception that the array has
+			/// one input)
+			int VarArrayIndex = FMath::Min(VarIndex, VarArray.Num() - 1);
+
+			if (VarIndex >= VarArray.Num() && VarArray.Num() > 1)
+			{
+				FString ErrorMsg = FString::Printf(TEXT("Input array mismatch. Array input %s has total number of %d items but the maximum input array length for the node is: %d"),
+					*ArrayPin->GetArgumentName().ToString(), VarArray.Num(), MaxCount);
+				TextureGraphEngine::GetErrorReporter(InContext->Cycle->GetMix())
+					->ReportWarning((int32)ETextureGraphErrorType::INPUT_ARRAY_WARNING, ErrorMsg, InNode);
+			}
+
+			FTG_Evaluation::TransferVarToPin(ArrayPin, InContext, VarArrayIndex);
+		}
+
+		UTG_Expression* Expression = InNode->GetExpression();
+		/// Trigger the Evaluation
+		Expression->SetupAndEvaluate(InContext);
+
+		for (auto& OutputVar : InContext->Outputs.VarArguments)
+		{
+			FTG_VariantArray* ArrayOutput = ArrayOutputs.Find(OutputVar.Key);
+			check(ArrayOutput && ArrayOutput->Num() > VarIndex);
+
+			FName ConverterKey = FName();
+			FString FromType = OutputVar.Value.Argument.GetCPPTypeName().ToString();
+			FString ToType = TEXT("FTG_Variant");
+
+			/// If the output type is not a variant then we need to convert it variant
+			if (FromType.Find(ToType) != 0)
+			{
+				FName FromToName = MakeConvertKey(FName(*FromType), FName(*ToType));
+				VarConverter* Converter = DefaultConverters.Find(FromToName);
+				check(Converter);
+
+				FTG_Var VariantVar;
+				FTG_Variant Variant;
+				VariantVar.SetAs<FTG_Variant>(Variant);
+
+				FTG_Evaluation::VarConverterInfo Info;
+				Info.Context = InContext;
+				Info.InVar = OutputVar.Value.Var;
+				Info.OutVar = &VariantVar; // Generate the converted version of the connected input into the CoonvertedVar of the Pin
+				Info.Index = -1;
+
+				(*Converter)(Info);
+
+				/// We just copy it over in the array
+				ArrayOutput->Set(VarIndex, VariantVar.GetAs<FTG_Variant>());
+			}
+			else
+			{
+				ArrayOutput->Set(VarIndex, OutputVar.Value.Var->GetAs<FTG_Variant>());
+			}
+		}
+	}
+
+	/// We secretly replace the output var(s) with the array outputs that we've accumulated in the 
+	/// loop above. This will ensure that the array of outputs flows downstream to other connected nodes
+	for (auto& OutputVar : InContext->Outputs.VarArguments)
+	{
+		FTG_VariantArray* ArrayOutput = ArrayOutputs.Find(OutputVar.Key);
+		OutputVar.Value.Var->ResetAs<FTG_VariantArray>();
+		OutputVar.Value.Var->SetAs<FTG_VariantArray>(*ArrayOutput);
+		OutputVar.Value.Var->SetArray();
+	}
+}
+
 void FTG_Evaluation::EvaluateNode(UTG_Node* InNode, FTG_EvaluationContext* InContext)
 {
 	const auto Expression = InNode->GetExpression();
@@ -535,18 +733,29 @@ void FTG_Evaluation::EvaluateNode(UTG_Node* InNode, FTG_EvaluationContext* InCon
 		// Grab the vars from the pins and load them in the context's input and output arrays
 		InContext->Inputs.Empty();
 		InContext->Outputs.Empty();
-		for (auto Pin : InNode->Pins)
-		{
-			FTG_Evaluation::TransferVarToPin(Pin, InContext);
-		}
 
-		// Trigger the evaluation of the expression
-		Expression->SetupAndEvaluate(InContext);
+		TArray<TObjectPtr<UTG_Pin>> ArrayInputs, NonArrayPins;
+		int MaxCount = FilterArrayInputs(InContext, InNode->Pins, ArrayInputs, NonArrayPins); 
 
+		if (ArrayInputs.IsEmpty())
+		{
+			/// Just do normal evaluation
+			for (UTG_Pin* Pin : InNode->Pins)
+			{
+				FTG_Evaluation::TransferVarToPin(Pin, InContext, -1);
+			}
+
+			// Trigger the evaluation of the expression
+			Expression->SetupAndEvaluate(InContext);
+		}
+		else
+		{
+			EvaluateNodeArray(InNode, ArrayInputs, NonArrayPins, MaxCount, InContext);
+		}
 
 		// After the evaluation, notify the postEvaluate for thumbnails
 		InNode->GetGraph()->NotifyNodePostEvaluate(InNode, InContext);
 
 		InContext->CurrentNode = nullptr;
 	}
-}
\ No newline at end of file
+}
diff --git a/Engine/Plugins/Experimental/TextureGraph/Source/TextureGraph/Private/TG_Node.cpp b/Engine/Plugins/Experimental/TextureGraph/Source/TextureGraph/Private/TG_Node.cpp
--- a/Engine/Plugins/Experimental/TextureGraph/Source/TextureGraph/Private/TG_Node.cpp
+++ b/Engine/Plugins/Experimental/TextureGraph/Source/TextureGraph/Private/TG_Node.cpp
@@ -497,17 +497,33 @@ int UTG_Node::GetAllOutputValues(TArray<FTG_Variant>& OutVariants, TArray<FName>
 		{
 			// For each valid output pin, grab the result in a variant if valid
 			// and the name if container provided
-			FTG_Variant OutVariant;
-			if (Pin->GetValue(OutVariant))
+			if (!Pin->IsArgArray())
 			{
-				OutVariants.Add(OutVariant);
-				if (OutNames)
-					OutNames->Add(Pin->GetAliasName());
-				NumFounds++;
+				FTG_Variant OutVariant;
+				if (Pin->GetValue(OutVariant))
+				{
+					OutVariants.Add(OutVariant);
+					if (OutNames)
+						OutNames->Add(Pin->GetAliasName());
+					NumFounds++;
+				}
+				else
+				{
+					UE_LOG(LogTextureGraph, Warning, TEXT("Output {} failed to collect as variant"), *(Pin->GetAliasName().ToString()));
+				}
 			}
 			else
 			{
-				UE_LOG(LogTextureGraph, Warning, TEXT("Output {} failed to collect as variant"), *(Pin->GetAliasName().ToString()));
+				FTG_VariantArray VarArray;
+				Pin->GetValue(VarArray);
+
+				for (int Index = 0; Index < VarArray.Num(); Index++)
+				{
+					OutVariants.Add(VarArray.GetArray()[Index]);
+					if (OutNames)
+						OutNames->Add(Pin->GetAliasName());
+					NumFounds++;
+				}
 			}
 		}
 	}
diff --git a/Engine/Plugins/Experimental/TextureGraph/Source/TextureGraph/Private/TG_Pin.cpp b/Engine/Plugins/Experimental/TextureGraph/Source/TextureGraph/Private/TG_Pin.cpp
--- a/Engine/Plugins/Experimental/TextureGraph/Source/TextureGraph/Private/TG_Pin.cpp
+++ b/Engine/Plugins/Experimental/TextureGraph/Source/TextureGraph/Private/TG_Pin.cpp
@@ -494,7 +494,13 @@ bool UTG_Pin::GetValue(FVector4f& OutValue) const
 
 bool UTG_Pin::GetValue(FTG_Texture& OutValue) const
 {
-	if (IsArgTexture())
+	if (IsArgArray())
+	{
+		FTG_VariantArray& VarArray = GetSelfVar()->GetAs<FTG_VariantArray>();
+		OutValue = VarArray.GetViewTexture();
+		return true;
+	}
+	else if (IsArgTexture())
 	{
 		if (IsArgVariant())
 		{
@@ -511,6 +517,16 @@ bool UTG_Pin::GetValue(FTG_Texture& OutValue) const
 	return false;
 }
 
+bool UTG_Pin::GetValue(FTG_VariantArray& OutValue) const
+{
+	if (IsArgArray())
+	{
+		OutValue = GetSelfVar()->GetAs<FTG_VariantArray>();
+		return true;
+	}
+	return false;
+}
+
 
 // Variant value getter only if a Variant argument OR if the argument is one of the supported type by Variant (Scalar/Color/Vector/Texture)
 // After getting the value, the next step is to check the type contained in the variant.
@@ -518,6 +534,7 @@ bool UTG_Pin::GetValue(FTG_Variant& OutValue) const
 {
 	if (IsArgVariant())
 	{
+		check(!IsArgArray());
 		OutValue = GetSelfVar()->GetAs<FTG_Variant>();
 		return true;
 	}
diff --git a/Engine/Plugins/Experimental/TextureGraph/Source/TextureGraph/Private/TG_Signature.cpp b/Engine/Plugins/Experimental/TextureGraph/Source/TextureGraph/Private/TG_Signature.cpp
--- a/Engine/Plugins/Experimental/TextureGraph/Source/TextureGraph/Private/TG_Signature.cpp
+++ b/Engine/Plugins/Experimental/TextureGraph/Source/TextureGraph/Private/TG_Signature.cpp
@@ -39,6 +39,11 @@ bool FTG_Argument::IsTexture() const
 	return CPPTypeName == TEXT("FTG_Texture") || CPPTypeName == FTG_Variant::GetArgNameFromType(ETG_VariantType::Texture);
 }
 
+bool FTG_Argument::IsArray() const
+{
+	return CPPTypeName == TEXT("FTG_VariantArray");
+}
+
 bool FTG_Argument::IsVariant() const
 {
 	if (CPPTypeName == TEXT("FTG_Variant"))
diff --git a/Engine/Plugins/Experimental/TextureGraph/Source/TextureGraph/Private/TG_Var.cpp b/Engine/Plugins/Experimental/TextureGraph/Source/TextureGraph/Private/TG_Var.cpp
--- a/Engine/Plugins/Experimental/TextureGraph/Source/TextureGraph/Private/TG_Var.cpp
+++ b/Engine/Plugins/Experimental/TextureGraph/Source/TextureGraph/Private/TG_Var.cpp
@@ -226,7 +226,10 @@ void Generic_Struct_Serializer(FTG_Var::VarPropertySerialInfo& Info)
 	if (Info.CopyVarToProperty)
 	{
 		if (!Info.Var->IsEmpty() && Info.Var->IsValid())
-			*PropertyValue = Info.Var->GetAs<T_ValueType>();
+		{
+			if (!Info.Var->IsArray())
+				*PropertyValue = Info.Var->GetAs<T_ValueType>();
+		}
 		else
 		{
 			T_ValueType* DefaultValue = StructProperty->ContainerPtrToValuePtrForDefaults<T_ValueType>(StructProperty->Struct, Info.Owner, Info.ClampedIndex());
@@ -240,11 +243,35 @@ void Generic_Struct_Serializer(FTG_Var::VarPropertySerialInfo& Info)
 	}
 }
 
+FTG_Var::FTG_Var()
+{
+}
+
+FTG_Var::FTG_Var(FTG_Id InPinUuid) 
+	: PinId(InPinUuid) 
+{
+}
+
+FTG_Var::FTG_Var(const FTG_Var& InVar) 
+	: PinId(InVar.PinId)
+	, Concept(InVar.Concept) 
+{
+}
+
 void VarPropertySerializer_FTG_Texture(FTG_Var::VarPropertySerialInfo& Info)
 {
 	Generic_Struct_Serializer<FTG_Texture>(Info);
 }
 
+void VarPropertySerializer_FTG_VariantArray(FTG_Var::VarPropertySerialInfo& Info)
+{
+	/// Just call the base generic serializer and set the var to array
+	Generic_Struct_Serializer<FTG_VariantArray>(Info);
+
+	if (!Info.CopyVarToProperty)
+		Info.Var->SetArray();
+}
+
 void VarPropertySerializer_FVector4f(FTG_Var::VarPropertySerialInfo& Info)
 {
 	Generic_Struct_Serializer<FVector4f>(Info);
@@ -296,10 +323,12 @@ void VarPropertySerializer_FName(FTG_Var::VarPropertySerialInfo& Info)
 {
 	Generic_Simple_Serializer<FNameProperty, FName>(Info);
 }
+
 void VarPropertySerializer_FString(FTG_Var::VarPropertySerialInfo& Info)
 {
 	Generic_Simple_Serializer<FStrProperty, FString>(Info);
 }
+
 void VarPropertySerializer_UObjectPtr(FTG_Var::VarPropertySerialInfo& Info)
 {
 	Generic_Struct_Serializer<FTG_Texture>(Info);
@@ -315,12 +344,28 @@ void VarPropertySerializer_FTG_Variant(FTG_Var::VarPropertySerialInfo& Info)
 	{
 		if (!Info.Var->IsEmpty() && Info.Var->IsValid())
 		{
-			*PropertyValue = Info.Var->GetAs<FTG_Variant>();
+			if (!Info.Var->IsArray())
+				*PropertyValue = Info.Var->GetAs<FTG_Variant>();
+			else
+			{
+				const FTG_VariantArray& VarArray = Info.Var->GetAs<FTG_VariantArray>();
+				int Index = Info.ClampedIndex();
+				check(Index < VarArray.Num());
+				*PropertyValue = VarArray.GetArray()[Index];
+			}
 		}
 	}
 	else
 	{
-		Info.Var->SetAs<FTG_Variant>(*PropertyValue);
+		if (!Info.Var->IsArray())
+			Info.Var->SetAs<FTG_Variant>(*PropertyValue);
+		else
+		{
+			FTG_VariantArray& VarArray = Info.Var->GetAs<FTG_VariantArray>();
+			int Index = Info.ClampedIndex();
+			check(Index < VarArray.Num());
+			VarArray.Set(Index, *PropertyValue);
+		}
 	}
 }
 
@@ -468,6 +513,7 @@ FTG_Var::VarPropertySerializerMap FTG_Var::DefaultPropertySerializers
 		VAR_PROPERTY_SERIALIZER_DEF(FPatternMaskBevel_TS),
 		VAR_PROPERTY_SERIALIZER_DEF(FPatternMaskJitter_TS),
 		VAR_PROPERTY_SERIALIZER_DEF(FPatternMaskPlacement_TS),
+		VAR_PROPERTY_SERIALIZER_DEF(FTG_VariantArray),
 	}
 );
 
@@ -486,7 +532,7 @@ void FTG_Var::UnregisterVarPropertySerializer(FName CPPTypeName)
 }
 
 //////////////////////////////////////////////////////////////////////////////////////////////////////////////////////////////////////
-bool FTG_Var::CopyGeneric(UTG_Expression* Owner, const FTG_Argument& Arg, bool CopyVarToProperty)
+bool FTG_Var::CopyGeneric(UTG_Expression* Owner, const FTG_Argument& Arg, bool CopyVarToProperty, int Index)
 {
 	FProperty* Property = Owner->GetClass()->FindPropertyByName(Arg.GetName());
 	if (Property)
@@ -515,13 +561,11 @@ bool FTG_Var::CopyGeneric(UTG_Expression* Owner, const FTG_Argument& Arg, bool C
 			}
 		}
 
-		//const FTG_Var::VarPropertySerializer Serializer = *SerializerIt;
-
 		VarPropertySerialInfo Info = {
 			this,
 			Owner,
 			Arg,
-			-1,
+			Index,
 			CopyVarToProperty
 		};
 
@@ -534,16 +578,23 @@ bool FTG_Var::CopyGeneric(UTG_Expression* Owner, const FTG_Argument& Arg, bool C
 	return false;
 }
 
-void FTG_Var::CopyTo(UTG_Expression* Owner, FTG_Argument& Arg)
+void FTG_Var::CopyTo(UTG_Expression* Owner, FTG_Argument& Arg, int Index)
 {
 	Owner->CopyVarToExpressionArgument(Arg, this);
 }
 
-void FTG_Var::CopyFrom(UTG_Expression* Owner, FTG_Argument& Arg)
+void FTG_Var::CopyFrom(UTG_Expression* Owner, FTG_Argument& Arg, int Index)
 {
 	Owner->CopyVarFromExpressionArgument(Arg, this);
 }
 
+void FTG_Var::CopyTo(FTG_Var* InVar)
+{
+	if (Concept && InVar)
+	{
+		InVar->Concept = Concept->Clone();
+	}
+}
 
 void FTG_Var::CopyFrom(FTG_Var* InVar)
 {
@@ -578,6 +629,11 @@ void VarArchiveSerializer_FTG_Texture(FTG_Var::VarArchiveSerialInfo& Info)
 	//Generic_Simple_ArSerializer<FTG_Texture>(Info);
 }
 
+void VarArchiveSerializer_FTG_VariantArray(FTG_Var::VarArchiveSerialInfo& Info)
+{
+	//Generic_Simple_ArSerializer<FTG_TextureArray>(Info);
+}
+
 void VarArchiveSerializer_FVector4f(FTG_Var::VarArchiveSerialInfo& Info)
 {
 	Generic_Simple_ArSerializer<FVector4f>(Info);
@@ -623,6 +679,7 @@ FTG_Var::VarArchiveSerializerMap FTG_Var::DefaultArchiveSerializers
 		VAR_ARCHIVE_SERIALIZER_DEF(FVector4f),
 		VAR_ARCHIVE_SERIALIZER_DEF(FVector2f),
 		VAR_ARCHIVE_SERIALIZER_DEF(FLinearColor),
+		VAR_ARCHIVE_SERIALIZER_DEF(FTG_VariantArray)
 	}
 ); 
 
diff --git a/Engine/Plugins/Experimental/TextureGraph/Source/TextureGraph/Private/TG_Variant.cpp b/Engine/Plugins/Experimental/TextureGraph/Source/TextureGraph/Private/TG_Variant.cpp
--- a/Engine/Plugins/Experimental/TextureGraph/Source/TextureGraph/Private/TG_Variant.cpp
+++ b/Engine/Plugins/Experimental/TextureGraph/Source/TextureGraph/Private/TG_Variant.cpp
@@ -315,3 +315,34 @@ template <> void TG_Var_SetValueFromString(FTG_Variant& Value, const FString& St
 {
 	// TODO : Implement!
 }
+
+
+void FTG_VariantArray::SetNum(int Count)
+{
+	int CurrentCount = Array.Num();
+	Array.SetNum(Count);
+
+	/// Fill the new items with black
+	for (int Index = CurrentCount; Index < Count; Index++)
+	{
+		Array[Index] = FTG_Texture::GetBlack();
+	}
+}
+
+void FTG_VariantArray::Set(int Index, const FTG_Variant& Var)
+{
+	check(Index >= 0 && Index < Array.Num());
+	Array[Index] = Var;
+}
+
+void FTG_VariantArray::Set(int Index, const FTG_Texture& Tex)
+{
+	check(Index >= 0 && Index < Array.Num());
+	Array[Index].EditTexture() = Tex;
+}
+
+void FTG_VariantArray::SetViewIndex(int Index)
+{
+	check(Index >= 0 && Index < Array.Num());
+	ViewIndex = Index;
+}
diff --git a/Engine/Plugins/Experimental/TextureGraph/Source/TextureGraph/Public/Expressions/Arrays/TG_Expression_Array.h b/Engine/Plugins/Experimental/TextureGraph/Source/TextureGraph/Public/Expressions/Arrays/TG_Expression_Array.h
new file mode 100644
--- /dev/null
+++ b/Engine/Plugins/Experimental/TextureGraph/Source/TextureGraph/Public/Expressions/Arrays/TG_Expression_Array.h
@@ -0,0 +1,42 @@
+/// Copyright Epic Games, Inc. All Rights Reserved.
+
+#pragma once
+
+#include "CoreMinimal.h"
+#include "Expressions/TG_Expression.h"
+#include "TG_Texture.h"
+#include "TG_Expression_Array.generated.h"
+
+UCLASS()
+class TEXTUREGRAPH_API UTG_Expression_Array4 : public UTG_Expression
+{
+	 GENERATED_BODY()
+
+public:
+	TG_DECLARE_EXPRESSION(TG_Category::Arrays);
+	virtual void					Evaluate(FTG_EvaluationContext* InContext) override;
+	
+	/// Texture at the first index of the output texture array
+	UPROPERTY(meta = (TGType = "TG_Input", PinDisplayName = ""))
+	FTG_Variant						Input1;
+
+	/// Texture at the second index of the output texture array
+	UPROPERTY(meta = (TGType = "TG_Input", PinDisplayName = ""))
+	FTG_Variant						Input2;
+
+	/// Texture at the third index of the output texture array
+	UPROPERTY(meta = (TGType = "TG_Input", PinDisplayName = ""))
+	FTG_Variant						Input3;
+
+	/// Texture at the fourth index of the output texture array
+	UPROPERTY(meta = (TGType = "TG_Input", PinDisplayName = ""))
+	FTG_Variant						Input4;
+
+	/// All the input textures organized into an array
+	UPROPERTY(meta = (TGType = "TG_Output", PinDisplayName = "", HideInnerPropertiesInNode))
+	FTG_VariantArray				Output;
+
+	virtual FTG_Name				GetDefaultName() const override { return TEXT("Array4");}
+	virtual FText					GetTooltipText() const override { return FText::FromString(TEXT("Combines 4 inputs of any type and constructs an array out of it.")); } 
+ };
+
diff --git a/Engine/Plugins/Experimental/TextureGraph/Source/TextureGraph/Public/Expressions/TG_Expression.h b/Engine/Plugins/Experimental/TextureGraph/Source/TextureGraph/Public/Expressions/TG_Expression.h
--- a/Engine/Plugins/Experimental/TextureGraph/Source/TextureGraph/Public/Expressions/TG_Expression.h
+++ b/Engine/Plugins/Experimental/TextureGraph/Source/TextureGraph/Public/Expressions/TG_Expression.h
@@ -27,6 +27,7 @@ public:
 	static const FName Maths;
 	static const FName Utilities;
 	static const FName Filter;
+	static const FName Arrays;
 	static const FName Custom;
 };
 
diff --git a/Engine/Plugins/Experimental/TextureGraph/Source/TextureGraph/Public/TG_GraphEvaluation.h b/Engine/Plugins/Experimental/TextureGraph/Source/TextureGraph/Public/TG_GraphEvaluation.h
--- a/Engine/Plugins/Experimental/TextureGraph/Source/TextureGraph/Public/TG_GraphEvaluation.h
+++ b/Engine/Plugins/Experimental/TextureGraph/Source/TextureGraph/Public/TG_GraphEvaluation.h
@@ -88,15 +88,20 @@ struct TEXTUREGRAPH_API FTG_Evaluation
 	static const FString GFloatToTextureAutoConv_Name;
 
 	static void EvaluateGraph(UTG_Graph* InGraph, FTG_EvaluationContext* Context);
+	static void EvaluateNodeArray(UTG_Node* InNode, const TArray<TObjectPtr<UTG_Pin>>& ArrayInputs, 
+		const TArray<TObjectPtr<UTG_Pin>>& NonArrayPins, int MaxCount, FTG_EvaluationContext* Context);
 	static void EvaluateNode(UTG_Node* InNode, FTG_EvaluationContext* Context);
 
-	static void TransferVarToPin(UTG_Pin* InPin, FTG_EvaluationContext* Context);
+	static void TransferVarToPin(UTG_Pin* InPin, FTG_EvaluationContext* Context, int Index);
+	static int FilterArrayInputs(FTG_EvaluationContext* InContext, const TArray<TObjectPtr<UTG_Pin>>& InPins, 
+		TArray<TObjectPtr<UTG_Pin>>& ArrayPins, TArray<TObjectPtr<UTG_Pin>>& NonArrayPins);
 
 	// Converter used to introduce transformation from one var to another var when the arguments are compatible but require to be converted
 	struct VarConverterInfo
 	{
 		FTG_Var* InVar = nullptr;
 		FTG_Var* OutVar = nullptr;
+		int Index = -1;
 		FTG_EvaluationContext* Context = nullptr;
 	};
 	typedef std::function<void(VarConverterInfo& Info)> VarConverter;
@@ -110,7 +115,8 @@ struct TEXTUREGRAPH_API FTG_Evaluation
     {
     	FTG_Var* InVar = nullptr;
     	FTG_Var* OutVar = nullptr;
-    	FTG_EvaluationContext* Context = nullptr;
+		int Index = -1;
+		FTG_EvaluationContext* Context = nullptr;
     };
     typedef std::function<bool(VarConformerInfo& Info)> VarConformer;
 	
diff --git a/Engine/Plugins/Experimental/TextureGraph/Source/TextureGraph/Public/TG_Pin.h b/Engine/Plugins/Experimental/TextureGraph/Source/TextureGraph/Public/TG_Pin.h
--- a/Engine/Plugins/Experimental/TextureGraph/Source/TextureGraph/Public/TG_Pin.h
+++ b/Engine/Plugins/Experimental/TextureGraph/Source/TextureGraph/Public/TG_Pin.h
@@ -7,6 +7,7 @@
 #include "TG_Signature.h"
 #include "TG_Var.h"
 #include "TG_Texture.h"
+#include "TG_Variant.h"
 
 #include "TG_Pin.generated.h"
 
@@ -169,6 +170,7 @@ public:
 	bool IsArgString() const { return Argument.IsString(); }
 	bool IsArgColor() const { return Argument.IsColor(); }
 	bool IsArgVector() const { return Argument.IsVector(); }
+	bool IsArgArray() const { return Argument.IsArray() || (GetSelfVar() && GetSelfVar()->IsArray()); }
 	bool IsArgTexture() const { return Argument.IsTexture(); }
 	bool IsArgVariant() const { return Argument.IsVariant(); }
 	bool IsArgObject(UClass* InClass) const { return Argument.IsObject(InClass); }
@@ -193,6 +195,9 @@ public:
 	// Texture value getter only on a Texture argument
 	bool GetValue(FTG_Texture& OutValue) const;
 
+	// Texture value getter only on a Texture argument
+	bool GetValue(FTG_VariantArray& OutValue) const;
+
 	// Variant value getter only if a Variant argument OR if the argument is one of the supported type by Variant (Scalar/Color/Vector/Texture)
 	// After getting the value, the next step is to check the type contained in the variant.
 	bool GetValue(FTG_Variant& OutValue) const;
@@ -226,6 +231,7 @@ public:
 
 	bool SetValue(const float* Value, size_t Count);
 	bool SetValue(FTG_Texture&) { return false; } // this is not a possible case
+	bool SetValue(FTG_VariantArray&) { return false; } // this is not a possible case
 	bool SetValue(const FTG_TextureDescriptor& Value); 
 
 	// Template version of the setter works correctly as long as the selfvar contains exactly the specified type
diff --git a/Engine/Plugins/Experimental/TextureGraph/Source/TextureGraph/Public/TG_Signature.h b/Engine/Plugins/Experimental/TextureGraph/Source/TextureGraph/Public/TG_Signature.h
--- a/Engine/Plugins/Experimental/TextureGraph/Source/TextureGraph/Public/TG_Signature.h
+++ b/Engine/Plugins/Experimental/TextureGraph/Source/TextureGraph/Public/TG_Signature.h
@@ -156,6 +156,7 @@ public:
 	bool				IsColor() const; // FlinearColor
 	bool				IsVector() const; // FLinearColor or ...
 	bool				IsTexture() const; // FTG_Texture or ...
+	bool				IsArray() const; // FTG_TextureArray or ...
 	bool				IsVariant() const; // FTG_Variant
 	bool				IsObject(UClass* InClass) const; // check that the Argument is a UObject and of baseType InClass
 	bool				GetMetaData(const FName KeyToFind, FString& OutValue);
diff --git a/Engine/Plugins/Experimental/TextureGraph/Source/TextureGraph/Public/TG_Var.h b/Engine/Plugins/Experimental/TextureGraph/Source/TextureGraph/Public/TG_Var.h
--- a/Engine/Plugins/Experimental/TextureGraph/Source/TextureGraph/Public/TG_Var.h
+++ b/Engine/Plugins/Experimental/TextureGraph/Source/TextureGraph/Public/TG_Var.h
@@ -49,7 +49,7 @@ struct TEXTUREGRAPH_API FTG_Var
 	static void RegisterVarPropertySerializer(FName CPPTypeName, VarPropertySerializer Serializer);
 	static void UnregisterVarPropertySerializer(FName CPPTypeName);
 
-	bool CopyGeneric(UTG_Expression* Owner, const FTG_Argument& Arg, bool CopyVarToProperty);
+	bool CopyGeneric(UTG_Expression* Owner, const FTG_Argument& Arg, bool CopyVarToProperty, int Index = -1);
 
 	// Map of per type serializer for the var to/from an Archive
 	struct VarArchiveSerialInfo
@@ -98,9 +98,10 @@ struct TEXTUREGRAPH_API FTG_Var
 	}
 
 public:
-	FTG_Var()									{}
-	FTG_Var(FTG_Id InPinUuid)					: PinId(InPinUuid) {}
-	FTG_Var(const FTG_Var& InVar)				: PinId(InVar.PinId), Concept(InVar.Concept) {}
+	FTG_Var();
+	FTG_Var(FTG_Id InPinUuid);
+	FTG_Var(const FTG_Var& InVar);
+
 	FTG_Var& operator = (const FTG_Var& InVar) { PinId = InVar.PinId; Concept = InVar.Concept; return *this; }
 
 	void ShareData(const FTG_Var& InVar)
@@ -115,6 +116,8 @@ public:
 
 	// Check if the Var has No Data allocated
 	bool IsEmpty() const { return !Concept; } 
+	bool IsArray() const { return Concept && Concept->IsArray(); }
+	void SetArray() { check(Concept);  Concept->SetArray(); }
 
 	FString LogValue() const { if (Concept) return Concept->LogValue(); else return TEXT("nullptr"); }
 
@@ -122,19 +125,14 @@ public:
 	
 	// Manage data Arg provide the type name and rely on the VarAllocator map to create the concrete cpp type
 	// If Owner is provided then Var can be initialized from the matching owner's property if found
-	void CopyTo(UTG_Expression* Owner, FTG_Argument& Arg);
-	void CopyFrom(UTG_Expression* Owner, FTG_Argument& Arg);
+	void CopyTo(UTG_Expression* Owner, FTG_Argument& Arg, int Index = -1);
+	void CopyFrom(UTG_Expression* Owner, FTG_Argument& Arg, int Index = -1);
 
-	void CopyTo(FTG_Var* InVar)
-	{
-		if (Concept && InVar)
-		{
-			InVar->Concept = Concept->Clone();
-		}
-	}
+	void CopyTo(FTG_Var* InVar);
 	void CopyFrom(FTG_Var* InVar);
 
 	// Manage the Var Data with the knowledge of the Type
+	void Reset() { Concept = nullptr; }
 	template <class T> void ResetAs() const { Concept = MakeShared<FModel<T>>(); }
 	template <class T> T& GetAs() const
 	{
@@ -175,23 +173,16 @@ private:
 		virtual FString LogValue() = 0;
 		virtual void SetValueFromString(const FString& String) = 0;
 		virtual TSharedPtr<FConcept> Clone() = 0;
-
-		/// Array related
-		virtual int32 Count() const 
-		{ 
-			return 1; 
-		}
-
-		virtual bool IsArray() 
-		{ 
-			return Count() > 1; 
-		}
+		virtual void SetArray() = 0;
+		virtual bool IsArray() const = 0;
 	};
 
 	template <class T>
 	struct TEXTUREGRAPH_API FModel : public FConcept
 	{
 		T Value;
+		bool bIsArray = false;
+
 		FString LogValue() override 
 		{ 
 			return TG_Var_LogValue(Value); 
@@ -206,47 +197,17 @@ private:
 		{
 			auto Ptr = MakeShared<FModel<T>>();
 			Ptr->Value = Value;
+			Ptr->bIsArray = bIsArray;
 			return Ptr;
 		}
-	};
-
-	template <class T>
-	struct TEXTUREGRAPH_API FModelArray : public FModel<TArray<T> >
-	{
-		using FModelBase = FModel<TArray<T>>;
-		virtual int32 Count() const override
-		{ 
-			return FModelBase::Value.Num(); 
-		}
-		virtual bool IsArray() 
-		{ 
-			return true; 
-		}
 
-		FString LogValue() override 
+		virtual void SetArray() 
 		{ 
-			FString Log = "[";
-			for (int32 Index = 0; Index < FModelBase::Value.Num() - 1; Index++)
-				Log += (TG_Var_LogValue<T>(FModelBase::Value[Index]) + TEXT(","));
-
-			if (!FModelBase::Value.IsEmpty()) {
-				Log += TG_Var_LogValue<T>(FModelBase::Value[FModelBase::Value.Num() - 1]);
-			}
-
-			Log += TEXT("]");
-
-			return Log; 
+			bIsArray = true; 
 		}
-
-		void SetValueFromString(const FString& StrValue) override 
+		virtual bool IsArray() const 
 		{ 
-			TArray<FString> ArrayValues;
-			StrValue.ParseIntoArray(ArrayValues, TEXT(","));
-			FModelBase::Value.Empty();
-			FModelBase::Value.SetNum(ArrayValues.Num());
-
-			for (int32 Index = 0; Index < ArrayValues.Num(); Index++)
-				TG_Var_SetValueFromString<T>(FModelBase::Value[Index], ArrayValues[Index]);
+			return bIsArray; 
 		}
 	};
 
diff --git a/Engine/Plugins/Experimental/TextureGraph/Source/TextureGraph/Public/TG_Variant.h b/Engine/Plugins/Experimental/TextureGraph/Source/TextureGraph/Public/TG_Variant.h
--- a/Engine/Plugins/Experimental/TextureGraph/Source/TextureGraph/Public/TG_Variant.h
+++ b/Engine/Plugins/Experimental/TextureGraph/Source/TextureGraph/Public/TG_Variant.h
@@ -24,6 +24,9 @@ enum class ETG_VariantType : int8
 };
 ENUM_RANGE_BY_FIRST_AND_LAST(ETG_VariantType, ETG_VariantType::Scalar, ETG_VariantType::Texture);
 
+//////////////////////////////////////////////////////////////////////////
+/// Base class for making working with Variants easier and less repetitive
+//////////////////////////////////////////////////////////////////////////
 USTRUCT()
 struct TEXTUREGRAPH_API FTG_Variant
 {
@@ -177,6 +180,47 @@ struct TStructOpsTypeTraits<FTG_Variant>
 		WithIdenticalViaEquality = true
 	};
 };
+
 //////////////////////////////////////////////////////////////////////////
-/// Base class for making working with Variants easier and less repetitive
+/// Variant Arrays
 //////////////////////////////////////////////////////////////////////////
+USTRUCT()
+struct TEXTUREGRAPH_API FTG_VariantArray
+{
+	GENERATED_BODY()
+
+private: 
+	TArray<FTG_Variant> Array;			/// An array of textures
+	int ViewIndex = 0;
+
+public:
+	void SetNum(int Count);
+	void Set(int Index, const FTG_Variant& Var);
+	void Set(int Index, const FTG_Texture& Tex);
+	void SetViewIndex(int Index);
+
+	FORCEINLINE int Num() const
+	{
+		return Array.Num();
+	}
+
+	FORCEINLINE FTG_Variant& Get(int Index)
+	{
+		check(Index >= 0 && Index < Array.Num());
+		return Array[Index];
+	}
+
+	FORCEINLINE const TArray<FTG_Variant>& GetArray() const
+	{
+		return Array;
+	}
+
+	FORCEINLINE FTG_Texture GetViewTexture() const
+	{
+		if (ViewIndex >= 0 && ViewIndex < Array.Num() && Array[ViewIndex].IsTexture())
+			return Array[ViewIndex].GetTexture();
+
+		return FTG_Texture::GetBlack();
+	}
+};
+
diff --git a/Engine/Plugins/Experimental/TextureGraph/Source/TextureGraphEditor/Private/EdGraph/TG_EdGraph.cpp b/Engine/Plugins/Experimental/TextureGraph/Source/TextureGraphEditor/Private/EdGraph/TG_EdGraph.cpp
--- a/Engine/Plugins/Experimental/TextureGraph/Source/TextureGraphEditor/Private/EdGraph/TG_EdGraph.cpp
+++ b/Engine/Plugins/Experimental/TextureGraph/Source/TextureGraphEditor/Private/EdGraph/TG_EdGraph.cpp
@@ -142,26 +142,35 @@ void UTG_EdGraph::OnNodeCreateThumbnail(UTG_Node* InNode, const FTG_EvaluationCo
 			TArray<UTG_Pin*> OutputPins;
 			InNode->GetOutputPins(OutputPins);
 			// Loop over all Output pins
-			for(UTG_Pin* Pin : OutputPins)
+			for (UTG_Pin* Pin : OutputPins)
 			{
 				check (Pin);
-				if (Pin->IsArgTexture())
-				{
-					FTG_Texture OutTexture;
-					if(Pin->GetValue(OutTexture))
-					{
-						UMixInterface* Mix = InContext->Cycle->GetMix();
-						auto TargetId = InContext->TargetId;
+				FTG_Texture OutTexture;
 
-						if (!OutTexture.RasterBlob)
-						{
-							OutTexture = FTG_Texture::GetBlack();
-						}
+				if (Pin->IsArgArray())
+				{
+					FTG_VariantArray OutVarArray;
+					Pin->GetValue(OutVarArray);
+					OutTexture = OutVarArray.GetViewTexture();
+				}
+				else if (Pin->IsArgTexture())
+				{
+					Pin->GetValue(OutTexture);
+				}
 
-						TiledBlobPtr ThumbBlob = T_Thumbnail::Bind(Mix, Pin, OutTexture.RasterBlob, TargetId);
+				if (OutTexture)
+				{
+					UMixInterface* Mix = InContext->Cycle->GetMix();
+					int32 TargetId = InContext->TargetId;
 
-						CacheThumbBlob(Pin->GetId(), ThumbBlob);
+					if (!OutTexture.RasterBlob)
+					{
+						OutTexture = FTG_Texture::GetBlack();
 					}
+
+					TiledBlobPtr ThumbBlob = T_Thumbnail::Bind(Mix, Pin, OutTexture.RasterBlob, TargetId);
+
+					CacheThumbBlob(Pin->GetId(), ThumbBlob);
 				}
 			}
 		}
diff --git a/Engine/Plugins/Experimental/TextureGraph/Source/TextureGraphEngine/TextureGraphErrorReporter.h b/Engine/Plugins/Experimental/TextureGraph/Source/TextureGraphEngine/TextureGraphErrorReporter.h
--- a/Engine/Plugins/Experimental/TextureGraph/Source/TextureGraphEngine/TextureGraphErrorReporter.h
+++ b/Engine/Plugins/Experimental/TextureGraph/Source/TextureGraphEngine/TextureGraphErrorReporter.h
@@ -22,7 +22,8 @@ enum class ETextureGraphErrorType: int32
 	RECURSIVE_CALL,
 	UNSUPPORTED_MATERIAL,
 	SUBGRAPH_INTERNAL_ERROR,
-	NODE_WARNING
+	NODE_WARNING,
+	INPUT_ARRAY_WARNING
 };
 class TEXTUREGRAPHENGINE_API FTextureGraphErrorReporter
 {
