Array opeations

diff --git a/Engine/Plugins/Experimental/TextureGraph/Source/TextureGraph/Private/Expressions/Arrays/TG_Expression_ArrayOps.cpp b/Engine/Plugins/Experimental/TextureGraph/Source/TextureGraph/Private/Expressions/Arrays/TG_Expression_ArrayOps.cpp
new file mode 100644
--- /dev/null
+++ b/Engine/Plugins/Experimental/TextureGraph/Source/TextureGraph/Private/Expressions/Arrays/TG_Expression_ArrayOps.cpp
@@ -0,0 +1,78 @@
+// Copyright Epic Games, Inc. All Rights Reserved.
+
+#include "Expressions/Arrays/TG_Expression_ArrayOps.h"
+
+#include "TG_Graph.h"
+#include "2D/TextureHelper.h"
+#include "Model/StaticImageResource.h"
+
+/////////////////////////////////////////////////////////////
+/// Array concatenation
+/////////////////////////////////////////////////////////////
+void UTG_Expression_ArrayConcat::Evaluate(FTG_EvaluationContext* InContext)
+{
+	Super::Evaluate(InContext);
+
+	Output.SetNum(Input1.Num() + Input2.Num());
+
+	for (int i = 0; i < Input1.Num(); i++)
+		Output.Set(i, Input1.GetArray()[i]);
+
+	for (int i = 0; i < Input2.Num(); i++)
+		Output.Set(i + Input1.Num(), Input2.GetArray()[i]);
+}
+
+/////////////////////////////////////////////////////////////
+/// Array slicing/splicing
+/////////////////////////////////////////////////////////////
+void UTG_Expression_ArraySplit::Evaluate(FTG_EvaluationContext* InContext)
+{
+	Super::Evaluate(InContext);
+
+	if (Input.GetArray().IsEmpty())
+	{
+		Sliced.SetNum(1);
+		Spliced.SetNum(1);
+		Sliced.Set(0, FTG_Texture::GetBlack());
+		Spliced.Set(0, FTG_Texture::GetBlack());
+		return;
+	}
+
+	if (EndIndex < 0 || EndIndex < StartIndex)
+		EndIndex = Input.Num();
+	if (StartIndex < 0)
+		StartIndex = 0;
+
+	check(StartIndex <= EndIndex);
+
+	int32 InputCount = Input.Num();
+	int32 SlicedCount = EndIndex - StartIndex;
+	int32 SplicedCount = StartIndex + (InputCount - EndIndex);
+
+	if (SlicedCount > 0)
+	{
+		Sliced.SetNum(SlicedCount);
+		for (int i = 0; i < SlicedCount; i++)
+			Sliced.Set(i, Input.GetArray()[StartIndex + i]);
+	}
+	else
+	{
+		Sliced.SetNum(1);
+		Sliced.Set(0, FTG_Texture::GetBlack());
+	}
+
+	if (SplicedCount > 0)
+	{
+		Spliced.SetNum(SplicedCount);
+
+		for (int i = 0; i < StartIndex; i++)
+			Spliced.Set(i, Input.GetArray()[i]);
+		for (int i = EndIndex; i < Input.Num(); i++)
+			Spliced.Set(StartIndex + i, Input.GetArray()[i]);
+	}
+	else
+	{
+		Spliced.SetNum(1);
+		Spliced.Set(0, FTG_Texture::GetBlack());
+	}
+}
diff --git a/Engine/Plugins/Experimental/TextureGraph/Source/TextureGraph/Private/TG_Var.cpp b/Engine/Plugins/Experimental/TextureGraph/Source/TextureGraph/Private/TG_Var.cpp
--- a/Engine/Plugins/Experimental/TextureGraph/Source/TextureGraph/Private/TG_Var.cpp
+++ b/Engine/Plugins/Experimental/TextureGraph/Source/TextureGraph/Private/TG_Var.cpp
@@ -227,8 +227,7 @@ void Generic_Struct_Serializer(FTG_Var::VarPropertySerialInfo& Info)
 	{
 		if (!Info.Var->IsEmpty() && Info.Var->IsValid())
 		{
-			if (!Info.Var->IsArray())
-				*PropertyValue = Info.Var->GetAs<T_ValueType>();
+			*PropertyValue = Info.Var->GetAs<T_ValueType>();
 		}
 		else
 		{
diff --git a/Engine/Plugins/Experimental/TextureGraph/Source/TextureGraph/Public/Expressions/Arrays/TG_Expression_ArrayOps.h b/Engine/Plugins/Experimental/TextureGraph/Source/TextureGraph/Public/Expressions/Arrays/TG_Expression_ArrayOps.h
new file mode 100644
--- /dev/null
+++ b/Engine/Plugins/Experimental/TextureGraph/Source/TextureGraph/Public/Expressions/Arrays/TG_Expression_ArrayOps.h
@@ -0,0 +1,77 @@
+// Copyright Epic Games, Inc. All Rights Reserved.
+
+#pragma once
+
+#include "CoreMinimal.h"
+#include "Expressions/TG_Expression.h"
+#include "TG_Texture.h"
+#include "TG_Expression_ArrayOps.generated.h"
+
+/////////////////////////////////////////////////////////////
+/// Array concatenation
+/////////////////////////////////////////////////////////////
+UCLASS()
+class TEXTUREGRAPH_API UTG_Expression_ArrayConcat : public UTG_Expression
+{
+	 GENERATED_BODY()
+
+public:
+	TG_DECLARE_EXPRESSION(TG_Category::Arrays);
+	virtual void					Evaluate(FTG_EvaluationContext* InContext) override;
+	
+	/// The first array as part of the concatenation process
+	UPROPERTY(meta = (TGType = "TG_Input", PinDisplayName = ""))
+	FTG_VariantArray				Input1;
+
+	/// The second array as part of the concatenation process
+	UPROPERTY(meta = (TGType = "TG_Input", PinDisplayName = ""))
+	FTG_VariantArray				Input2;
+
+	/// Concatenated array
+	UPROPERTY(meta = (TGType = "TG_Output", PinDisplayName = "", HideInnerPropertiesInNode))
+	FTG_VariantArray				Output;
+
+	virtual FTG_Name				GetDefaultName() const override { return TEXT("Concatenate Array");}
+	virtual FText					GetTooltipText() const override { return FText::FromString(TEXT("Combines two arrays into a single one.")); } 
+ };
+
+/////////////////////////////////////////////////////////////
+/// Split array
+/////////////////////////////////////////////////////////////
+UCLASS()
+class TEXTUREGRAPH_API UTG_Expression_ArraySplit: public UTG_Expression
+{
+	 GENERATED_BODY()
+
+public:
+	TG_DECLARE_EXPRESSION(TG_Category::Arrays);
+	virtual void					Evaluate(FTG_EvaluationContext* InContext) override;
+	
+	/// The array to slice
+	UPROPERTY(meta = (TGType = "TG_Input", PinDisplayName = "", MinValue))
+	FTG_VariantArray				Input;
+
+	/// The starting index of the slicing operation. 0-indexed meaning that 0 indicates the first element of the input array
+	UPROPERTY(meta = (TGType = "TG_Input", PinDisplayName = ""))
+	int32							StartIndex = 0;
+
+	/// The ending index of the slicing operation (non-inclusive). Negative values mean the end of the array.
+	UPROPERTY(meta = (TGType = "TG_Input", PinDisplayName = ""))
+	int32							EndIndex = -1;
+
+	/// The ending index of the slicing operation (non-inclusive). Negative values mean the end of the array.
+	UPROPERTY(meta = (TGType = "TG_Setting", PinDisplayName = "Splice"))
+	bool							bSplice = false;
+
+	/// Sliced part of the array
+	UPROPERTY(meta = (TGType = "TG_Output", PinDisplayName = "Sliced", HideInnerPropertiesInNode))
+	FTG_VariantArray				Sliced;
+
+	/// Spliced part of thea array
+	UPROPERTY(meta = (TGType = "TG_Output", PinDisplayName = "Spliced", HideInnerPropertiesInNode))
+	FTG_VariantArray				Spliced;
+
+	virtual FTG_Name				GetDefaultName() const override { return TEXT("Modify Array (Slice/Splice)");}
+	virtual FText					GetTooltipText() const override { return FText::FromString(TEXT("Slice a portion of an input array and put that as the output of this array. Alternatively splice an array and remove some elements from the input array")); } 
+ };
+
