Texture Graph Material Node: Supporting StaticSwitch setting
In order to support assigning the static bool parameter on a Material used in a Node from an exposed Bool pin,
we had to change the underneath implementation from a MaterialInstanceDynamic to a MaterialInstanceConstant configured with the proper bool set permutation.
The various permutations requested are created on demand and cached at runtime

diff --git a/Engine/Plugins/Experimental/TextureGraph/Source/TextureGraph/Private/Expressions/Input/TG_Expression_Material.cpp b/Engine/Plugins/Experimental/TextureGraph/Source/TextureGraph/Private/Expressions/Input/TG_Expression_Material.cpp
--- a/Engine/Plugins/Experimental/TextureGraph/Source/TextureGraph/Private/Expressions/Input/TG_Expression_Material.cpp
+++ b/Engine/Plugins/Experimental/TextureGraph/Source/TextureGraph/Private/Expressions/Input/TG_Expression_Material.cpp
@@ -30,6 +30,7 @@ void UTG_Expression_Material::PostEditUndo()
 	// So recreate it internally without notifying, normally, the node's pins should match
 	DynSignature.Reset();
 	GetSignature();
+	FeedbackPinValue(GET_MEMBER_NAME_CHECKED(UTG_Expression_Material, RenderedAttribute), RenderedAttribute);
 
 	Super::PostEditUndo();
 }
@@ -109,7 +110,7 @@ void UTG_Expression_Material::SetMaterial(UMaterialInterface* InMaterial)
 	if (InMaterial == Material)
 	{
 		// Just check that the MaterialInstance is valid, if not reassign below
-		if ((!Material) || (Material && MaterialCopy && MaterialInstance))
+		if ((!Material) || (Material && MaterialCopy && MaterialPermutations))
 			return;
 	}
 
diff --git a/Engine/Plugins/Experimental/TextureGraph/Source/TextureGraph/Private/Expressions/Input/TG_Expression_MaterialFunction.cpp b/Engine/Plugins/Experimental/TextureGraph/Source/TextureGraph/Private/Expressions/Input/TG_Expression_MaterialFunction.cpp
--- a/Engine/Plugins/Experimental/TextureGraph/Source/TextureGraph/Private/Expressions/Input/TG_Expression_MaterialFunction.cpp
+++ b/Engine/Plugins/Experimental/TextureGraph/Source/TextureGraph/Private/Expressions/Input/TG_Expression_MaterialFunction.cpp
@@ -59,14 +59,12 @@ void UTG_Expression_MaterialFunction::PostEditUndo()
 
 void UTG_Expression_MaterialFunction::Initialize()
 {
-	if (MaterialFunction && !MaterialInstance)
+	if (MaterialFunction && !ReferenceMaterial)
 	{
 		ReferenceMaterial = CreateMaterialReference();
-		if (ReferenceMaterial)
-		{
-			MaterialInstance = UMaterialInstanceDynamic::Create(ReferenceMaterial, this);
-		}
 	}
+
+	Super::Initialize(); // Allocate the BaseMaterial infrastructure using the ReferenceMaterial.
 }
 
 void UTG_Expression_MaterialFunction::SetMaterialFunctionInternal(UMaterialFunctionInterface* InMaterialFunction)
@@ -101,7 +99,7 @@ void UTG_Expression_MaterialFunction::SetMaterialFunction(UMaterialFunctionInter
 	if (InMaterialFunction == MaterialFunction)
 	{
 		// Just check that the MaterialInstance is valid, if not reassign below
-		if ((!MaterialFunction) || (MaterialFunction && MaterialInstance))
+		if ((!MaterialFunction) || (MaterialFunction && MaterialPermutations))
 			return;
 	}
 
diff --git a/Engine/Plugins/Experimental/TextureGraph/Source/TextureGraph/Private/Expressions/TG_Expression_MaterialBase.cpp b/Engine/Plugins/Experimental/TextureGraph/Source/TextureGraph/Private/Expressions/TG_Expression_MaterialBase.cpp
--- a/Engine/Plugins/Experimental/TextureGraph/Source/TextureGraph/Private/Expressions/TG_Expression_MaterialBase.cpp
+++ b/Engine/Plugins/Experimental/TextureGraph/Source/TextureGraph/Private/Expressions/TG_Expression_MaterialBase.cpp
@@ -11,6 +11,7 @@
 #include "TG_Var.h"
 #include "FxMat/MaterialManager.h"
 #include "Materials/MaterialAttributeDefinitionMap.h"
+#include "Materials/MaterialInstanceConstant.h"
 
 #if WITH_EDITOR
 #include "TextureCompiler.h"
@@ -78,24 +79,15 @@ bool UTG_Expression_MaterialBase::Validate(MixUpdateCyclePtr	Cycle)
 	return true;
 }
 
-
 void UTG_Expression_MaterialBase::SetMaterialInternal(UMaterialInterface* InMaterial)
 {
 
 	if (!InMaterial)
 	{
-		MaterialInstance = nullptr;
+		MaterialPermutations.Reset();
 	}
 	else {
-		if (InMaterial->IsA<UMaterial>())
-		{
-			// create a material instance
-			MaterialInstance = UMaterialInstanceDynamic::Create(InMaterial, this);
-		}
-		else if (InMaterial->IsA<UMaterialInstance>())
-		{
-			MaterialInstance = UMaterialInstanceDynamic::Create(InMaterial, this);
-		}
+		MaterialPermutations = FMaterialInstanceStaticSwitchPermutationMap::Create(GetMaterial());
 	}
 
 	// Detect the set of available material properties for rendering
@@ -105,7 +97,6 @@ void UTG_Expression_MaterialBase::SetMaterialInternal(UMaterialInterface* InMate
 	NotifySignatureChanged();
 }
 
-
 void UTG_Expression_MaterialBase::GenerateMaterialAttributeOptions()
 {
 	// Detect the set of available material properties for rendering
@@ -136,13 +127,38 @@ void UTG_Expression_MaterialBase::GenerateMaterialAttributeOptions()
 
 void UTG_Expression_MaterialBase::Initialize()
 {
-	if(GetMaterial() && !MaterialInstance)
+	if (GetMaterial() && !MaterialPermutations)
 	{
-		MaterialInstance = UMaterialInstanceDynamic::Create(GetMaterial(), this);
+		MaterialPermutations = FMaterialInstanceStaticSwitchPermutationMap::Create(GetMaterial());
 	}
 	GenerateMaterialAttributeOptions(); // also populate the attributes availables
 }
 
+UMaterialInstanceConstant* UTG_Expression_MaterialBase::GetEvaluationMaterialInstanceConstant(FTG_EvaluationContext* InContext)
+{
+	TArray<FStaticSwitchParameter> EvaluationCombination;
+
+	int32 ArgIndex = 0;
+	for (auto& ArgToMatParam : ArgToMatParams)
+	{
+		if (ArgToMatParam.MatType == EMaterialParameterType::StaticSwitch)
+		{
+			FStaticSwitchParameter StaticSwitchParameter(FMaterialParameterInfo(ArgToMatParam.MatParamName), ArgToMatParam.Value.AsStaticSwitch(), false, ArgToMatParam.MatParamGuid);
+
+			FTG_Var* Var = InContext->Inputs.GetVar(ArgToMatParam.ArgName);
+			if (Var && !Var->IsEmpty())
+			{
+				auto ParamValue = Var->EditAs<bool>();
+				StaticSwitchParameter.Value = ParamValue;
+				StaticSwitchParameter.bOverride = true;
+			}
+			EvaluationCombination.Add(StaticSwitchParameter);
+		}
+	}
+
+	return MaterialPermutations->GetMaterialInstance(EvaluationCombination);
+}
+
 void UTG_Expression_MaterialBase::Evaluate(FTG_EvaluationContext* InContext)
 {
 	Super::Evaluate(InContext);
@@ -154,8 +170,11 @@ void UTG_Expression_MaterialBase::Evaluate(FTG_EvaluationContext* InContext)
 
 	if (GetMaterial())
 	{
+
 		FString AssetName = GetMaterial()->GetName();
-		const auto RenderMaterial = std::make_shared<RenderMaterial_BP>(AssetName, GetMaterial(), nullptr);
+
+		UMaterialInstanceConstant* MIC = GetEvaluationMaterialInstanceConstant(InContext);
+		const auto RenderMaterial = std::make_shared<RenderMaterial_BP>(AssetName, MIC);
 		Result = CreateRenderMaterialJob(InContext, RenderMaterial, Output.GetBufferDescriptor(), GetRenderedAttributeId());
 	}
 
@@ -247,6 +266,7 @@ void UTG_Expression_MaterialBase::LinkMaterialParameters(FTG_EvaluationContext*
 {
 	if (InMaterialJob)
 	{	
+		int32 ArgIndex = 0;
 		for (auto& ArgToMatParam : ArgToMatParams)
 		{
 			FTG_Var* Var = InContext->Inputs.GetVar(ArgToMatParam.ArgName);
@@ -308,17 +328,15 @@ void UTG_Expression_MaterialBase::LinkMaterialParameters(FTG_EvaluationContext*
 				}
 				break;
 				case EMaterialParameterType::StaticSwitch:
-				{
-#if WITH_EDITORONLY_DATA
-					auto ParamValue = Var->EditAs<bool>();
-					InMaterialJob->AddArg(ARG_INT(ParamValue, TCHAR_TO_UTF8(*ArgToMatParam.MatParamName.ToString())));
-#endif
-				}
+					// Static switch are not passed to the BlobTransofrm as arg, they are part of the evaluation method 
+					// to decide WHICH material permutation to use
 				break;
 				default:
 					break;
 				}
 			}
+
+			ArgIndex++;
 		}
 	}
 }
@@ -343,7 +361,7 @@ FName UTG_Expression_MaterialBase::CPPTypeNameFromMaterialParamType(EMaterialPar
 }
 
 void UTG_Expression_MaterialBase::AddSignatureParam(const TArray<FMaterialParameterInfo>& OutParameterInfo, const TArray<FGuid>& OutGuids,
-		EMaterialParameterType MatType, FTG_Signature::FInit& SignatureInit) const
+		EMaterialParameterType MatParamType, FTG_Signature::FInit& SignatureInit) const
 {
 	for (int i = 0; i < OutParameterInfo.Num(); ++i)
 	{
@@ -352,7 +370,7 @@ void UTG_Expression_MaterialBase::AddSignatureParam(const TArray<FMaterialParame
 
 		TMap<FName, FString> MetaDataMap;
 #if WITH_EDITOR
-		if(MatType == EMaterialParameterType::Scalar)
+		if(MatParamType == EMaterialParameterType::Scalar)
 		{
 			// In case of scalar parameter, material has input range. We use that range to handle the input in the node.
 			float MinValue, MaxValue;
@@ -370,13 +388,17 @@ void UTG_Expression_MaterialBase::AddSignatureParam(const TArray<FMaterialParame
 		FName ArgName = TG_MakeNameUniqueInCollection(MatParam.Name, ArgNames);
 	
 		// New Arg of the signature
-		FTG_Argument Arg{ ArgName, CPPTypeNameFromMaterialParamType(MatType), {ETG_Access::In}, MetaDataMap};
+		FTG_Argument Arg{ ArgName, CPPTypeNameFromMaterialParamType(MatParamType), {ETG_Access::In}, MetaDataMap};
 		
 		Arg.SetPersistentSelfVar(); // Set the material parameter persistent selfvar in order to save the state
 		SignatureInit.Arguments.Add(Arg);
 
+		// Grab the actual default value from the material
+		FMaterialParameterMetadata MatParamMetadata;
+		GetMaterial()->GetParameterValue(MatParamType, MatParam, MatParamMetadata);
+
 		// And new entry in ArgtoMatParams
-		ArgToMatParams.Add({ ArgName, MatParam.Name, Guid, MatType });
+		ArgToMatParams.Add({ ArgName, MatParam.Name, Guid, MatParamType, MatParamMetadata.Value});
 	}
 }
 
@@ -388,7 +410,6 @@ FTG_SignaturePtr UTG_Expression_MaterialBase::BuildSignatureDynamically() const
 	// start fresh here
 	ArgToMatParams.Empty();
 
-
 	// append
 	if (GetMaterial())
 	{
@@ -422,7 +443,7 @@ FTG_SignaturePtr UTG_Expression_MaterialBase::BuildSignatureDynamically() const
 void UTG_Expression_MaterialBase::CopyVarGeneric(const FTG_Argument& Arg, FTG_Var* InVar, bool CopyVarToArg)
 {
 	FArgToMaterialParamInfo* ArgToMatParam = ArgToMatParams.FindByKey(Arg.Name);
-	if (MaterialInstance && ArgToMatParam)
+	if (ArgToMatParam)
 	{
 		switch (ArgToMatParam->MatType)
 		{
@@ -430,12 +451,11 @@ void UTG_Expression_MaterialBase::CopyVarGeneric(const FTG_Argument& Arg, FTG_Va
 		{
 			if (CopyVarToArg)
 			{
-				MaterialInstance->SetScalarParameterValue(Arg.GetName(), InVar->GetAs<float>());
+				ArgToMatParam->Value = FMaterialParameterValue(InVar->GetAs<float>());
 			}
 			else
 			{
-				FMaterialParameterInfo ParameterInfo(Arg.GetName());
-				MaterialInstance->GetScalarParameterValue(ParameterInfo, InVar->EditAs<float>());
+				InVar->EditAs<float>() = ArgToMatParam->Value.AsScalar();
 			}
 		}
 		break;
@@ -443,59 +463,34 @@ void UTG_Expression_MaterialBase::CopyVarGeneric(const FTG_Argument& Arg, FTG_Va
 		{
 			if (CopyVarToArg)
 			{
-				MaterialInstance->SetVectorParameterValue(Arg.GetName(), InVar->GetAs<FLinearColor>());
+				ArgToMatParam->Value = FMaterialParameterValue(InVar->GetAs<FLinearColor>());
 			}
 			else
 			{
-				FMaterialParameterInfo ParameterInfo(Arg.GetName());
-				MaterialInstance->GetVectorParameterValue(ParameterInfo, InVar->EditAs<FLinearColor>());
+				InVar->EditAs<FLinearColor>() = ArgToMatParam->Value.AsLinearColor();
 			}
 		}
 		break;
 		case EMaterialParameterType::DoubleVector:
 		{
 			// TODO: Support that case later ?
-			/*if (CopyVarToArg)
-			{
-				MaterialInstance->SetDoubleVectorParameterValue(Arg.GetName(), InVar->GetAs<FVector4d>());
-			}
-			else
-			{
-				FMaterialParameterInfo ParameterInfo(Arg.GetName());
-				MaterialInstance->GetDoubleVectorParameterValue(ParameterInfo, InVar->EditAs<FVector4d>());
-			}*/
 		}
 		break;
 		case EMaterialParameterType::Texture:
 		{
 			// TODO: Support that case later ?
-			/*	if (CopyVarToArg)
-			{
-				MaterialInstance->SetScalarParameterValue(Arg.GetName(), InVar->GetAs<float>());
-			}
-			else
-			{
-				FMaterialParameterInfo ParameterInfo(Arg.GetName());
-				MaterialInstance->GetScalarParameterValue(ParameterInfo, InVar->EditAs<float>());
-			}*/
 		}
 		break;
 		case EMaterialParameterType::StaticSwitch:
 		{
-#if WITH_EDITORONLY_DATA
 			if (CopyVarToArg)
 			{
-				// Disable assigning static bool values to MID for now see UE-209533 & UE-219306
-				// We need to use a different solution for Material with Static Switch, we can't use MID for this, evetnually MIC editor only...
-				//MaterialInstance->SetStaticSwitchParameterValueEditorOnly(Arg.GetName(), InVar->GetAs<bool>());
+				ArgToMatParam->Value = FMaterialParameterValue(InVar->GetAs<bool>());
 			}
 			else
 			{
-				FMaterialParameterInfo ParameterInfo(Arg.GetName());
-				FGuid OutParameterId;
-				MaterialInstance->GetStaticSwitchParameterValue(ParameterInfo, InVar->EditAs<bool>(), OutParameterId);
+				InVar->EditAs<bool>() = ArgToMatParam->Value.AsStaticSwitch();
 			}
-#endif
 		}
 		break;
 		default:
diff --git a/Engine/Plugins/Experimental/TextureGraph/Source/TextureGraph/Public/Expressions/TG_Expression_MaterialBase.h b/Engine/Plugins/Experimental/TextureGraph/Source/TextureGraph/Public/Expressions/TG_Expression_MaterialBase.h
--- a/Engine/Plugins/Experimental/TextureGraph/Source/TextureGraph/Public/Expressions/TG_Expression_MaterialBase.h
+++ b/Engine/Plugins/Experimental/TextureGraph/Source/TextureGraph/Public/Expressions/TG_Expression_MaterialBase.h
@@ -31,6 +31,7 @@ enum class EDrawMaterialAttributeTarget : uint8
 	Count UMETA(Hidden),
 };
 
+struct FMaterialInstanceStaticSwitchPermutationMap;
 
 UCLASS(abstract)
 class TEXTUREGRAPH_API UTG_Expression_MaterialBase : public UTG_Expression
@@ -64,9 +65,9 @@ public:
 	static FName CPPTypeNameFromMaterialParamType(EMaterialParameterType InMatType);
 
 protected:
-	// A local per Instance Material is recreated from the reference material assigned through SetMaterialInternal
-	UPROPERTY(Transient, DuplicateTransient)
-	TObjectPtr<UMaterialInstanceDynamic> MaterialInstance = nullptr;
+	// A local Material Instance Constant is recreated from the reference material assigned through SetMaterialInternal
+	// All static switches permutations are cached in this map and only created on demand
+	TSharedPtr<FMaterialInstanceStaticSwitchPermutationMap> MaterialPermutations;
 
 	TArray<EDrawMaterialAttributeTarget>	AvailableMaterialAttributeIds; // The set of material properties available for rendering
 	TArray<FName>							AvailableMaterialAttributeNames; // same with the attribute names
@@ -79,14 +80,19 @@ protected:
 		FGuid MatParamGuid;
 		EMaterialParameterType MatType = EMaterialParameterType::None;
 
+		FMaterialParameterValue Value;
+
 		bool operator== (const FName& InArgName) const { return ArgName == InArgName; }
 	};
-	mutable TArray<FArgToMaterialParamInfo> ArgToMatParams; // mutable because it is populated in the BuildSignatureDynamically()
+	mutable TArray<FArgToMaterialParamInfo> ArgToMatParams;			// mutable because it is populated in the BuildSignatureDynamically()
 
 	virtual void Initialize() override;
 
 	void GenerateMaterialAttributeOptions(); // Based on the current Material, list of material attributes available
 
+	// During evaluation, pick the MIC from the input static switch(es) combination
+	UMaterialInstanceConstant* GetEvaluationMaterialInstanceConstant(FTG_EvaluationContext* InContext);
+
 	TiledBlobPtr CreateRenderMaterialJob(FTG_EvaluationContext* InContext, const FString& InName, const FString& InMaterialPath, const BufferDescriptor& InDescriptor, EDrawMaterialAttributeTarget InDrawMaterialAttributeTarget);
 	TiledBlobPtr CreateRenderMaterialJob(FTG_EvaluationContext* InContext, const RenderMaterial_BPPtr& InRenderMaterial, const BufferDescriptor& InDescriptor, EDrawMaterialAttributeTarget InDrawMaterialAttributeTarget);
 	void LinkMaterialParameters(FTG_EvaluationContext* InContext, JobUPtr& InMaterialJob, const UMaterialInterface* InMaterial, BufferDescriptor InDescriptor);
diff --git a/Engine/Plugins/Experimental/TextureGraph/Source/TextureGraphEditor/Private/TG_RenderModeManager.cpp b/Engine/Plugins/Experimental/TextureGraph/Source/TextureGraphEditor/Private/TG_RenderModeManager.cpp
--- a/Engine/Plugins/Experimental/TextureGraph/Source/TextureGraphEditor/Private/TG_RenderModeManager.cpp
+++ b/Engine/Plugins/Experimental/TextureGraph/Source/TextureGraphEditor/Private/TG_RenderModeManager.cpp
@@ -44,7 +44,6 @@ void TG_RenderModeManager::UpdateRenderMode(UTextureGraph* TextureGraph)
 	for (size_t ti = 0; ti < NumTargets; ti++)
 	{
 		RenderMaterial_BPPtr renderMaterial = GetTargetRenderModeMaterial(ti,_currentRenderMode);
-		renderMaterial->Instance()->ClearParameterValues();
 		
 		const TargetTextureSetPtr& target = settings->Target(ti);
 
diff --git a/Engine/Plugins/Experimental/TextureGraph/Source/TextureGraphEngine/FxMat/MaterialManager.cpp b/Engine/Plugins/Experimental/TextureGraph/Source/TextureGraphEngine/FxMat/MaterialManager.cpp
--- a/Engine/Plugins/Experimental/TextureGraph/Source/TextureGraphEngine/FxMat/MaterialManager.cpp
+++ b/Engine/Plugins/Experimental/TextureGraph/Source/TextureGraphEngine/FxMat/MaterialManager.cpp
@@ -47,23 +47,23 @@ RenderMaterial_BPPtr UMaterialManager::CreateMaterial_BP(FString Name, const FSt
 
 RenderMaterial_BPPtr UMaterialManager::CreateMaterial_BP(FString InName, UMaterial* InMaterial)
 {
-	return std::make_shared<RenderMaterial_BP>(InName, InMaterial, nullptr);
+	return std::make_shared<RenderMaterial_BP>(InName, InMaterial);
 }
 
 RenderMaterial_BP_NoTilePtr UMaterialManager::CreateMaterial_BP_NoTile(FString Name, const FString& Path)
 {
 	UMaterial* Mat = LoadMaterial(Path);
-	return std::make_shared<RenderMaterial_BP_NoTile>(Name, Mat, nullptr);
+	return std::make_shared<RenderMaterial_BP_NoTile>(Name, Mat);
 }
 
 RenderMaterial_ThumbPtr UMaterialManager::CreateMaterial_Thumbnail(FString Name, const FString& Path)
 {
 	UMaterial* Mat = LoadMaterial(Path);
-	return std::make_shared<RenderMaterial_Thumbnail>(Name, Mat, nullptr);
+	return std::make_shared<RenderMaterial_Thumbnail>(Name, Mat);
 }
 
 RenderMaterial_BP_TileArgsPtr UMaterialManager::CreateMaterialOfType_BP_TileArgs(FString Name, const FString& Path)
 {
 	UMaterial* Mat = LoadMaterial(Path);
-	return std::make_shared<RenderMaterial_BP_TileArgs>(Name, Mat, nullptr);
+	return std::make_shared<RenderMaterial_BP_TileArgs>(Name, Mat);
 }
diff --git a/Engine/Plugins/Experimental/TextureGraph/Source/TextureGraphEngine/FxMat/RenderMaterial_BP.cpp b/Engine/Plugins/Experimental/TextureGraph/Source/TextureGraphEngine/FxMat/RenderMaterial_BP.cpp
--- a/Engine/Plugins/Experimental/TextureGraph/Source/TextureGraphEngine/FxMat/RenderMaterial_BP.cpp
+++ b/Engine/Plugins/Experimental/TextureGraph/Source/TextureGraphEngine/FxMat/RenderMaterial_BP.cpp
@@ -8,6 +8,8 @@
 #include "Model/Mix/MixUpdateCycle.h"
 #include <Materials/Material.h>
 #include <Materials/MaterialInstanceDynamic.h>
+#include <Materials/MaterialInstanceConstant.h>
+#include <Factories/MaterialInstanceConstantFactoryNew.h>
 #include <Engine/World.h>
 #include <TextureResource.h>
 #include <UObject/UObjectGlobals.h>
@@ -15,18 +17,119 @@
 #include "Async/Async.h"
 #include "FxMaterial_DrawMaterial.h"
 #include "Job/Job.h"
+#include "AssetToolsModule.h"
+#include "IAssetTools.h"
+
+
+TSharedPtr<FMaterialInstanceStaticSwitchPermutationMap> FMaterialInstanceStaticSwitchPermutationMap::Create(UMaterialInterface* InMaterial)
+{
+	TSharedPtr<FMaterialInstanceStaticSwitchPermutationMap> MISSPMap;
+	if (InMaterial)
+	{
+		MISSPMap = MakeShared<FMaterialInstanceStaticSwitchPermutationMap>();
+
+		InMaterial->GetStaticParameterValues(MISSPMap->DefaultStaticParameterSet);
+
+		int32 NumSwitches = MISSPMap->DefaultStaticParameterSet.StaticSwitchParameters.Num();
+
+		MISSPMap->DefaultKey = MISSPMap->KeyFromStaticSwitchParameters(MISSPMap->DefaultStaticParameterSet.StaticSwitchParameters);
+
+		UMaterialInstanceConstant* MIC = NewObject<UMaterialInstanceConstant>(InMaterial);
+		MIC->SetParentEditorOnly(InMaterial);
+		MIC->ClearParameterValuesEditorOnly();
+		MIC->SetFlags(RF_Standalone);
+		MIC->MarkPackageDirty();
+		MIC->PreEditChange(NULL);
+		MIC->PostEditChange();
+
+		MISSPMap->PermutationsMap.Add(MISSPMap->DefaultKey, MIC);
+	}
+	return MISSPMap;
+}
+
+UMaterialInterface* FMaterialInstanceStaticSwitchPermutationMap::GetRootMaterial()
+{
+	TObjectPtr<UMaterialInstanceConstant>* MICP = PermutationsMap.Find(DefaultKey);
+	if (MICP)
+	{
+		return Cast<UMaterialInterface>((*MICP)->GetMaterial());
+	}
+	return nullptr;
+}
+
+int32 FMaterialInstanceStaticSwitchPermutationMap::KeyFromStaticSwitchParameters(const TArray<FStaticSwitchParameter>& Parameters)
+{
+	int32 Key = 0;
+	for (int32 Index = 0; Index < Parameters.Num(); ++Index)
+	{
+		int32 ParamBitfield = ((int32)(Parameters[Index].Value) << Index);
+		Key |= ParamBitfield;
+	}
+
+	return Key;
+}
+
+UMaterialInstanceConstant* FMaterialInstanceStaticSwitchPermutationMap::GetMaterialInstance(const TArray<FStaticSwitchParameter>& Parameters)
+{
+	check(DefaultStaticParameterSet.StaticSwitchParameters.Num() == Parameters.Num());
+
+	int32 Key = KeyFromStaticSwitchParameters(Parameters);
+
+	TObjectPtr<UMaterialInstanceConstant>* MICP = PermutationsMap.Find(Key);
+	if (!MICP)
+	{
+		UMaterialInstanceConstant* MIC = NewObject<UMaterialInstanceConstant>(GetRootMaterial());
+		MIC->SetParentEditorOnly(GetRootMaterial());
+		MIC->ClearParameterValuesEditorOnly();
+		MIC->SetFlags(RF_Standalone);
+		MIC->MarkPackageDirty();
+
+		bool NeedStaticUpdate = false;
+
+		FStaticParameterSet InstanceStaticParameterSet = DefaultStaticParameterSet;
+		InstanceStaticParameterSet.StaticSwitchParameters = Parameters;
+
+		MIC->UpdateStaticPermutation(InstanceStaticParameterSet);
+		MIC->PreEditChange(NULL);
+		MIC->PostEditChange();
+
+		PermutationsMap.Add(Key, MIC);
+
+		return MIC;
+	}
+	else
+	{
+		return MICP->Get();
+	}
+}
+
 
 bool RenderMaterial_BP::ValidateMaterialCompatible(UMaterialInterface* InMaterial)
 {
 	return FxMaterial_QuadDrawMaterial::ValidateMaterial(InMaterial);
 }
 
-RenderMaterial_BP::RenderMaterial_BP(FString Name, UMaterialInterface* InMaterial, UMaterialInstanceDynamic* InMaterialInstance)
+RenderMaterial_BP::RenderMaterial_BP(FString Name, UMaterialInterface* InMaterial)
 	: RenderMaterial(!Name.IsEmpty() ? Name : InMaterial->GetName())
 	, Material(InMaterial)
-	, MaterialInstance(TStrongObjectPtr<UMaterialInstanceDynamic>(InMaterialInstance ? InMaterialInstance : UMaterialInstanceDynamic::Create(InMaterial, nullptr)))
 	, FXMaterialObj(std::make_shared<FxMaterial_QuadDrawMaterial>())
 {
+	// Cast to or build a MIC to be used by this RenderMaterial and its duplicates
+	if (Material && Material->IsA< UMaterialInstanceConstant>())
+	{
+		MaterialInstance = TStrongObjectPtr<UMaterialInstanceConstant>(Cast<UMaterialInstanceConstant>(Material));
+	}
+	else if (Material && !MaterialInstance)
+	{
+		MaterialInstance = TStrongObjectPtr<UMaterialInstanceConstant>(NewObject<UMaterialInstanceConstant>(Material));
+		MaterialInstance->SetParentEditorOnly(Material);
+		MaterialInstance->ClearParameterValuesEditorOnly();
+		MaterialInstance->SetFlags(RF_Standalone);
+		MaterialInstance->MarkPackageDirty();
+		MaterialInstance->PreEditChange(NULL);
+		MaterialInstance->PostEditChange();
+	}
+
 	check(Material && MaterialInstance);
 
 
@@ -53,7 +156,9 @@ CHashPtr RenderMaterial_BP::Hash() const
 void RenderMaterial_BP::SetTexture(FName InName, const UTexture* Texture) const
 {
 	check(Material && MaterialInstance);
-	MaterialInstance->SetTextureParameterValue(InName, const_cast<UTexture*>(Texture));
+#if WITH_EDITOR
+	MaterialInstance->SetTextureParameterValueEditorOnly(InName, const_cast<UTexture*>(Texture));
+#endif
 }
 
 void RenderMaterial_BP::SetArrayTexture(FName InName, const std::vector<const UTexture*>& Textures) const
@@ -63,40 +168,37 @@ void RenderMaterial_BP::SetArrayTexture(FName InName, const std::vector<const UT
 
 void RenderMaterial_BP::SetInt(FName InName, int32 Value) const
 {
-	check(Material && MaterialInstance);
-	MaterialInstance->SetScalarParameterValue(InName, Value);
+	// Same as set scalar for material instance constant
+	SetFloat(InName, Value);
 }
 
 void RenderMaterial_BP::SetFloat(FName InName, float Value) const
 {
 	check(Material && MaterialInstance);
-	MaterialInstance->SetScalarParameterValue(InName, Value);
+#if WITH_EDITOR
+	MaterialInstance->SetScalarParameterValueEditorOnly(InName, Value);
+#endif
 }
 
 void RenderMaterial_BP::SetColor(FName InName, const FLinearColor& Value) const
 {
 	check(Material && MaterialInstance);
-	MaterialInstance->SetVectorParameterValue(InName, Value);
+#if WITH_EDITOR
+	MaterialInstance->SetVectorParameterValueEditorOnly(InName, Value);
+#endif
 }
 
 void RenderMaterial_BP::SetIntVector4(FName InName, const FIntVector4& Value) const
 {
 	check(Material && MaterialInstance);
-	MaterialInstance->SetVectorParameterValue(InName, FVector4(Value));
+#if WITH_EDITOR
+	MaterialInstance->SetVectorParameterValueEditorOnly(InName, FLinearColor(Value.X, Value.Y, Value.Z, Value.W));
+#endif
 }
 
 void RenderMaterial_BP::SetMatrix(FName InName, const FMatrix& Value) const
 {
-	check(Material && MaterialInstance);
-
-	// UMaterialInstanceDynamic does not take matrix parameters
-
-	MaterialInstance->SetVectorParameterValue(FName(*(InName.ToString() + "_0")), Value.GetColumn(0));
-	MaterialInstance->SetVectorParameterValue(FName(*(InName.ToString() + "_1")), Value.GetColumn(1));
-	MaterialInstance->SetVectorParameterValue(FName(*(InName.ToString() + "_2")), Value.GetColumn(2));
-	MaterialInstance->SetVectorParameterValue(FName(*(InName.ToString() + "_3")), Value.GetColumn(3));
-	/*_instance->set
-	_instance->SetMatrixParameterValue(Name, Value);*/
+	// UMaterialInstanceConstant does not support assigning matrix parameters
 }
 
 std::shared_ptr<BlobTransform> RenderMaterial_BP::DuplicateInstance(FString InName)
@@ -104,7 +206,7 @@ std::shared_ptr<BlobTransform> RenderMaterial_BP::DuplicateInstance(FString InNa
 	if (InName.IsEmpty())
 		InName = Name;
 
-	RenderMaterial_BPPtr MaterialBP = std::make_shared<RenderMaterial_BP>(InName, Material, nullptr); //We would want a new instance every time
+	RenderMaterial_BPPtr MaterialBP = std::make_shared<RenderMaterial_BP>(InName, MaterialInstance.Get()); //Reuse the same material instance for every invocations
 
 	MaterialBP->MaterialInstanceValidated = MaterialInstanceValidated;
 	std::shared_ptr<BlobTransform> Result = std::static_pointer_cast<RenderMaterial>(MaterialBP);
@@ -121,6 +223,8 @@ AsyncPrepareResult RenderMaterial_BP::PrepareResources(const TransformArgs& Args
 		RenderMaterial_BP* Self = this;
 		return cti::make_continuable<int32>([Self, Mix, ErrorOwner](auto&& Promise) mutable
 			{
+
+
 				Util::OnRenderingThread([FWD_PROMISE(Promise), Self, Mix, ErrorOwner](FRHICommandListImmediate& RHI) mutable
 					{
 						bool ValidationResult = ValidateMaterialCompatible(Self->MaterialInstance.Get());
@@ -205,20 +309,17 @@ void RenderMaterial_BP::BlitTo(FRHICommandListImmediate& RHI, UTextureRenderTarg
 
 void RenderMaterial_BP::Bind(int32 Value, const ResourceBindInfo& BindInfo)
 {
-	check(Material && MaterialInstance);
-	MaterialInstance->SetScalarParameterValue(*BindInfo.Target, Value);
+	SetInt(*BindInfo.Target, Value);
 }
 
 void RenderMaterial_BP::Bind(float Value, const ResourceBindInfo& BindInfo)
 {
-	check(Material && MaterialInstance);
-	MaterialInstance->SetScalarParameterValue(*BindInfo.Target, Value);	
+	SetFloat(*BindInfo.Target, Value);	
 }
 
 void RenderMaterial_BP::Bind(const FLinearColor& Value, const ResourceBindInfo& BindInfo)
 {
-	check(Material && MaterialInstance);
-	MaterialInstance->SetVectorParameterValue(*BindInfo.Target, Value);
+	SetColor(*BindInfo.Target, Value);
 
 	FName BindName(BindInfo.Target);
 	if (BindName == FTextureGraphMaterialShaderPS::PSCONTROL_ARG)
@@ -227,13 +328,11 @@ void RenderMaterial_BP::Bind(const FLinearColor& Value, const ResourceBindInfo&
 
 void RenderMaterial_BP::Bind(const FIntVector4& Value, const ResourceBindInfo& BindInfo)
 {
-	check(Material && MaterialInstance);
-	MaterialInstance->SetVectorParameterValue(*BindInfo.Target, FVector4(Value));
+	SetIntVector4(*BindInfo.Target, Value);
 }
 
 void RenderMaterial_BP::Bind(const FMatrix& Value, const ResourceBindInfo& BindInfo)
 {
-	check(Material && MaterialInstance);
 	SetMatrix(*BindInfo.Target, Value);
 }
 
diff --git a/Engine/Plugins/Experimental/TextureGraph/Source/TextureGraphEngine/FxMat/RenderMaterial_BP.h b/Engine/Plugins/Experimental/TextureGraph/Source/TextureGraphEngine/FxMat/RenderMaterial_BP.h
--- a/Engine/Plugins/Experimental/TextureGraph/Source/TextureGraphEngine/FxMat/RenderMaterial_BP.h
+++ b/Engine/Plugins/Experimental/TextureGraph/Source/TextureGraphEngine/FxMat/RenderMaterial_BP.h
@@ -4,18 +4,34 @@
 #include "CoreMinimal.h"
 #include "RenderMaterial.h"
 #include "Materials/MaterialInstanceDynamic.h"
+#include "Materials/MaterialInstanceConstant.h"
 #include "UObject/StrongObjectPtr.h"
 
 class FxMaterial;
 typedef std::shared_ptr< FxMaterial> FxMaterialPtr;
 
+// Helper container of MIC for the different static switch permutations of a Material
+// Used by the caller of a RenderMaterial_BP to pick a specific permutation different from the default set.
+// See TG_Expression_MaterialBase for a usage example
+struct TEXTUREGRAPHENGINE_API FMaterialInstanceStaticSwitchPermutationMap
+{
+	int32												DefaultKey = 0;
+	FStaticParameterSet									DefaultStaticParameterSet;
+	TMap<int32, TObjectPtr<UMaterialInstanceConstant>>	PermutationsMap;
+
+	static TSharedPtr<FMaterialInstanceStaticSwitchPermutationMap> Create(UMaterialInterface* InMaterial);
+	UMaterialInterface*				GetRootMaterial();
+	int32							KeyFromStaticSwitchParameters(const TArray<FStaticSwitchParameter>& Parameters);
+	UMaterialInstanceConstant*		GetMaterialInstance(const TArray<FStaticSwitchParameter>& Parameters);
+};
+
 class TEXTUREGRAPHENGINE_API RenderMaterial_BP : public RenderMaterial
 {
 protected:
 
 	UMaterialInterface*				Material = nullptr;			/// The base material that is used for this job 
-	TStrongObjectPtr<UMaterialInstanceDynamic>
-									MaterialInstance;			/// An instance of the material
+	TStrongObjectPtr<UMaterialInstanceConstant> MaterialInstance;			/// An instance of the material
+
 	CHashPtr						HashValue;					/// The hash for this material
 	bool							RequestMaterialValidation = true;		/// 
 	bool							MaterialInstanceValidated = false;		/// 
@@ -28,7 +44,7 @@ protected:
 									FVector2D CoordinatePosition, FVector2D CoordinateSize=FVector2D::UnitVector, float Rotation=0.f,
 									FVector2D PivotPoint=FVector2D(0.5f,0.5f)) const;
 public:
-									RenderMaterial_BP(FString Name, UMaterialInterface* InMaterial, UMaterialInstanceDynamic* InMaterialInstance = nullptr);
+									RenderMaterial_BP(FString Name, UMaterialInterface* InMaterial);
 
 	virtual							~RenderMaterial_BP() override;
 
@@ -64,10 +80,9 @@ public:
 	//////////////////////////////////////////////////////////////////////////
 	/// Inline functions
 	////////////////////////////////////////////////////////////////////////// 
-	FORCEINLINE UMaterialInterface*			GetMaterial() { return Material; }
-
-	FORCEINLINE UMaterialInstanceDynamic* Instance() { return MaterialInstance.Get(); }
-	FORCEINLINE const UMaterialInstanceDynamic* Instance() const { return MaterialInstance.Get(); }
+	FORCEINLINE UMaterialInterface*					GetMaterial() { return Material; }
+	FORCEINLINE UMaterialInstanceConstant*			Instance() { return MaterialInstance.Get(); }
+	FORCEINLINE const UMaterialInstanceConstant*	Instance() const { return MaterialInstance.Get(); }
 };
 
 typedef std::shared_ptr<RenderMaterial_BP> RenderMaterial_BPPtr;
\ No newline at end of file
diff --git a/Engine/Plugins/Experimental/TextureGraph/Source/TextureGraphEngine/FxMat/RenderMaterial_BP_NoTile.cpp b/Engine/Plugins/Experimental/TextureGraph/Source/TextureGraphEngine/FxMat/RenderMaterial_BP_NoTile.cpp
--- a/Engine/Plugins/Experimental/TextureGraph/Source/TextureGraphEngine/FxMat/RenderMaterial_BP_NoTile.cpp
+++ b/Engine/Plugins/Experimental/TextureGraph/Source/TextureGraphEngine/FxMat/RenderMaterial_BP_NoTile.cpp
@@ -21,7 +21,7 @@ std::shared_ptr<BlobTransform> RenderMaterial_BP_TileArgs::DuplicateInstance(FSt
 	if (NewName.IsEmpty())
 		NewName = Name;
 
-	RenderMaterial_BP_TileArgsPtr MaterialBP = std::make_shared<RenderMaterial_BP_TileArgs>(NewName, GetMaterial(), nullptr); //We would want a new instance every time
+	RenderMaterial_BP_TileArgsPtr MaterialBP = std::make_shared<RenderMaterial_BP_TileArgs>(NewName, Instance()); //We would want a new instance every time
 	std::shared_ptr<BlobTransform> Result = std::static_pointer_cast<RenderMaterial_BP_TileArgs>(MaterialBP);
 
 	return Result;
diff --git a/Engine/Plugins/Experimental/TextureGraph/Source/TextureGraphEngine/FxMat/RenderMaterial_BP_NoTile.h b/Engine/Plugins/Experimental/TextureGraph/Source/TextureGraphEngine/FxMat/RenderMaterial_BP_NoTile.h
--- a/Engine/Plugins/Experimental/TextureGraph/Source/TextureGraphEngine/FxMat/RenderMaterial_BP_NoTile.h
+++ b/Engine/Plugins/Experimental/TextureGraph/Source/TextureGraphEngine/FxMat/RenderMaterial_BP_NoTile.h
@@ -11,7 +11,7 @@ class TEXTUREGRAPHENGINE_API RenderMaterial_BP_NoTile : public RenderMaterial_BP
 {
 
 public:
-	RenderMaterial_BP_NoTile(FString InName, UMaterialInterface* InMaterial, UMaterialInstanceDynamic* InMaterialInstance = nullptr) : RenderMaterial_BP(InName, InMaterial, InMaterialInstance) {}
+	RenderMaterial_BP_NoTile(FString InName, UMaterialInterface* InMaterial) : RenderMaterial_BP(InName, InMaterial) {}
 	virtual bool					CanHandleTiles() const override { return false; };
 };
 
@@ -23,7 +23,7 @@ private:
 	FName							MemberName = "TileInfo";
 
 public:
-	RenderMaterial_BP_TileArgs(FString InName, UMaterialInterface* InMaterial, UMaterialInstanceDynamic* InMaterialInstance = nullptr) : RenderMaterial_BP(InName, InMaterial, InMaterialInstance) {}
+	RenderMaterial_BP_TileArgs(FString InName, UMaterialInterface* InMaterial) : RenderMaterial_BP(InName, InMaterial) {}
 	void							AddTileArgs(TransformArgs& Args);
 	virtual AsyncPrepareResult		PrepareResources(const TransformArgs& Args) override;
 	virtual bool					CanHandleTiles() const override { return true; };
diff --git a/Engine/Plugins/Experimental/TextureGraph/Source/TextureGraphEngine/FxMat/RenderMaterial_Thumbnail.cpp b/Engine/Plugins/Experimental/TextureGraph/Source/TextureGraphEngine/FxMat/RenderMaterial_Thumbnail.cpp
--- a/Engine/Plugins/Experimental/TextureGraph/Source/TextureGraphEngine/FxMat/RenderMaterial_Thumbnail.cpp
+++ b/Engine/Plugins/Experimental/TextureGraph/Source/TextureGraphEngine/FxMat/RenderMaterial_Thumbnail.cpp
@@ -14,7 +14,7 @@
 
 ////////////////////////////////////////////////////////////////////////// RenderMaterial_Thumbnail //////////////////////////////////////////////////////////////////////////
 
-RenderMaterial_Thumbnail::RenderMaterial_Thumbnail(FString InName, UMaterialInterface* InMaterial, UMaterialInstanceDynamic* instance /*= nullptr*/) : RenderMaterial_BP(InName, InMaterial, instance)
+RenderMaterial_Thumbnail::RenderMaterial_Thumbnail(FString InName, UMaterialInterface* InMaterial) : RenderMaterial_BP(InName, InMaterial)
 {
 }
 
@@ -39,7 +39,7 @@ std::shared_ptr<BlobTransform> RenderMaterial_Thumbnail::DuplicateInstance(FStri
 	if (InName.IsEmpty())
 		InName = Name;
 
-	std::shared_ptr<RenderMaterial_Thumbnail> MaterialThm = std::make_shared<RenderMaterial_Thumbnail>(InName, Material, nullptr); //We would want a new instance every time
+	std::shared_ptr<RenderMaterial_Thumbnail> MaterialThm = std::make_shared<RenderMaterial_Thumbnail>(InName, Instance()); //We would want a new instance every time
 
 	MaterialThm->FeatureLevel = FeatureLevel;
 	MaterialThm->MaterialInstanceValidated = MaterialInstanceValidated;
diff --git a/Engine/Plugins/Experimental/TextureGraph/Source/TextureGraphEngine/FxMat/RenderMaterial_Thumbnail.h b/Engine/Plugins/Experimental/TextureGraph/Source/TextureGraphEngine/FxMat/RenderMaterial_Thumbnail.h
--- a/Engine/Plugins/Experimental/TextureGraph/Source/TextureGraphEngine/FxMat/RenderMaterial_Thumbnail.h
+++ b/Engine/Plugins/Experimental/TextureGraph/Source/TextureGraphEngine/FxMat/RenderMaterial_Thumbnail.h
@@ -17,7 +17,7 @@ private:
 	ERHIFeatureLevel::Type FeatureLevel;
 
 public:
-	RenderMaterial_Thumbnail(FString InName, UMaterialInterface* InMaterial, UMaterialInstanceDynamic* instance = nullptr);
+	RenderMaterial_Thumbnail(FString InName, UMaterialInterface* InMaterial);
 
 	virtual void BlitTo(FRHICommandListImmediate& RHI, UTextureRenderTarget2D* DstRT, const RenderMesh* MeshObj, int32 TargetId) const override;
 	virtual bool CanHandleTiles() const override { return true; };
