Array grid rendering support

diff --git a/Engine/Plugins/Experimental/TextureGraph/Source/TextureGraph/Private/Expressions/Arrays/TG_Expression_ArrayGrid.cpp b/Engine/Plugins/Experimental/TextureGraph/Source/TextureGraph/Private/Expressions/Arrays/TG_Expression_ArrayGrid.cpp
new file mode 100644
--- /dev/null
+++ b/Engine/Plugins/Experimental/TextureGraph/Source/TextureGraph/Private/Expressions/Arrays/TG_Expression_ArrayGrid.cpp
@@ -0,0 +1,77 @@
+// Copyright Epic Games, Inc. All Rights Reserved.
+
+#include "Expressions/Arrays/TG_Expression_ArrayGrid.h"
+
+#include "TG_Graph.h"
+#include "2D/TextureHelper.h"
+#include "Model/StaticImageResource.h"
+#include "Transform/Expressions/T_ArrayGrid.h"
+
+void UTG_Expression_ArrayGrid::Evaluate(FTG_EvaluationContext* InContext)
+{
+	Super::Evaluate(InContext);
+
+	TArray<TiledBlobPtr> Inputs;
+	Inputs.Reserve(Input.Num());
+
+	/// Filter out the input textures, just in case someone tries to slip in a non-texture
+	/// through the variant system
+	for (int Index = 0; Index < Input.Num(); Index++)
+	{
+		FTG_Variant& InputItem = Input.Get(Index);
+		if (InputItem.IsTexture())
+		{
+			FTG_Texture InputTexture = InputItem.GetTexture();
+			Inputs.Add(InputTexture);
+		}
+	}
+
+	/// If there are no inputs then we simply render a black texture
+	if (Inputs.IsEmpty())
+	{
+		Output = FTG_Texture::GetBlack();
+		return;
+	}
+
+	int32 OutRows = 0;
+	int32 OutCols = 0;
+	if (Rows <= 0 && Columns <= 0)
+	{
+		int32 Count = FMath::CeilToInt(FMath::Sqrt((float)Inputs.Num()));
+		OutRows = OutCols = Count;
+	}
+	else if (Rows <= 0 && Columns > 0)
+	{
+		OutCols = std::min<int32>(Columns, Inputs.Num());
+		OutRows = std::ceil((float)Inputs.Num() / (float)OutCols);
+	}
+	else if (Rows > 0 && Columns <= 0)
+	{
+		OutRows = std::min<int32>(Rows, Inputs.Num());
+		OutCols = std::ceil((float)Inputs.Num() / (float)OutRows);
+	}
+	else
+	{
+		OutRows = Rows;
+		OutCols = Columns;
+	}
+
+	/// Make sure there are no zeros
+	OutRows = std::max(OutRows, 1);
+	OutCols = std::max(OutCols, 1);
+
+	if (OutRows * OutCols > Inputs.Num())
+	{
+		Output = FTG_Texture::GetBlack();
+		FTextureGraphErrorReporter* ErrorReporter = TextureGraphEngine::GetErrorReporter(InContext->Cycle->GetMix());
+		if (ErrorReporter)
+		{
+			ErrorReporter->ReportWarning((int32)ETextureGraphErrorType::NODE_WARNING, FString::Printf(TEXT("Number of rows and columns for the tiling do not match the number of inputs. Grid: %d x %d [Max Inputs: %d]"), 
+				OutRows, OutCols, Inputs.Num()), GetParentNode());
+		}
+		return;
+	}
+
+	BufferDescriptor DesiredDesc = Output.Descriptor;
+	Output = T_ArrayGrid::Create(InContext->Cycle, DesiredDesc, Inputs, OutRows, OutCols);
+}
diff --git a/Engine/Plugins/Experimental/TextureGraph/Source/TextureGraph/Private/Expressions/Input/TG_Expression_Texture.cpp b/Engine/Plugins/Experimental/TextureGraph/Source/TextureGraph/Private/Expressions/Input/TG_Expression_Texture.cpp
--- a/Engine/Plugins/Experimental/TextureGraph/Source/TextureGraph/Private/Expressions/Input/TG_Expression_Texture.cpp
+++ b/Engine/Plugins/Experimental/TextureGraph/Source/TextureGraph/Private/Expressions/Input/TG_Expression_Texture.cpp
@@ -25,7 +25,7 @@ void UTG_Expression_Texture::Evaluate(FTG_EvaluationContext* InContext)
 		auto DesiredDesc = Output.GetBufferDescriptor();
 		DesiredDesc.bIsSRGB = Source->SRGB;
 
-		Output = StaticImageResource->GetBlob(InContext->Cycle, DesiredDesc, 0);
+		Output = StaticImageResource->GetBlob(InContext->Cycle, &DesiredDesc, 0);
 	}
 	else
 	{
diff --git a/Engine/Plugins/Experimental/TextureGraph/Source/TextureGraph/Private/Expressions/Input/TG_Expression_TexturePath.cpp b/Engine/Plugins/Experimental/TextureGraph/Source/TextureGraph/Private/Expressions/Input/TG_Expression_TexturePath.cpp
--- a/Engine/Plugins/Experimental/TextureGraph/Source/TextureGraph/Private/Expressions/Input/TG_Expression_TexturePath.cpp
+++ b/Engine/Plugins/Experimental/TextureGraph/Source/TextureGraph/Private/Expressions/Input/TG_Expression_TexturePath.cpp
@@ -5,6 +5,7 @@
 #include "TG_Graph.h"
 #include "2D/TextureHelper.h"
 #include "Model/StaticImageResource.h"
+#include "HAL/FileManagerGeneric.h"
 
 // Special case for TexturePath Constant signature, we want to keep the Path Input connectable in that case
 // so do this in the override version of BuildInputConstantSignature()
@@ -56,6 +57,15 @@ bool UTG_Expression_TexturePath::ValidateInputPath(FString& ValidatedPath) const
 	return false;
 }
 
+TiledBlobPtr UTG_Expression_TexturePath::LoadStaticImage(FTG_EvaluationContext* InContext, const FString& LocalPath, BufferDescriptor* DesiredDesc)
+{
+	UStaticImageResource* StaticImageResource = UStaticImageResource::CreateNew<UStaticImageResource>();
+	StaticImageResource->SetAssetUUID(LocalPath);
+	StaticImageResource->SetIsFileSystem(true);
+
+	return StaticImageResource->GetBlob(InContext->Cycle, DesiredDesc, 0);
+}
+
 void UTG_Expression_TexturePath::Evaluate(FTG_EvaluationContext* InContext)
 {
 	Super::Evaluate(InContext);
@@ -65,21 +75,34 @@ void UTG_Expression_TexturePath::Evaluate(FTG_EvaluationContext* InContext)
 
 	if (!Path.IsEmpty() && bValidation)
 	{
-		UStaticImageResource* StaticImageResource = UStaticImageResource::CreateNew<UStaticImageResource>();
-		StaticImageResource->SetAssetUUID(LocalPath);
-		StaticImageResource->SetIsFileSystem(true);
-
-		//Until we have srgb value exposed in the UI we need to set the Srgb of the Output Descriptor here from the source
-		//This gets updated for the late bond case but since we do not have the UI to specify the override in other nodes 
-		// the override value will always be set to false while combining the buffers
-		auto DesiredDesc = Output.GetBufferDescriptor();
-		DesiredDesc.bIsSRGB = true;
-
-		Output = StaticImageResource->GetBlob(InContext->Cycle, DesiredDesc, 0);
+		if (FPaths::FileExists(Path))
+		{
+			Output.SetNum(1);
+			FTG_Texture Item = LoadStaticImage(InContext, LocalPath, nullptr);
+			Output.Set(0, Item);
+		}
+		else if (FPaths::DirectoryExists(Path))
+		{
+			FFileManagerGeneric FileManager;
+			TArray<FString> Files;
+			FileManager.FindFiles(Files, *Path);
+
+			Output.SetNum(Files.Num());
+
+			/// Now load a texture against all the files
+			for (int32 Index = 0; Index < Files.Num(); Index++)
+			{
+				const FString& File = Files[Index];
+				FString Filename = FPaths::Combine(Path, File);
+				FTG_Texture Item = LoadStaticImage(InContext, Filename, nullptr);
+				Output.Set(Index, Item);
+			}
+		}
 	}
 	else
 	{
-		Output = FTG_Texture::GetBlack();
+		Output.SetNum(1);
+		Output.Set(0, FTG_Texture::GetBlack());
 	}
 
 	//For the connected pin we will report error here in eveluate because it does not have the updated value during validation.
diff --git a/Engine/Plugins/Experimental/TextureGraph/Source/TextureGraph/Public/Expressions/Arrays/TG_Expression_ArrayGrid.h b/Engine/Plugins/Experimental/TextureGraph/Source/TextureGraph/Public/Expressions/Arrays/TG_Expression_ArrayGrid.h
new file mode 100644
--- /dev/null
+++ b/Engine/Plugins/Experimental/TextureGraph/Source/TextureGraph/Public/Expressions/Arrays/TG_Expression_ArrayGrid.h
@@ -0,0 +1,38 @@
+// Copyright Epic Games, Inc. All Rights Reserved.
+
+#pragma once
+
+#include "CoreMinimal.h"
+#include "Expressions/TG_Expression.h"
+#include "TG_Texture.h"
+#include "TG_Expression_ArrayGrid.generated.h"
+
+UCLASS()
+class TEXTUREGRAPH_API UTG_Expression_ArrayGrid : public UTG_Expression
+{
+	 GENERATED_BODY()
+
+public:
+	TG_DECLARE_EXPRESSION(TG_Category::Arrays);
+	virtual void					Evaluate(FTG_EvaluationContext* InContext) override;
+	
+	/// The input texture array. Must be an array object of type textures
+	UPROPERTY(meta = (TGType = "TG_Input", PinDisplayName = "", HideInnerPropertiesInNode))
+	FTG_VariantArray				Input;
+
+	/// The number of rows in the output grid. 0 or Negative values mean automatically distribute
+	UPROPERTY(meta = (TGType = "TG_Input", PinDisplayName = ""))
+	int32							Rows = 0;
+
+	/// The number of colums in the output grid. 0 or Negative values mean automatically distribute
+	UPROPERTY(meta = (TGType = "TG_Input", PinDisplayName = ""))
+	int32							Columns = 0;
+
+	/// A single output texture
+	UPROPERTY(meta = (TGType = "TG_Output", PinDisplayName = ""))
+	FTG_Texture						Output;
+
+	virtual FTG_Name				GetDefaultName() const override { return TEXT("ArrayGrid");}
+	virtual FText					GetTooltipText() const override { return FText::FromString(TEXT("Arranges an input texture array into an MxN sized grid.")); } 
+ };
+
diff --git a/Engine/Plugins/Experimental/TextureGraph/Source/TextureGraphEngine/Data/Blob.h b/Engine/Plugins/Experimental/TextureGraph/Source/TextureGraphEngine/Data/Blob.h
--- a/Engine/Plugins/Experimental/TextureGraph/Source/TextureGraphEngine/Data/Blob.h
+++ b/Engine/Plugins/Experimental/TextureGraph/Source/TextureGraphEngine/Data/Blob.h
@@ -164,6 +164,7 @@ public:
 	//////////////////////////////////////////////////////////////////////////
 	FORCEINLINE uint32				GetWidth() const { return GetDescriptor().Width; }
 	FORCEINLINE uint32				GetHeight() const { return GetDescriptor().Height; }
+	FORCEINLINE bool				IsWellDefined() { return !IsLateBound() && GetWidth() > 0 && GetHeight() > 0; }
 	//FORCEINLINE const DeviceBufferRef& Buffer() const { return _buffer; }
 
 	FORCEINLINE	DeviceType			GetDeviceType() const { return Buffer.GetDeviceType(); }
diff --git a/Engine/Plugins/Experimental/TextureGraph/Source/TextureGraphEngine/Data/Blobber.cpp b/Engine/Plugins/Experimental/TextureGraph/Source/TextureGraphEngine/Data/Blobber.cpp
--- a/Engine/Plugins/Experimental/TextureGraph/Source/TextureGraphEngine/Data/Blobber.cpp
+++ b/Engine/Plugins/Experimental/TextureGraph/Source/TextureGraphEngine/Data/Blobber.cpp
@@ -671,6 +671,7 @@ void Blobber::UpdateHash(HashType OldHash, CHashPtr Hash)
 
 CHashPtr Blobber::AddGloballyUniqueHashInternal(CHashPtr Hash)
 {
+	check(Hash);
 	/// For internal use only [not thread-safe]
 	auto Iter = GlobalHashes.find(Hash->Value());
 
diff --git a/Engine/Plugins/Experimental/TextureGraph/Source/TextureGraphEngine/Data/TiledBlob.cpp b/Engine/Plugins/Experimental/TextureGraph/Source/TextureGraphEngine/Data/TiledBlob.cpp
--- a/Engine/Plugins/Experimental/TextureGraph/Source/TextureGraphEngine/Data/TiledBlob.cpp
+++ b/Engine/Plugins/Experimental/TextureGraph/Source/TextureGraphEngine/Data/TiledBlob.cpp
@@ -34,7 +34,7 @@ TiledBlob::TiledBlob(const BufferDescriptor& InDesc, const BlobPtrTiles& Tiles)
 
 	for (size_t TileX = 0; TileX < Tiles.Rows() && bIsFinalised; TileX++)
 	{
-		for (size_t TileY = 0; TileY < Tiles.Rows() && bIsFinalised; TileY++)
+		for (size_t TileY = 0; TileY < Tiles.Cols() && bIsFinalised; TileY++)
 		{
 			check(Tiles[TileX][TileY]);
 			BlobPtr Tile = Tiles[TileX][TileY];
@@ -61,7 +61,7 @@ std::shared_ptr<TiledBlob> TiledBlob::InitFromTiles(const BufferDescriptor& InDe
 
 	for (size_t TileX = 0; TileX < Tiles.Rows(); TileX++)
 	{
-		for (size_t TileY = 0; TileY < Tiles.Rows(); TileY++)
+		for (size_t TileY = 0; TileY < Tiles.Cols(); TileY++)
 		{
 			check(Tiles[TileX][TileY]);
 			BlobPtr Tile = Tiles[TileX][TileY];
@@ -99,7 +99,6 @@ TiledBlob::TiledBlob(const BufferDescriptor& InDesc, size_t NumTilesX, size_t Nu
 	, Desc(InDesc)
 {
 	check(NumTilesX > 0 && NumTilesY > 0);
-
 	
 	if (Desc.Format != BufferFormat::LateBound)
 	{
@@ -526,7 +525,7 @@ void TiledBlob::CalcHashNow() const
 		{
 			BlobPtr Tile = Tiles[TileX][TileY];
 
-			if (!Tile || !Tile->IsValid())
+			if (!Tile)
 				return;
 
 			CHashPtr TileHash = Tile->Hash();
@@ -1036,7 +1035,7 @@ TiledBlob_Promise::TiledBlob_Promise(DeviceBufferRef Buffer) : TiledBlob(Buffer)
 {
 }
 
-TiledBlob_Promise::TiledBlob_Promise(const BufferDescriptor& Desc, size_t NumTilesX, size_t NumTilesY, CHashPtr hash) : TiledBlob(Desc, NumTilesX, NumTilesY, hash)
+TiledBlob_Promise::TiledBlob_Promise(const BufferDescriptor& Desc, size_t NumTilesX, size_t NumTilesY, CHashPtr Hash) : TiledBlob(Desc, NumTilesX, NumTilesY, Hash)
 {
 }
 
@@ -1262,12 +1261,16 @@ AsyncBufferResultPtr TiledBlob_Promise::Flush(const ResourceBindInfo& BindInfo)
 	return TiledBlob::Flush(BindInfo);
 }
 
-TiledBlob_Promise& TiledBlob_Promise::operator = (const TiledBlob_Promise& RHS)
+TiledBlob_Promise& TiledBlob_Promise::operator = (const TiledBlob& RHS)
 {
 	TiledBlob::operator = (RHS);
+	return *this;
+}
 
+TiledBlob_Promise& TiledBlob_Promise::operator = (const TiledBlob_Promise& RHS)
+{
+	TiledBlob::operator = (RHS);
 	bMakeSingleBlob = RHS.bMakeSingleBlob;
-
 	return *this;
 }
 
diff --git a/Engine/Plugins/Experimental/TextureGraph/Source/TextureGraphEngine/Data/TiledBlob.h b/Engine/Plugins/Experimental/TextureGraph/Source/TextureGraphEngine/Data/TiledBlob.h
--- a/Engine/Plugins/Experimental/TextureGraph/Source/TextureGraphEngine/Data/TiledBlob.h
+++ b/Engine/Plugins/Experimental/TextureGraph/Source/TextureGraphEngine/Data/TiledBlob.h
@@ -123,6 +123,7 @@ public:
 	FORCEINLINE uint32				GetHeight() const { return Desc.Height; }
 	FORCEINLINE bool				TiledTarget() const { return bTiledTarget; }
 	FORCEINLINE bool&				TiledTarget() { return bTiledTarget; }
+	FORCEINLINE bool				IsWellDefined() { return !IsLateBound() && GetWidth() > 0 && GetHeight() > 0; }
 	FORCEINLINE bool				IsValidTileIndex(int32 TileX, int32 TileY) const
 	{
 		return TileX < static_cast<int32>(Tiles.Rows()) && TileY < static_cast<int32>(Tiles.Cols());
@@ -163,6 +164,7 @@ protected:
 
 	virtual void					OnFinaliseInternal(BlobReadyCallback Callback) const override;
 	TiledBlob_Promise&				operator = (const TiledBlob_Promise& RHS);
+	TiledBlob_Promise&				operator = (const TiledBlob& RHS);
 	virtual void					NotifyCallbacks();
 
 	virtual void					AddLinkedBlob(BlobPtr LinkedBlob) override;
diff --git a/Engine/Plugins/Experimental/TextureGraph/Source/TextureGraphEngine/FxMat/FxMaterial.h b/Engine/Plugins/Experimental/TextureGraph/Source/TextureGraphEngine/FxMat/FxMaterial.h
--- a/Engine/Plugins/Experimental/TextureGraph/Source/TextureGraphEngine/FxMat/FxMaterial.h
+++ b/Engine/Plugins/Experimental/TextureGraph/Source/TextureGraphEngine/FxMat/FxMaterial.h
@@ -355,14 +355,14 @@ public:
 			SetUAVParameterValue(FName(*OutputId), RenderTargetUAV);
 		}
 
-		TShaderMapRef<CmpSH_Type> computeShader(GetGlobalShaderMap(GMaxRHIFeatureLevel), PermutationDomain);
-		FIntVector groupSize = computeShader->ThreadGroupSize();
+		TShaderMapRef<CmpSH_Type> ComputeShader(GetGlobalShaderMap(GMaxRHIFeatureLevel), PermutationDomain);
+		FIntVector GroupSize = ComputeShader->ThreadGroupSize();
 
-		FComputeShaderUtils::Dispatch(RHI, computeShader, Params, 
+		FComputeShaderUtils::Dispatch(RHI, ComputeShader, Params, 
 			FIntVector(
-				FMath::DivideAndRoundUp(NumThreadsX, groupSize.X),
-				FMath::DivideAndRoundUp(NumThreadsY, groupSize.Y), 
-				FMath::DivideAndRoundUp(NumThreadsZ, groupSize.Z)
+				FMath::DivideAndRoundUp(NumThreadsX, GroupSize.X),
+				FMath::DivideAndRoundUp(NumThreadsY, GroupSize.Y), 
+				FMath::DivideAndRoundUp(NumThreadsZ, GroupSize.Z)
 			)
 		);
 
diff --git a/Engine/Plugins/Experimental/TextureGraph/Source/TextureGraphEngine/Job/Job.cpp b/Engine/Plugins/Experimental/TextureGraph/Source/TextureGraphEngine/Job/Job.cpp
--- a/Engine/Plugins/Experimental/TextureGraph/Source/TextureGraphEngine/Job/Job.cpp
+++ b/Engine/Plugins/Experimental/TextureGraph/Source/TextureGraphEngine/Job/Job.cpp
@@ -204,7 +204,7 @@ CHashPtr Job::Hash() const
 	return HashValue;
 }
 
-CHashPtr Job::TileHash(int32 TileX, int32 TileY) const
+CHashPtr Job::CalcTileHash(int32 TileX, int32 TileY) const
 {
 	CHashPtrVec ArgHashes;
 	ArgHashes.reserve(Args.Num() + 1); /// +1 for the Transform
@@ -265,7 +265,12 @@ CHashPtr Job::TileHash(int32 TileX, int32 TileY) const
 		ArgHashes.push_back(std::make_shared<CHash>(DataUtil::Hash_Int32(NumCols), true));
 	}
 
-	return TextureGraphEngine::GetBlobber()->AddGloballyUniqueHash(CHash::ConstructFromSources(ArgHashes));
+	return CHash::ConstructFromSources(ArgHashes);
+}
+
+CHashPtr Job::TileHash(int32 TileX, int32 TileY) const
+{
+	return TextureGraphEngine::GetBlobber()->AddGloballyUniqueHash(CalcTileHash(TileX, TileY));
 }
 
 BufferDescriptor Job::GetCombinedDesc(BufferDescriptor& ArgsDescCombined, size_t& Count) const
diff --git a/Engine/Plugins/Experimental/TextureGraph/Source/TextureGraphEngine/Job/Job.h b/Engine/Plugins/Experimental/TextureGraph/Source/TextureGraphEngine/Job/Job.h
--- a/Engine/Plugins/Experimental/TextureGraph/Source/TextureGraphEngine/Job/Job.h
+++ b/Engine/Plugins/Experimental/TextureGraph/Source/TextureGraphEngine/Job/Job.h
@@ -126,6 +126,7 @@ public:
 
 	virtual CHashPtr				Hash() const;
 	virtual CHashPtr				TileHash(int32 TileX, int32 TileY) const;
+	virtual CHashPtr				CalcTileHash(int32 TileX, int32 TileY) const;
 	virtual bool					CanHandleTiles() const;
 	virtual bool					CheckDefaultArgs() const;
 
diff --git a/Engine/Plugins/Experimental/TextureGraph/Source/TextureGraphEngine/Model/StaticImageResource.cpp b/Engine/Plugins/Experimental/TextureGraph/Source/TextureGraphEngine/Model/StaticImageResource.cpp
--- a/Engine/Plugins/Experimental/TextureGraph/Source/TextureGraphEngine/Model/StaticImageResource.cpp
+++ b/Engine/Plugins/Experimental/TextureGraph/Source/TextureGraphEngine/Model/StaticImageResource.cpp
@@ -63,7 +63,7 @@ AsyncTiledBlobRef UStaticImageResource::Load(MixUpdateCyclePtr Cycle)
 }
 
 
-TiledBlobPtr UStaticImageResource::GetBlob(MixUpdateCyclePtr Cycle, BufferDescriptor DesiredDesc, int32 TargetId)
+TiledBlobPtr UStaticImageResource::GetBlob(MixUpdateCyclePtr Cycle, BufferDescriptor* DesiredDesc, int32 TargetId)
 {
 	if (BlobObj)
 		return BlobObj;
@@ -76,8 +76,13 @@ TiledBlobPtr UStaticImageResource::GetBlob(MixUpdateCyclePtr Cycle, BufferDescri
 	JobObj->AddArg(ARG_STRING(LastModified.ToString(), "TimeStamp"));
 
 	FString Name = "[StaticImageResource]-" + AssetUUID;
-	DesiredDesc.Format = BufferFormat::LateBound;
-	BlobObj = JobObj->InitResult(Name, &DesiredDesc);
+	BufferDescriptor Desc;
+
+	if (DesiredDesc)
+		Desc = *DesiredDesc;
+
+	Desc.Format = BufferFormat::LateBound;
+	BlobObj = JobObj->InitResult(Name, &Desc);
 
 	Cycle->AddJob(TargetId, std::move(JobObj));
 
diff --git a/Engine/Plugins/Experimental/TextureGraph/Source/TextureGraphEngine/Model/StaticImageResource.h b/Engine/Plugins/Experimental/TextureGraph/Source/TextureGraphEngine/Model/StaticImageResource.h
--- a/Engine/Plugins/Experimental/TextureGraph/Source/TextureGraphEngine/Model/StaticImageResource.h
+++ b/Engine/Plugins/Experimental/TextureGraph/Source/TextureGraphEngine/Model/StaticImageResource.h
@@ -36,7 +36,7 @@ private:
 public:
 	virtual ~UStaticImageResource() override;
 	
-	virtual TiledBlobPtr GetBlob(MixUpdateCyclePtr Cycle, BufferDescriptor DesiredDesc, int32 TargetId);
+	virtual TiledBlobPtr GetBlob(MixUpdateCyclePtr Cycle, BufferDescriptor* DesiredDesc, int32 TargetId);
 
 	//////////////////////////////////////////////////////////////////////////
 	/// Inline functions
diff --git a/Engine/Plugins/Experimental/TextureGraph/Source/TextureGraphEngine/Transform/Expressions/T_ArrayGrid.cpp b/Engine/Plugins/Experimental/TextureGraph/Source/TextureGraphEngine/Transform/Expressions/T_ArrayGrid.cpp
new file mode 100644
--- /dev/null
+++ b/Engine/Plugins/Experimental/TextureGraph/Source/TextureGraphEngine/Transform/Expressions/T_ArrayGrid.cpp
@@ -0,0 +1,168 @@
+// Copyright Epic Games, Inc. All Rights Reserved.
+
+#include "T_ArrayGrid.h"
+#include "Job/JobArgs.h"
+#include "TextureGraphEngine.h"
+#include "Math/Vector.h"
+#include "FxMat/MaterialManager.h"
+#include "Job/JobBatch.h"
+#include "Helper/GraphicsUtil.h"
+#include "Transform/Utility/T_CombineTiledBlob.h"
+#include "Transform/Utility/T_SplitToTiles.h"
+#include "Model/Mix/MixInterface.h"
+#include "Device/FX/Device_FX.h"
+
+class TEXTUREGRAPHENGINE_API Job_ArrayGrid : public Job
+{
+private:
+	TArray<TiledBlobPtr>			Inputs;			/// The inputs to this array grid
+	int32							NumRows = 0;	/// Number of rows
+	int32							NumCols = 0;	/// Number of columns
+	uint32							MaxColWidth = 0;/// Maximum column width
+	uint32							MaxRowHeight = 0;/// Maximum row height
+	BufferDescriptor				InDesiredOutputDesc;	/// The desired input desc
+	TiledBlob_PromisePtr			TiledInput;		/// The tiled input which is a combination of all the array tiles
+
+public:
+	Job_ArrayGrid(UMixInterface* InMix, int32 TargetId, TArray<TiledBlobPtr> InInputs, 
+		int32 InNumRows, int32 InNumCols, uint32 InMaxColWidth, uint32 InMaxRowHeight, BufferDescriptor& InDesiredOutputDesc_, 
+		TiledBlob_PromisePtr InTiledInput)
+		: Job(InMix, TargetId, std::make_shared<Null_Transform>(Device_FX::Get(), TEXT("T_ArrayGrid"), true, true))
+		, Inputs(InInputs)
+		, NumRows(InNumRows)
+		, NumCols(InNumCols)
+		, MaxColWidth(InMaxColWidth)
+		, MaxRowHeight(InMaxRowHeight)
+		, InDesiredOutputDesc(InDesiredOutputDesc_)
+		, TiledInput(InTiledInput)
+	{
+		Name = TEXT("T_ArrayGrid");
+	}
+
+	static BlobPtrTiles				GetTiles(TArray<TiledBlobPtr>& Inputs, int32 InNumRows, int32 InNumCols, 
+		int32 NumRows, int32 NumCols, uint32* OutMaxRowHeight = nullptr, uint32* OutMaxColWidth = nullptr, 
+		bool* bIsWellDefinedOut = nullptr, CHashPtrVec* OutHashes = nullptr)
+	{
+		BlobPtrTiles ResultTiles(InNumRows * NumRows, InNumCols * NumCols); 
+		CHashPtrVec InputHashes(NumRows * NumCols);
+		bool bIsWellDefined = true;
+		uint32 MaxRowHeight = 0;
+		uint32 MaxColWidth = 0;
+
+		/// First of all we combine all the tiles into a single result
+		for (int32 X = 0; X < NumCols; X++)
+		{
+			int32 ColIndex = X * InNumCols;
+			for (int32 Y = 0; Y < NumRows; Y++)
+			{
+				int32 TileIndex = Y * NumCols + X;
+				check(TileIndex < Inputs.Num());
+				TiledBlobPtr Input = Inputs[TileIndex];
+				int32 RowIndex = Y * InNumRows;
+
+				MaxRowHeight = std::max(MaxRowHeight, Input->GetHeight());
+				MaxColWidth = std::max(MaxColWidth, Input->GetWidth());
+
+				InputHashes[TileIndex] = Input->Hash();
+				check(InputHashes[TileIndex]);
+				bIsWellDefined &= Input->IsWellDefined();
+
+				for (int32 TileX = 0; TileX < Input->Rows(); TileX++)
+				{
+					for (int32 TileY = 0; TileY < Input->Cols(); TileY++)
+					{
+						BlobPtr Tile = Input->GetTile(TileX, TileY);
+						ResultTiles[RowIndex + TileX][ColIndex + TileY] = Tile;
+
+						if (Tile)
+							bIsWellDefined &= Tile->IsWellDefined();
+					}
+				}
+			}
+		}
+
+		if (OutMaxRowHeight)
+			*OutMaxRowHeight = MaxRowHeight;
+		if (OutMaxColWidth)
+			*OutMaxColWidth = MaxColWidth;
+		if (bIsWellDefinedOut)
+			*bIsWellDefinedOut = bIsWellDefined;
+		if (OutHashes)
+			*OutHashes = std::move(InputHashes);
+
+		return ResultTiles;
+	}
+
+protected:
+	virtual AsyncPrepareResult		PrepareTargets(JobBatch* Batch) override
+	{
+		int32 InNumRows = Inputs[0]->Rows();
+		int32 InNumCols = Inputs[0]->Cols();
+		BlobPtrTiles ResultTiles = GetTiles(Inputs, InNumRows, InNumCols, NumRows, NumCols, &MaxRowHeight, &MaxColWidth);
+
+		/// The we combine those tiles into a single result
+		const BufferDescriptor& InputDesc = Inputs[0]->GetDescriptor();
+		BufferDescriptor DesiredOutputDesc = BufferDescriptor::Combine(InputDesc, InDesiredOutputDesc);
+
+		if (DesiredOutputDesc.IsLateBound())
+			DesiredOutputDesc.Format = InputDesc.Format;
+
+		DesiredOutputDesc.Width = MaxColWidth * NumCols;
+		DesiredOutputDesc.Height = MaxRowHeight * NumRows;
+
+		TiledBlobPtr CombinedResultTiled = TiledBlob::InitFromTiles(DesiredOutputDesc, ResultTiles);
+
+		if (ResultOrg)
+			ResultOrg->ResolveLateBound(DesiredOutputDesc, true);
+
+		if (Result && Result.get().get() != ResultOrg.get())
+			Result->ResolveLateBound(DesiredOutputDesc, true);
+
+		std::shared_ptr<CombineTiledBlob_Transform> CombinedTiledBlobMat = std::make_shared<CombineTiledBlob_Transform>(Name, CombinedResultTiled);
+		Job::Transform = CombinedTiledBlobMat;
+
+		return Job::PrepareTargets(Batch);
+	}
+};
+
+//////////////////////////////////////////////////////////////////////////
+TiledBlobPtr T_ArrayGrid::Create(MixUpdateCyclePtr Cycle, BufferDescriptor& InDesiredOutputDesc, TArray<TiledBlobPtr> Inputs, int32 NumRows, int32 NumCols, int32 TargetId)
+{
+	check(!Inputs.IsEmpty());
+
+	int32 InNumRows = Cycle->GetMix()->GetNumXTiles();
+	int32 InNumCols = Cycle->GetMix()->GetNumYTiles();
+	CHashPtrVec InputHashes(Inputs.Num());
+	bool bIsWellDefined = true;
+	uint32 MaxRowHeight = 0;
+	uint32 MaxColWidth = 0;
+
+	BlobPtrTiles ResultTiles = Job_ArrayGrid::GetTiles(Inputs, InNumRows, InNumCols, NumRows, NumCols, &MaxColWidth, &MaxRowHeight, &bIsWellDefined, &InputHashes);
+	TiledBlobPtr FinalResult;
+
+	/// If the result is late bound then we have no option but to resize the result render target through a custom job
+	if (!bIsWellDefined)
+	{
+		CHashPtr Hash = CHash::ConstructFromSources(InputHashes);
+		BufferDescriptor Desc = InDesiredOutputDesc;
+		Desc.Format = BufferFormat::LateBound;
+		Desc.Width = Desc.Height = 0;
+		TiledBlob_PromisePtr CombinedResultTiled = std::make_shared<TiledBlob_Promise>(Desc, InNumRows * NumRows, InNumCols * NumCols, Hash);
+		JobUPtr JobObj = std::make_unique<Job_ArrayGrid>(Cycle->GetMix(), TargetId, Inputs, NumRows, NumCols, MaxColWidth, MaxRowHeight, InDesiredOutputDesc, CombinedResultTiled);
+		TiledBlobPtr CombinedResultOne = T_CombineTiledBlob::Create(Cycle, InDesiredOutputDesc, TargetId, CombinedResultTiled, std::move(JobObj));
+		FinalResult = T_SplitToTiles::Create(Cycle, TargetId, CombinedResultOne);
+	}
+	else
+	{
+		BufferDescriptor DesiredOutputDesc = BufferDescriptor::Combine(Inputs[0]->GetDescriptor(), InDesiredOutputDesc);
+		DesiredOutputDesc.Width = MaxColWidth * NumCols;
+		DesiredOutputDesc.Height = MaxRowHeight * NumRows;
+
+		/// The we combine those tiles into a single result
+		TiledBlobPtr CombinedResultTiled = TiledBlob::InitFromTiles(DesiredOutputDesc, ResultTiles);
+		TiledBlobPtr CombinedResultOne = T_CombineTiledBlob::Create(Cycle, DesiredOutputDesc, TargetId, CombinedResultTiled);
+		FinalResult = T_SplitToTiles::Create(Cycle, TargetId, CombinedResultOne);
+	}
+
+	return FinalResult;
+}
diff --git a/Engine/Plugins/Experimental/TextureGraph/Source/TextureGraphEngine/Transform/Expressions/T_ArrayGrid.h b/Engine/Plugins/Experimental/TextureGraph/Source/TextureGraphEngine/Transform/Expressions/T_ArrayGrid.h
new file mode 100644
--- /dev/null
+++ b/Engine/Plugins/Experimental/TextureGraph/Source/TextureGraphEngine/Transform/Expressions/T_ArrayGrid.h
@@ -0,0 +1,18 @@
+// Copyright Epic Games, Inc. All Rights Reserved.
+#pragma once
+
+#include "CoreMinimal.h"
+#include "Math/Vector2D.h"
+
+#include "Job/Job.h"
+#include "FxMat/FxMaterial.h"
+#include "Model/Mix/MixUpdateCycle.h"
+#include <DataDrivenShaderPlatformInfo.h>
+
+//////////////////////////////////////////////////////////////////////////
+class TEXTUREGRAPHENGINE_API T_ArrayGrid
+{
+public:
+	static TiledBlobPtr Create(MixUpdateCyclePtr Cycle, BufferDescriptor& DesiredOutputDesc, TArray<TiledBlobPtr> Inputs, int32 NumRows, int32 NumCols, int32 TargetId = 0);
+};
+
diff --git a/Engine/Plugins/Experimental/TextureGraph/Source/TextureGraphEngine/Transform/Utility/T_CombineTiledBlob.cpp b/Engine/Plugins/Experimental/TextureGraph/Source/TextureGraphEngine/Transform/Utility/T_CombineTiledBlob.cpp
--- a/Engine/Plugins/Experimental/TextureGraph/Source/TextureGraphEngine/Transform/Utility/T_CombineTiledBlob.cpp
+++ b/Engine/Plugins/Experimental/TextureGraph/Source/TextureGraphEngine/Transform/Utility/T_CombineTiledBlob.cpp
@@ -7,15 +7,21 @@
 #include "Device/FX/Device_FX.h"
 #include "Model/Mix/MixInterface.h"
 
-TiledBlobPtr T_CombineTiledBlob::Create(MixUpdateCyclePtr Cycle, BufferDescriptor DesiredOutputDesc, int32 TargetId, TiledBlobPtr SourceTex)
+TiledBlobPtr T_CombineTiledBlob::Create(MixUpdateCyclePtr Cycle, BufferDescriptor DesiredOutputDesc, int32 TargetId, TiledBlobPtr SourceTex, JobUPtr JobToUse)
 {
 	if (SourceTex && SourceTex->IsTiled())
 	{
-		FString Name = TEXT("T_CombineTiledBlob"); //FString::Printf(TEXT("[%s].[%d].[%llu] CombineTiledBlob"), *SourceTex->Name(), TargetId, Cycle->GetBatch()->GetFrameId());
+		JobUPtr JobObj;
 
-		std::shared_ptr<CombineTiledBlob_Transform> CombinedTiledBlobMat = std::make_shared< CombineTiledBlob_Transform>(Name, SourceTex);
+		if (!JobToUse)
+		{
+			FString Name = TEXT("T_CombineTiledBlob"); //FString::Printf(TEXT("[%s].[%d].[%llu] CombineTiledBlob"), *SourceTex->Name(), TargetId, Cycle->GetBatch()->GetFrameId());
+			std::shared_ptr<CombineTiledBlob_Transform> CombinedTiledBlobMat = std::make_shared< CombineTiledBlob_Transform>(Name, SourceTex);
+			JobObj = std::make_unique<Job>(Cycle->GetMix(), TargetId, std::static_pointer_cast<BlobTransform>(CombinedTiledBlobMat));
+		}
+		else
+			JobObj = std::move(JobToUse);
 
-		JobUPtr JobObj = std::make_unique<Job>(Cycle->GetMix(), TargetId, std::static_pointer_cast<BlobTransform>(CombinedTiledBlobMat));
 		JobObj->AddArg(WithUnbounded(ARG_BLOB(SourceTex, "Source")));
 
 		// Express the dependency of the new job on the job delivering the SourceTex
@@ -23,7 +29,7 @@ TiledBlobPtr T_CombineTiledBlob::Create(MixUpdateCyclePtr Cycle, BufferDescripto
 		if (!PrevJob.expired())
 			JobObj->AddPrev(std::static_pointer_cast<DeviceNativeTask>(PrevJob.lock()));
 
-		TiledBlobPtr Result = JobObj->InitResult(Name, &DesiredOutputDesc, 1, 1);
+		TiledBlobPtr Result = JobObj->InitResult(JobObj->GetName(), &DesiredOutputDesc, 1, 1);
 
 		Cycle->AddJob(TargetId, std::move(JobObj));
 		Result->MakeSingleBlob();
@@ -71,9 +77,9 @@ AsyncTransformResultPtr	CombineTiledBlob_Transform::Exec(const TransformArgs& Ar
 	{
 		for (int32 TileY = 0; TileY < Source->Cols(); TileY++)
 		{
-			check(Source->GetTiles()[TileX][TileY]);
+			check(Source->GetTile(TileX, TileY));
 
-			BlobPtr Tile = Source->GetTiles()[TileX][TileY];
+			BlobPtr Tile = Source->GetTile(TileX, TileY);
 			check(Tile->GetBufferRef());
 
 			Tiles[TileX][TileY] = Tile->GetBufferRef();
diff --git a/Engine/Plugins/Experimental/TextureGraph/Source/TextureGraphEngine/Transform/Utility/T_CombineTiledBlob.h b/Engine/Plugins/Experimental/TextureGraph/Source/TextureGraphEngine/Transform/Utility/T_CombineTiledBlob.h
--- a/Engine/Plugins/Experimental/TextureGraph/Source/TextureGraphEngine/Transform/Utility/T_CombineTiledBlob.h
+++ b/Engine/Plugins/Experimental/TextureGraph/Source/TextureGraphEngine/Transform/Utility/T_CombineTiledBlob.h
@@ -40,5 +40,5 @@ public:
 	/// Static functions
 	//////////////////////////////////////////////////////////////////////////
 	
-	static TiledBlobPtr				Create(MixUpdateCyclePtr Cycle, BufferDescriptor DesiredOutputDesc, int32 TargetId, TiledBlobPtr SourceTex);
+	static TiledBlobPtr				Create(MixUpdateCyclePtr Cycle, BufferDescriptor DesiredOutputDesc, int32 TargetId, TiledBlobPtr SourceTex, JobUPtr JobToUse = nullptr);
 };
diff --git a/Engine/Plugins/Experimental/TextureGraph/Source/TextureGraphEngine/Transform/Utility/T_SplitToTiles.cpp b/Engine/Plugins/Experimental/TextureGraph/Source/TextureGraphEngine/Transform/Utility/T_SplitToTiles.cpp
--- a/Engine/Plugins/Experimental/TextureGraph/Source/TextureGraphEngine/Transform/Utility/T_SplitToTiles.cpp
+++ b/Engine/Plugins/Experimental/TextureGraph/Source/TextureGraphEngine/Transform/Utility/T_SplitToTiles.cpp
@@ -21,30 +21,22 @@ TiledBlobPtr CreateSplitToTilesCompute(MixUpdateCyclePtr Cycle, TiledBlobPtr Sou
 {
 	CSH_SplitToTiles::FPermutationDomain PermutationVector;
 
-	FString Name = FString::Printf(TEXT("[%s].[%d].[%llu] SplitToTiles"), *SourceTex->Name(), TargetId, Cycle->GetBatch()->GetFrameId());
-
-	int32 SrcWidth = SourceTex->GetWidth();
-	int32 SrcHeight = SourceTex->GetHeight();
-
-	int32 DstWidth = SrcWidth;
-	int32 DstHeight = SrcHeight;
-
-	int32 srcNumCols = (int32)SourceTex->GetTiles().Cols();
-	int32 srcNumRows = (int32)SourceTex->GetTiles().Rows();	
+	FString Name = FString::Printf(TEXT("[%s]_SplitToTiles"), *SourceTex->Name());
 
 	int32 DstNumCols = Cycle->GetMix()->GetNumXTiles();
 	int32 DstNumRows = Cycle->GetMix()->GetNumYTiles();
 
-	FIntVector4 SrcDimensions(DstWidth / DstNumCols, DstHeight / DstNumRows, srcNumCols, srcNumRows);
+	/// This is the correct check to see if we can actually produce the correct number of rows and columns
 	FTileInfo TileInfo;
 
-	RenderMaterial_FXPtr Transform = TextureGraphEngine::GetMaterialManager()->CreateMaterial_CmpFX<CSH_SplitToTiles>(
-		Name, TEXT("Result"), PermutationVector, SrcDimensions.X, SrcDimensions.Y, 1);
+	std::shared_ptr<FxMaterial_SplitToTiles> Mat = std::make_shared<FxMaterial_SplitToTiles>(TEXT("Result"), DstNumCols, DstNumRows, SourceTex, &PermutationVector);
+	RenderMaterial_FXPtr Transform = std::make_shared<RenderMaterial_FX>(Name, Mat);
 
 	JobUPtr JobObj = std::make_unique<Job>(Cycle->GetMix(), TargetId, std::static_pointer_cast<BlobTransform>(Transform));
 	JobObj
 		->AddArg(ARG_BLOB(SourceTex, "SourceTexture"))
-		->AddArg(ARG_TILEINFO(TileInfo, "TileInfo"));
+		->AddArg(ARG_TILEINFO(TileInfo, "TileInfo"))
+		;
 
 	BufferDescriptor Desc = SourceTex->GetDescriptor();
 	Desc.Name = Name;
@@ -62,6 +54,5 @@ TiledBlobPtr CreateSplitToTilesCompute(MixUpdateCyclePtr Cycle, TiledBlobPtr Sou
 TiledBlobPtr T_SplitToTiles::Create(MixUpdateCyclePtr Cycle, int32 TargetId, TiledBlobPtr SourceTex)
 {
 	/// Cannot generate. Just return the source texture
-	check(TextureHelper::IsPowerOf2(SourceTex->GetWidth()) && TextureHelper::IsPowerOf2(SourceTex->GetHeight()));
 	return CreateSplitToTilesCompute(Cycle, SourceTex, TargetId);
 }
diff --git a/Engine/Plugins/Experimental/TextureGraph/Source/TextureGraphEngine/Transform/Utility/T_SplitToTiles.h b/Engine/Plugins/Experimental/TextureGraph/Source/TextureGraphEngine/Transform/Utility/T_SplitToTiles.h
--- a/Engine/Plugins/Experimental/TextureGraph/Source/TextureGraphEngine/Transform/Utility/T_SplitToTiles.h
+++ b/Engine/Plugins/Experimental/TextureGraph/Source/TextureGraphEngine/Transform/Utility/T_SplitToTiles.h
@@ -30,6 +30,43 @@ public:
 	}
 };
 
+class FxMaterial_SplitToTiles : public FxMaterial_Compute<CSH_SplitToTiles>
+{
+private:
+	int32 NumTilesX;
+	int32 NumTilesY;
+	TiledBlobPtr RefBlob;
+
+public:
+	explicit FxMaterial_SplitToTiles(FString InOutputId, int32 InNumTilesX, int32 InNumTilesY, TiledBlobPtr InRefBlob, 
+		const CmpSHPermutationDomain* InPermDomain = nullptr, FUnorderedAccessViewRHIRef InUnorderedAccessView = nullptr) 
+		: FxMaterial_Compute<CSH_SplitToTiles>(InOutputId, InPermDomain, 0, 0, 1, InUnorderedAccessView)
+		, NumTilesX(InNumTilesX)
+		, NumTilesY(InNumTilesY)
+		, RefBlob(InRefBlob)
+	{
+	}
+
+	virtual void Blit(FRHICommandListImmediate& RHI, FRHITexture* Target, const RenderMesh* MeshObj, int32 TargetId, FGraphicsPipelineStateInitializer* PSO = nullptr) override
+	{
+		if (NumThreadsX == 0 || NumThreadsY == 0)
+		{
+			NumThreadsX = RefBlob->GetWidth() / NumTilesX;
+			NumThreadsY = RefBlob->GetHeight() / NumTilesY;
+		}
+
+		check(NumThreadsX > 0 && NumThreadsY > 0);
+
+		return FxMaterial_Compute<CSH_SplitToTiles>::Blit(RHI, Target, MeshObj, TargetId, PSO);
+	}
+
+	virtual std::shared_ptr<FxMaterial> Clone() override
+	{
+		std::shared_ptr<FxMaterial_SplitToTiles> Mat = std::make_shared<FxMaterial_SplitToTiles>(OutputId, NumTilesX, NumTilesY, RefBlob, &PermutationDomain, UnorderedAccessView);
+		return std::static_pointer_cast<FxMaterial>(Mat);
+	}
+};
+
 /**
  * SplitToTiles Transform
  */
