Build fix

diff --git a/Engine/Plugins/Experimental/TextureGraph/Source/TextureGraph/Public/Expressions/Input/TG_Expression_TexturePath.h b/Engine/Plugins/Experimental/TextureGraph/Source/TextureGraph/Public/Expressions/Input/TG_Expression_TexturePath.h
--- a/Engine/Plugins/Experimental/TextureGraph/Source/TextureGraph/Public/Expressions/Input/TG_Expression_TexturePath.h
+++ b/Engine/Plugins/Experimental/TextureGraph/Source/TextureGraph/Public/Expressions/Input/TG_Expression_TexturePath.h
@@ -26,16 +26,17 @@ protected:
 	bool ValidateInputPath(FString& ValidatedPath) const;
 
 	void ReportError(MixUpdateCyclePtr Cycle);
+	static TiledBlobPtr LoadStaticImage(FTG_EvaluationContext* InContext, const FString& Filename, BufferDescriptor* DesiredDesc);
 public:
 
 	virtual void Evaluate(FTG_EvaluationContext* InContext) override;
 	virtual bool Validate(MixUpdateCyclePtr	Cycle) override;
 	
-	// The output of the node, which is the loaded texture from the path
+	// All the output textures from this node. Note that it outputs an array of textures.
 	UPROPERTY(meta = (TGType = "TG_Output", PinDisplayName = "", HideInnerPropertiesInNode))
-	FTG_Texture Output;
+	FTG_VariantArray Output;
 
-	// Input file path of the texture
+	// Input file path of the texture. If the path is a directory then all files from that directory are loaded
 	UPROPERTY(EditAnywhere, BlueprintReadWrite, Category = NoCategory, meta = (TGType = "TG_InputParam", NoResetToDefault) )
 	FString Path;
 
@@ -43,6 +44,6 @@ public:
 	virtual FTG_Name GetDefaultName() const override { return TEXT("TexturePath"); }
 	virtual void SetTitleName(FName NewName) override;
 	virtual FName GetTitleName() const override;
-	virtual FText GetTooltipText() const override { return FText::FromString(TEXT("Loads a texture from a path.")); }
+	virtual FText GetTooltipText() const override { return FText::FromString(TEXT("Loads texture(s) from a path. If the path is a single file name then one texture is loaded. If it's a directory then all textures are loaded from that directory.")); }
 };
 
