Protecting editor only code

diff --git a/Engine/Plugins/Experimental/TextureGraph/Source/TextureGraphEngine/FxMat/RenderMaterial_BP.cpp b/Engine/Plugins/Experimental/TextureGraph/Source/TextureGraphEngine/FxMat/RenderMaterial_BP.cpp
--- a/Engine/Plugins/Experimental/TextureGraph/Source/TextureGraphEngine/FxMat/RenderMaterial_BP.cpp
+++ b/Engine/Plugins/Experimental/TextureGraph/Source/TextureGraphEngine/FxMat/RenderMaterial_BP.cpp
@@ -16,8 +16,6 @@
 #include "Async/Async.h"
 #include "FxMaterial_DrawMaterial.h"
 #include "Job/Job.h"
-#include "AssetToolsModule.h"
-#include "IAssetTools.h"
 
 
 TSharedPtr<FMaterialInstanceStaticSwitchPermutationMap> FMaterialInstanceStaticSwitchPermutationMap::Create(UMaterialInterface* InMaterial)
@@ -27,20 +25,18 @@ TSharedPtr<FMaterialInstanceStaticSwitchPermutationMap> FMaterialInstanceStaticS
 	{
 		MISSPMap = MakeShared<FMaterialInstanceStaticSwitchPermutationMap>();
 
+		UMaterialInstanceConstant* MIC = NewObject<UMaterialInstanceConstant>(InMaterial);
+#if WITH_EDITOR
 		InMaterial->GetStaticParameterValues(MISSPMap->DefaultStaticParameterSet);
-
-		int32 NumSwitches = MISSPMap->DefaultStaticParameterSet.StaticSwitchParameters.Num();
-
 		MISSPMap->DefaultKey = MISSPMap->KeyFromStaticSwitchParameters(MISSPMap->DefaultStaticParameterSet.StaticSwitchParameters);
 
-		UMaterialInstanceConstant* MIC = NewObject<UMaterialInstanceConstant>(InMaterial);
 		MIC->SetParentEditorOnly(InMaterial);
 		MIC->ClearParameterValuesEditorOnly();
 		MIC->SetFlags(RF_Standalone);
 		MIC->MarkPackageDirty();
 		MIC->PreEditChange(NULL);
 		MIC->PostEditChange();
-
+#endif
 		MISSPMap->PermutationsMap.Add(MISSPMap->DefaultKey, MIC);
 	}
 	return MISSPMap;
@@ -78,20 +74,18 @@ UMaterialInstanceConstant* FMaterialInstanceStaticSwitchPermutationMap::GetMater
 	if (!MICP)
 	{
 		UMaterialInstanceConstant* MIC = NewObject<UMaterialInstanceConstant>(GetRootMaterial());
+#if WITH_EDITOR
 		MIC->SetParentEditorOnly(GetRootMaterial());
 		MIC->ClearParameterValuesEditorOnly();
 		MIC->SetFlags(RF_Standalone);
 		MIC->MarkPackageDirty();
 
-		bool NeedStaticUpdate = false;
-
 		FStaticParameterSet InstanceStaticParameterSet = DefaultStaticParameterSet;
 		InstanceStaticParameterSet.StaticSwitchParameters = Parameters;
-
 		MIC->UpdateStaticPermutation(InstanceStaticParameterSet);
 		MIC->PreEditChange(NULL);
 		MIC->PostEditChange();
-
+#endif
 		PermutationsMap.Add(Key, MIC);
 
 		return MIC;
@@ -120,13 +114,16 @@ RenderMaterial_BP::RenderMaterial_BP(FString Name, UMaterialInterface* InMateria
 	}
 	else if (Material && !MaterialInstance)
 	{
+
 		MaterialInstance = TStrongObjectPtr<UMaterialInstanceConstant>(NewObject<UMaterialInstanceConstant>(Material));
+#if WITH_EDITOR
 		MaterialInstance->SetParentEditorOnly(Material);
 		MaterialInstance->ClearParameterValuesEditorOnly();
 		MaterialInstance->SetFlags(RF_Standalone);
 		MaterialInstance->MarkPackageDirty();
 		MaterialInstance->PreEditChange(NULL);
 		MaterialInstance->PostEditChange();
+#endif
 	}
 
 	check(Material && MaterialInstance);
