Fixed Texture Graph Node/3D Previews don't appear upon opening asset.

diff --git a/Engine/Plugins/Experimental/TextureGraph/Source/TextureGraphEditor/Private/TG_Editor.cpp b/Engine/Plugins/Experimental/TextureGraph/Source/TextureGraphEditor/Private/TG_Editor.cpp
--- a/Engine/Plugins/Experimental/TextureGraph/Source/TextureGraphEditor/Private/TG_Editor.cpp
+++ b/Engine/Plugins/Experimental/TextureGraph/Source/TextureGraphEditor/Private/TG_Editor.cpp
@@ -637,6 +637,7 @@ void FTG_Editor::OnRenderingDone(UMixInterface* TextureGraph, const FInvalidatio
 {
 	if (TextureGraph != nullptr)
 	{
+		GetEditorViewport()->UpdateRenderMode();
 		RefreshNodePreview(GraphEditorWidget->GetSelectedNodes(), Details, true);
 		RefreshViewport();
 	}
diff --git a/Engine/Plugins/Experimental/TextureGraph/Source/TextureGraphEngine/Helper/DataUtil.h b/Engine/Plugins/Experimental/TextureGraph/Source/TextureGraphEngine/Helper/DataUtil.h
--- a/Engine/Plugins/Experimental/TextureGraph/Source/TextureGraphEngine/Helper/DataUtil.h
+++ b/Engine/Plugins/Experimental/TextureGraph/Source/TextureGraphEngine/Helper/DataUtil.h
@@ -21,7 +21,7 @@ typedef std::vector<CHashPtrW>		CHashPtrWVec;
 //{
 //	uint64							HashValue;			/// The actual HashValue Value
 //};
-DECLARE_LOG_CATEGORY_EXTERN(LogData, Log, Verbose);
+DECLARE_LOG_CATEGORY_EXTERN(LogData, Log, All);
 
 #define MX_HASH_VAL(HashValue, Prime, val) ((HashValue * Prime) ^ (val))
 #define MX_HASH_VAL_DEF(val)		MX_HASH_VAL(DataUtil::GFNVInit, DataUtil::GFNVPrime, (HashType)val)
diff --git a/Engine/Plugins/Experimental/TextureGraph/Source/TextureGraphEngine/Model/Mix/MixManager.cpp b/Engine/Plugins/Experimental/TextureGraph/Source/TextureGraphEngine/Model/Mix/MixManager.cpp
--- a/Engine/Plugins/Experimental/TextureGraph/Source/TextureGraphEngine/Model/Mix/MixManager.cpp
+++ b/Engine/Plugins/Experimental/TextureGraph/Source/TextureGraphEngine/Model/Mix/MixManager.cpp
@@ -48,7 +48,12 @@ void MixManager::Update(float Delta)
 	std::vector<MixInvalidateInfo> MixInfos = Queue.to_vector_and_clear();
 
 	if (!MixInfos.size())
+	{
+		// if no invalidations queued, we can clear out MixesToFlush as well,
+		// this avoids a bug where the next time this same mix is opened it is immediately flushed and doesn't trigger an update
+		MixesToFlush.clear();
 		return;
+	}
 
 	if (!MixesToFlush.empty())
 	{
