Reload texture graph crash fix

diff --git a/Engine/Plugins/Experimental/TextureGraph/Source/TextureGraph/Private/Expressions/CommandChange/TG_PinConnectionChange.cpp b/Engine/Plugins/Experimental/TextureGraph/Source/TextureGraph/Private/Expressions/CommandChange/TG_PinConnectionChange.cpp
--- a/Engine/Plugins/Experimental/TextureGraph/Source/TextureGraph/Private/Expressions/CommandChange/TG_PinConnectionChange.cpp
+++ b/Engine/Plugins/Experimental/TextureGraph/Source/TextureGraph/Private/Expressions/CommandChange/TG_PinConnectionChange.cpp
@@ -9,7 +9,7 @@
 void FTG_PinConnectionChange::StoreChange(UTG_Graph* InGraph, UTG_Node& NodeFrom, FTG_Name& PinFromName)
 {
 #if WITH_EDITOR
-	if (!GIsTransacting)
+	if (!GIsTransacting && GUndo)
 	{
 		TUniquePtr<FTG_PinConnectionChange> Change = MakeUnique<FTG_PinConnectionChange>(InGraph, NodeFrom, PinFromName);
 		GUndo->StoreUndo(InGraph, MoveTemp(Change));
@@ -38,7 +38,7 @@ void FTG_PinConnectionChange::Revert(UObject* Object)
 void FTG_PinConnectionBreakChange::StoreChange(UTG_Graph* InGraph, UTG_Node& NodeFrom, FTG_Name& PinFromName)
 {
 #if WITH_EDITOR
-	if (!GIsTransacting)
+	if (!GIsTransacting && GUndo)
 	{
 		TUniquePtr<FTG_PinConnectionBreakChange> Change = MakeUnique<FTG_PinConnectionBreakChange>(InGraph, NodeFrom, PinFromName);
 		GUndo->StoreUndo(InGraph, MoveTemp(Change));
