Path validation failure fixes.

diff --git a/Engine/Plugins/Experimental/TextureGraph/Source/TextureGraph/Private/Expressions/Input/TG_Expression_TexturePath.cpp b/Engine/Plugins/Experimental/TextureGraph/Source/TextureGraph/Private/Expressions/Input/TG_Expression_TexturePath.cpp
--- a/Engine/Plugins/Experimental/TextureGraph/Source/TextureGraph/Private/Expressions/Input/TG_Expression_TexturePath.cpp
+++ b/Engine/Plugins/Experimental/TextureGraph/Source/TextureGraph/Private/Expressions/Input/TG_Expression_TexturePath.cpp
@@ -48,7 +48,7 @@ bool UTG_Expression_TexturePath::ValidateInputPath(FString& ValidatedPath) const
 		ValidatedPath = LocalPath;
 		return true;
 	}
-	else if(FPaths::FileExists(LocalPath))
+	else if(FPaths::FileExists(LocalPath) || FPaths::DirectoryExists(LocalPath))
 	{
 		ValidatedPath = LocalPath;
 		return true;
