Exposing Material output as individual out pins on the node

diff --git a/Engine/Plugins/Experimental/TextureGraph/Source/TextureGraph/Private/Expressions/TG_Expression_MaterialBase.cpp b/Engine/Plugins/Experimental/TextureGraph/Source/TextureGraph/Private/Expressions/TG_Expression_MaterialBase.cpp
--- a/Engine/Plugins/Experimental/TextureGraph/Source/TextureGraph/Private/Expressions/TG_Expression_MaterialBase.cpp
+++ b/Engine/Plugins/Experimental/TextureGraph/Source/TextureGraph/Private/Expressions/TG_Expression_MaterialBase.cpp
@@ -102,6 +102,7 @@ void UTG_Expression_MaterialBase::GenerateMaterialAttributeOptions()
 	// Detect the set of available material properties for rendering
 	AvailableMaterialAttributeIds.Empty();
 	AvailableMaterialAttributeNames.Empty();
+
 	if (GetMaterial())
 	{
 		UMaterial* RefMaterial = GetMaterial()->GetMaterial();
@@ -163,22 +164,38 @@ void UTG_Expression_MaterialBase::Evaluate(FTG_EvaluationContext* InContext)
 {
 	Super::Evaluate(InContext);
 
-	TiledBlobPtr Result = TextureHelper::GetBlack();
-
+	
 	/// Set it to false always
 	TiledMode = true;
 
+	Output = TextureHelper::GetBlack();
 	if (GetMaterial())
 	{
 
 		FString AssetName = GetMaterial()->GetName();
-
 		UMaterialInstanceConstant* MIC = GetEvaluationMaterialInstanceConstant(InContext);
-		const auto RenderMaterial = std::make_shared<RenderMaterial_BP>(AssetName, MIC);
-		Result = CreateRenderMaterialJob(InContext, RenderMaterial, Output.GetBufferDescriptor(), GetRenderedAttributeId());
+
+		for (int32 i = 0; i < AvailableMaterialAttributeIds.Num(); ++i)
+		{
+			const auto RenderMaterial = std::make_shared<RenderMaterial_BP>(AssetName, MIC);
+			TiledBlobPtr Result = TextureHelper::GetBlack();
+			Result = CreateRenderMaterialJob(InContext, RenderMaterial, Output.GetBufferDescriptor(), AvailableMaterialAttributeIds[i]);
+			Outputs[i] = Result;
+
+			if (GetRenderedAttributeId() == AvailableMaterialAttributeIds[i])
+			{
+				Output = Result;
+			}
+
+			// Reassign the generated result to the corresponding Output arg
+			FTG_Var* OutVar = InContext->Outputs.GetVar(MatAttributesOutputArgNames[i]);
+			if (OutVar)
+			{
+				OutVar->EditAs<FTG_Texture>() = Result;
+			}
+		}
 	}
 
-	Output = Result;
 }
 
 TiledBlobPtr UTG_Expression_MaterialBase::CreateRenderMaterialJob(FTG_EvaluationContext* InContext, const FString& InName, const FString& InMaterialPath, const BufferDescriptor& InDescriptor, EDrawMaterialAttributeTarget InDrawMaterialAttributeTarget)
@@ -434,6 +451,31 @@ FTG_SignaturePtr UTG_Expression_MaterialBase::BuildSignatureDynamically() const
 		GetMaterial()->GetAllStaticSwitchParameterInfo(OutParameterInfo, OutParameterIds);
 		AddSignatureParam(OutParameterInfo, OutParameterIds, EMaterialParameterType::StaticSwitch, SignatureInit);
 #endif		
+
+		// Add the output argument(s) declaration
+		// Make the output argument names unique in the Node signature.
+		// These are not Param, so we will alias the outputs exactly like the attribute names so it looks
+		// like the expected attributes of the material (base color, normal, ....)
+		TArray<FName> ArgNames = TG_MakeArrayOfArgumentNames(SignatureInit.Arguments);
+		MatAttributesOutputArgNames.Empty();
+
+		Outputs.Empty(AvailableMaterialAttributeNames.Num());
+		for (int32 Index = 0; Index < AvailableMaterialAttributeNames.Num(); ++Index)
+		{
+			FName ArgName = TG_MakeNameUniqueInCollection(AvailableMaterialAttributeNames[Index], ArgNames);
+			ArgNames.Add(ArgName);
+			MatAttributesOutputArgNames.Add(ArgName);
+
+			TMap<FName, FString> MetaDataMap;
+			MetaDataMap.Add(TEXT("Alias"), AvailableMaterialAttributeNames[Index].ToString());
+
+			// New Arg of the signature
+			FTG_Argument Arg{ ArgName, TEXT("FTG_Texture"), {ETG_Access::Out}, MetaDataMap };
+
+			SignatureInit.Arguments.Add(Arg);
+
+			Outputs.Add(FTG_Texture());
+		}
 	}
 	
 	return MakeShared<FTG_Signature>(SignatureInit);
diff --git a/Engine/Plugins/Experimental/TextureGraph/Source/TextureGraph/Private/TG_Graph.cpp b/Engine/Plugins/Experimental/TextureGraph/Source/TextureGraph/Private/TG_Graph.cpp
--- a/Engine/Plugins/Experimental/TextureGraph/Source/TextureGraph/Private/TG_Graph.cpp
+++ b/Engine/Plugins/Experimental/TextureGraph/Source/TextureGraph/Private/TG_Graph.cpp
@@ -315,8 +315,9 @@ FTG_Id UTG_Graph::AllocatePin(UTG_Node* Node, const FTG_Argument& Arg, int32 InP
 
 	// A brand new pin is probably correctly named in the scope of the node
 	// but maybe its alias name needs to be eventually edited if it is a param of the graph
-	NewPin->AliasName = Node->ValidateGeneratePinAliasName(Arg.GetName(), PinId);
-
+	// Alias name can be assigned from signature through metadata and is the same as arg name by default
+	NewPin->AliasName = Node->ValidateGeneratePinAliasName(NewPin->AliasName, PinId);
+	
 	// set up conformant functor to the pin here
 	Node->ValidateGenerateConformer(NewPin);
 
diff --git a/Engine/Plugins/Experimental/TextureGraph/Source/TextureGraph/Private/TG_Hash.cpp b/Engine/Plugins/Experimental/TextureGraph/Source/TextureGraph/Private/TG_Hash.cpp
--- a/Engine/Plugins/Experimental/TextureGraph/Source/TextureGraph/Private/TG_Hash.cpp
+++ b/Engine/Plugins/Experimental/TextureGraph/Source/TextureGraph/Private/TG_Hash.cpp
@@ -18,12 +18,12 @@ DEFINE_LOG_CATEGORY(LogTextureGraph);
 
 // Util function to validate that a name is unique in a collection and make a new candidate name if not
 // It returns a FName equal to the Candidate if unique in the collection or a proposed edited candidate name which is unique
-FName TG_MakeNameUniqueInCollection(FName CandidateName, const TArray<FName>& Collection, int32 RecursionCount)
+FName TG_MakeNameUniqueInCollection(FName CandidateName, const TArray<FName>& Collection, FName RecursionPreFix, int32 RecursionCount)
 {
 	FName Tested = CandidateName;
 	if (RecursionCount > 0)
 	{
-		Tested = FName(CandidateName.ToString() + "_" + FString::FromInt(RecursionCount));
+		Tested = FName(CandidateName.ToString() + RecursionPreFix.ToString() + FString::FromInt(RecursionCount));
 	}
 	auto Found = Collection.Find(Tested);
 	if (Found == INDEX_NONE)
@@ -33,6 +33,6 @@ FName TG_MakeNameUniqueInCollection(FName CandidateName, const TArray<FName>& Co
 	else
 	{
 		++RecursionCount;
-		return 	TG_MakeNameUniqueInCollection(CandidateName, Collection, RecursionCount);
+		return 	TG_MakeNameUniqueInCollection(CandidateName, Collection, RecursionPreFix, RecursionCount);
 	}
 }
\ No newline at end of file
diff --git a/Engine/Plugins/Experimental/TextureGraph/Source/TextureGraph/Private/TG_Node.cpp b/Engine/Plugins/Experimental/TextureGraph/Source/TextureGraph/Private/TG_Node.cpp
--- a/Engine/Plugins/Experimental/TextureGraph/Source/TextureGraph/Private/TG_Node.cpp
+++ b/Engine/Plugins/Experimental/TextureGraph/Source/TextureGraph/Private/TG_Node.cpp
@@ -65,13 +65,13 @@ void UTG_Node::OnPostUndo()
 }
 #endif
 
-TArray<FName> UTG_Node::GetPinAliasNames() const
+TArray<FName> UTG_Node::GetPinAliasNames(uint8 InAccessFilter) const
 {
 	TArray<FName> Names;
 
 	for (auto Pin : Pins)
 	{
-		if (Pin)
+		if (Pin && Pin->Argument.ArgumentType.MatchAccessFilter(InAccessFilter))
 			Names.Add(Pin->GetAliasName());
 	}
 
@@ -82,15 +82,22 @@ TArray<FName> UTG_Node::GetPinAliasNames() const
 FName UTG_Node::ValidateGeneratePinAliasName(FName CandidateName, const FTG_Id& PinId) const
 {
 	// The pin is testing a "new" Alias name and we want to make sure it is unique...
+	const UTG_Pin* TestedPin = Pins[PinId.PinIdx()];
+	uint8 Filter = 0;
+	if (TestedPin->IsOutput())
+	{
+		Filter =AccessToUint8(OutputBitMask);
+	}
 
 	// ... in the scope of the Node
 	// Among all the Node's pins AliasNames
-	TArray<FName> Names = GetPinAliasNames();
-	Names.RemoveAt(PinId.PinIdx()); // Remove the Pin's Alias name we are testing from the collection
+	TArray<FName> Names = GetPinAliasNames(Filter);
+	// Remove the Pin's Alias name we are testing from the collection
+	Names.Remove(CandidateName); 
 
 	// ... in the scope of the Graph if it is a Param
 	// Need to add the param names EXCEPT the one for this
-	if (Pins[PinId.PinIdx()]->IsParam())
+	if (TestedPin->IsParam())
 	{
 		TArray<FName> ParamNames = GetGraph()->GetParamNames();
 		for (auto Param : ParamNames)
diff --git a/Engine/Plugins/Experimental/TextureGraph/Source/TextureGraph/Private/TG_Pin.cpp b/Engine/Plugins/Experimental/TextureGraph/Source/TextureGraph/Private/TG_Pin.cpp
--- a/Engine/Plugins/Experimental/TextureGraph/Source/TextureGraph/Private/TG_Pin.cpp
+++ b/Engine/Plugins/Experimental/TextureGraph/Source/TextureGraph/Private/TG_Pin.cpp
@@ -27,7 +27,16 @@ void UTG_Pin::Construct(FTG_Id InId, const FTG_Argument& InArgument)
 {
 	Id = InId;
 	Argument = InArgument;
-	AliasName = Argument.GetName();
+
+	const FString* DefaultAliasName = InArgument.MetaDataMap.Find(TEXT("Alias"));
+	if (DefaultAliasName)
+	{
+		AliasName = FName(*DefaultAliasName);
+	}
+	else
+	{
+		AliasName = Argument.GetName();
+	}
 
 	InitSelfVar();
 }
diff --git a/Engine/Plugins/Experimental/TextureGraph/Source/TextureGraph/Public/Expressions/TG_Expression_MaterialBase.h b/Engine/Plugins/Experimental/TextureGraph/Source/TextureGraph/Public/Expressions/TG_Expression_MaterialBase.h
--- a/Engine/Plugins/Experimental/TextureGraph/Source/TextureGraph/Public/Expressions/TG_Expression_MaterialBase.h
+++ b/Engine/Plugins/Experimental/TextureGraph/Source/TextureGraph/Public/Expressions/TG_Expression_MaterialBase.h
@@ -53,6 +53,8 @@ public:
 	UPROPERTY(meta = (TGType = "TG_Output") )
 	FTG_Texture Output;
 
+	mutable TArray<FTG_Texture> Outputs;
+
 	virtual void Evaluate(FTG_EvaluationContext* InContext) override;
 
 	// Validate internal checks, warnings and errors
@@ -69,8 +71,10 @@ protected:
 	// All static switches permutations are cached in this map and only created on demand
 	TSharedPtr<FMaterialInstanceStaticSwitchPermutationMap> MaterialPermutations;
 
+	// The list of the material attributes declared in the signature as Output of the expression,
 	TArray<EDrawMaterialAttributeTarget>	AvailableMaterialAttributeIds; // The set of material properties available for rendering
 	TArray<FName>							AvailableMaterialAttributeNames; // same with the attribute names
+	mutable TArray<FName>					MatAttributesOutputArgNames; // Their argument names (because they may have to be different from the true attribute material name)
 
 	// Arg to Material Param Info array records the map of arg name to the corresponding Material Param
 	struct FArgToMaterialParamInfo
@@ -84,7 +88,9 @@ protected:
 
 		bool operator== (const FName& InArgName) const { return ArgName == InArgName; }
 	};
-	mutable TArray<FArgToMaterialParamInfo> ArgToMatParams;			// mutable because it is populated in the BuildSignatureDynamically()
+	mutable TArray<FArgToMaterialParamInfo> ArgToMatParams;					// mutable because it is populated in the BuildSignatureDynamically()
+
+
 
 	virtual void Initialize() override;
 
diff --git a/Engine/Plugins/Experimental/TextureGraph/Source/TextureGraph/Public/TG_Node.h b/Engine/Plugins/Experimental/TextureGraph/Source/TextureGraph/Public/TG_Node.h
--- a/Engine/Plugins/Experimental/TextureGraph/Source/TextureGraph/Public/TG_Node.h
+++ b/Engine/Plugins/Experimental/TextureGraph/Source/TextureGraph/Public/TG_Node.h
@@ -175,7 +175,9 @@ public:
 protected:
 
 	// Return the array of alias name of all the pins
-	TArray<FName> GetPinAliasNames() const;
+	// Can be filtered based on the argument access type bitfield
+	// Default is not filtered, aka AccessFilter is 0
+	TArray<FName> GetPinAliasNames(uint8 AccessFilter = 0) const;
 
 	// Check that a candidate Pin's alias name is unique in the Node's Pins
 	// And in the Graph's Param Pins if it is a Param.
diff --git a/Engine/Plugins/Experimental/TextureGraph/Source/TextureGraph/Public/TG_Signature.h b/Engine/Plugins/Experimental/TextureGraph/Source/TextureGraph/Public/TG_Signature.h
--- a/Engine/Plugins/Experimental/TextureGraph/Source/TextureGraph/Public/TG_Signature.h
+++ b/Engine/Plugins/Experimental/TextureGraph/Source/TextureGraph/Public/TG_Signature.h
@@ -57,6 +57,7 @@ struct TEXTUREGRAPH_API FTG_ArgumentType
 	bool			IsParam() const		{ return ((AccessToUint8(AccessBitMask) & Flags) & AccessToUint8(ParamBitMask)); }
 	bool			IsSetting() const	{ return ((AccessToUint8(AccessBitMask) & Flags) & AccessToUint8(SettingBitMask)); }
 	bool			IsPrivate() const	{ return ((AccessToUint8(AccessBitMask) & Flags) & AccessToUint8(PrivateBitMask)); }
+	bool			MatchAccessFilter(uint8 InFilter) const { return ((AccessToUint8(AccessBitMask) & Flags) & (InFilter == 0 ? AccessToUint8(AccessBitMask) : InFilter));  }
 
 	bool			IsPersistentSelfVar() const { return (Flags & AccessToUint8(PersistentSelfVarFlag)) != 0; }
 	void			SetPersistentSelfVar() { Flags = Flags | AccessToUint8(PersistentSelfVarFlag); }
diff --git a/Engine/Plugins/Experimental/TextureGraph/Source/TextureGraph/Public/TG_SystemTypes.h b/Engine/Plugins/Experimental/TextureGraph/Source/TextureGraph/Public/TG_SystemTypes.h
--- a/Engine/Plugins/Experimental/TextureGraph/Source/TextureGraph/Public/TG_SystemTypes.h
+++ b/Engine/Plugins/Experimental/TextureGraph/Source/TextureGraph/Public/TG_SystemTypes.h
@@ -111,7 +111,7 @@ using FTG_Name = const FName;
 
 // Util function to validate that a name is unique in a collection and make a new candidate name if not
 // It returns a FName equal to the Candidate if unique in the collection or a proposed edited candidate name which is unique
-FName TG_MakeNameUniqueInCollection(FName CandidateName, const TArray<FName>& Collection, int32 RecursionCount = 0);
+FName TG_MakeNameUniqueInCollection(FName CandidateName, const TArray<FName>& Collection, FName RecursionPreFix = TEXT("_"), int32 RecursionCount = 0);
 
 
 using FTG_Hash = uint64;
diff --git a/Engine/Plugins/Experimental/TextureGraph/Source/TextureGraphEngine/FxMat/RenderMaterial_BP.cpp b/Engine/Plugins/Experimental/TextureGraph/Source/TextureGraphEngine/FxMat/RenderMaterial_BP.cpp
--- a/Engine/Plugins/Experimental/TextureGraph/Source/TextureGraphEngine/FxMat/RenderMaterial_BP.cpp
+++ b/Engine/Plugins/Experimental/TextureGraph/Source/TextureGraphEngine/FxMat/RenderMaterial_BP.cpp
@@ -47,7 +47,9 @@ UMaterialInterface* FMaterialInstanceStaticSwitchPermutationMap::GetRootMaterial
 	TObjectPtr<UMaterialInstanceConstant>* MICP = PermutationsMap.Find(DefaultKey);
 	if (MICP)
 	{
-		return Cast<UMaterialInterface>((*MICP)->GetMaterial());
+		// Use the FIRST MIC that initiated the MISWPermutationMap as the "root Material" 
+		// in order to preserve the default values
+		return Cast<UMaterialInterface>((*MICP));
 	}
 	return nullptr;
 }
