Texture Graph Instance support. This CL repurposed the work done for TextureGraphExporter to create a common base for both the exporter UI and the new TG Instance Editor.

#virtualized

diff --git a/Engine/Plugins/Experimental/TextureGraph/Source/TextureGraph/Private/Expressions/Input/TG_Expression_Texture.cpp b/Engine/Plugins/Experimental/TextureGraph/Source/TextureGraph/Private/Expressions/Input/TG_Expression_Texture.cpp
--- a/Engine/Plugins/Experimental/TextureGraph/Source/TextureGraph/Private/Expressions/Input/TG_Expression_Texture.cpp
+++ b/Engine/Plugins/Experimental/TextureGraph/Source/TextureGraph/Private/Expressions/Input/TG_Expression_Texture.cpp
@@ -9,7 +9,8 @@
 void UTG_Expression_Texture::Evaluate(FTG_EvaluationContext* InContext)
 {
 	Super::Evaluate(InContext);
-	if (Texture)
+	
+	if (Texture && Texture.RasterBlob)
 	{
 		Output = Texture;
 	}
@@ -24,7 +25,6 @@ void UTG_Expression_Texture::Evaluate(FTG_EvaluationContext* InContext)
 		// the override value will always be set to false while combining the buffers
 		auto DesiredDesc = Output.GetBufferDescriptor();
 		DesiredDesc.bIsSRGB = Source->SRGB;
-
 		Output = StaticImageResource->GetBlob(InContext->Cycle, &DesiredDesc, 0);
 	}
 	else
@@ -51,6 +51,90 @@ bool UTG_Expression_Texture::Validate(MixUpdateCyclePtr Cycle)
 	
 	return true;
 }
+
+
+void UTG_Expression_Texture::SetSource(UTexture* InSource)
+{
+	bool bHasChanged = Source != InSource;
+	Source = InSource;
+
+	if (bHasChanged)
+	{
+		SetSourceInternal();
+	}
+}
+
+void UTG_Expression_Texture::SetSourceInternal()
+{
+	FString PrevTexturePath = Texture.TexturePath;
+	if (Source)
+	{
+		Texture.TexturePath = Source->GetPathName();
+	}
+	else
+	{
+		Texture.TexturePath = "";
+	}
+	// We need to find its property and trigger property change event manually.
+	const auto SourcePin = GetParentNode()->GetInputPin(GET_MEMBER_NAME_CHECKED(UTG_Expression_Texture, Source));
+	check(SourcePin)
+	if(SourcePin)
+	{
+		FProperty* Property = SourcePin->GetExpressionProperty();
+		NotifyExpressionChanged(FPropertyChangedEvent(Property));
+	}
+	
+	// if TexturePath changed
+	if (PrevTexturePath != Texture.TexturePath)
+	{
+		// We need to find its property and trigger property change event manually.
+		const auto TexturePin = GetParentNode()->GetInputPin(GET_MEMBER_NAME_CHECKED(UTG_Expression_Texture, Texture));
+		check(TexturePin)
+		if(TexturePin)
+		{
+			FProperty* Property = TexturePin->GetExpressionProperty();
+			NotifyExpressionChanged(FPropertyChangedEvent(Property));
+		}
+	}
+}
+
+void UTG_Expression_Texture::SetTexture(const FTG_Texture& InTexture)
+{
+	bool bHasChanged = Texture != InTexture;
+	Texture = InTexture;
+
+	if (bHasChanged)
+	{
+		SetTextureInternal();
+	}
+}
+
+void UTG_Expression_Texture::SetTextureInternal()
+{
+	bool bTextureHasPath = !Texture.TexturePath.IsEmpty();
+	if (bTextureHasPath)
+	{
+		FSoftObjectPath ObjectPath(Texture.TexturePath);
+		SetSource(Cast<UTexture>(ObjectPath.TryLoad()));
+	}
+}
+#if WITH_EDITOR
+void UTG_Expression_Texture::PostEditChangeProperty(FPropertyChangedEvent& PropertyChangedEvent)
+{
+	// if Graph changes catch it first
+	if (PropertyChangedEvent.GetPropertyName() == GET_MEMBER_NAME_CHECKED(UTG_Expression_Texture, Source))
+	{
+		SetSourceInternal();
+	}
+	if (PropertyChangedEvent.GetPropertyName() == GET_MEMBER_NAME_CHECKED(UTG_Expression_Texture, Texture))
+	{
+		SetTextureInternal();
+	}
+
+	Super::PostEditChangeProperty(PropertyChangedEvent);
+}
+#endif
+
 void UTG_Expression_Texture::SetTitleName(FName NewName)
 {
 	GetParentNode()->GetPin(GET_MEMBER_NAME_CHECKED(UTG_Expression_Texture, Texture))->SetAliasName(NewName);
diff --git a/Engine/Plugins/Experimental/TextureGraph/Source/TextureGraph/Private/Expressions/TG_Expression.cpp b/Engine/Plugins/Experimental/TextureGraph/Source/TextureGraph/Private/Expressions/TG_Expression.cpp
--- a/Engine/Plugins/Experimental/TextureGraph/Source/TextureGraph/Private/Expressions/TG_Expression.cpp
+++ b/Engine/Plugins/Experimental/TextureGraph/Source/TextureGraph/Private/Expressions/TG_Expression.cpp
@@ -273,13 +273,13 @@ void UTG_Expression::SetupAndEvaluate(FTG_EvaluationContext* InContext)
 	// Copy input var values TO this expression's matching properties
 	for(auto& InputVar : InContext->Inputs.VarArguments)
 	{
-		InputVar.Value.Var->CopyTo(this, InputVar.Value.Argument);
+		InputVar.Value.Var.CopyTo(this, InputVar.Value.Argument);
 	}
 
 	// Copy Output var values TO this expression's matching properties (Mainly to copy existing thumbs for Output FTG_Textures)
 	for (auto& OutputVar : InContext->Outputs.VarArguments)
 	{
-		OutputVar.Value.Var->CopyTo(this, OutputVar.Value.Argument);
+		OutputVar.Value.Var.CopyTo(this, OutputVar.Value.Argument);
 	}
 
 	// For debug purpose, log the expression evaluation
@@ -294,7 +294,8 @@ void UTG_Expression::SetupAndEvaluate(FTG_EvaluationContext* InContext)
 	// AFTER evaluation, copy output var values FROM this expression's matching properties
 	for (auto& OutputVar : InContext->Outputs.VarArguments)
 	{
-		OutputVar.Value.Var->CopyFrom(this, OutputVar.Value.Argument);
+		FTG_Var* OutputVarInGraph = InContext->Graph->GetVar(OutputVar.Value.Var.GetId());
+		OutputVarInGraph->CopyFrom(this, OutputVar.Value.Argument);
 	}
 }
 
@@ -318,10 +319,10 @@ void UTG_Expression::LogEvaluation(FTG_EvaluationContext* InContext)
 {
 	TArray<FTG_Id> InVarIds;
 	for (auto& VarMapEntry : InContext->Inputs.VarArguments)
-		InVarIds.Emplace(VarMapEntry.Value.Var->GetId());
+		InVarIds.Emplace(VarMapEntry.Value.Var.GetId());
 	TArray<FTG_Id> OutVarIds;
 	for (auto& VarMapEntry : InContext->Outputs.VarArguments)
-		OutVarIds.Emplace(VarMapEntry.Value.Var->GetId());
+		OutVarIds.Emplace(VarMapEntry.Value.Var.GetId());
 
 	UE_LOG(LogTextureGraph, VeryVerbose, TEXT("%-*s %s"), UTG_Graph::LogHeaderWidth, *GetDefaultName().ToString(), *UTG_Graph::LogCall(InVarIds, OutVarIds));
 }
diff --git a/Engine/Plugins/Experimental/TextureGraph/Source/TextureGraph/Private/TG_Graph.cpp b/Engine/Plugins/Experimental/TextureGraph/Source/TextureGraph/Private/TG_Graph.cpp
--- a/Engine/Plugins/Experimental/TextureGraph/Source/TextureGraph/Private/TG_Graph.cpp
+++ b/Engine/Plugins/Experimental/TextureGraph/Source/TextureGraph/Private/TG_Graph.cpp
@@ -861,6 +861,43 @@ int UTG_Graph::GetAllOutputParamValues(TArray<FTG_Variant>& OutVariants, TArray<
 	return NumFounds;
 }
 
+FTG_VarMap UTG_Graph::GetInputParamsVarMap()
+{
+	FTG_VarMap Map;
+	for (auto Param : Params)
+	{
+		UTG_Pin* ParamPin = GetPin(Param.Value);
+		if (ParamPin)
+		{
+			if (ParamPin->IsInput())
+			{
+				Map.VarArguments.Add(ParamPin->GetAliasName(), { *ParamPin->EditSelfVar(), ParamPin->GetArgument() });
+			}
+		}
+	};
+	return Map;
+}
+
+void UTG_Graph::SetInputParamsFromVarMap(FTG_VarMap InMap)
+{
+	auto InputIds = GetInputParamIds();
+	for (auto InputId : InputIds)
+	{
+		UTG_Pin* Pin = Cast<UTG_Pin>(GetPin(InputId));
+		FTG_Var* VarToCopy = InMap.GetVar(Pin->GetAliasName());
+		const FTG_Argument* VarArgument = InMap.GetVarArgument(Pin->GetAliasName());
+		if (VarToCopy)
+		{
+			Pin->EditSelfVar()->CopyFrom(VarToCopy);
+
+			UTG_Expression* OwnerExpression = Pin->GetNodePtr()->GetExpression();
+			// also copy to expression
+			OwnerExpression->CopyVarToExpressionArgument(*VarArgument, Pin->EditSelfVar());
+		}
+	};
+	
+}
+
 FTG_Ids UTG_Graph::GatherSourceNodes(const UTG_Node* InNode) const
 {
 	FTG_Ids SourceNodes; // the array of source nodes that will be return, empty for now
@@ -1072,7 +1109,7 @@ void UTG_Graph::Evaluate(FTG_EvaluationContext* InContext)
 
 FString NewLine = TEXT("\r\n");
 
-void UTG_Graph::Log()
+void UTG_Graph::Log() const
 {
 	FString LogMessage;
 	FString Tab = "   ";
@@ -1092,7 +1129,7 @@ void UTG_Graph::Log()
 	}
 }
 
-FString UTG_Graph::LogNodes(FString InTab)
+FString UTG_Graph::LogNodes(FString InTab) const
 {
 	FString Tab = InTab + "    ";
 	FString PinTab = Tab + "    ";
@@ -1109,7 +1146,7 @@ FString UTG_Graph::LogNodes(FString InTab)
 	return LogMessage;
 }
 
-FString UTG_Graph::LogParams(FString InTab)
+FString UTG_Graph::LogParams(FString InTab) const
 {
 	FString Tab = InTab + "    ";
 	FString ParamTab = InTab + "    " + "    ";
@@ -1134,7 +1171,7 @@ FString UTG_Graph::LogParams(FString InTab)
 	return LogMessage;
 }
 
-FString UTG_Graph::LogVars(FString InTab)
+FString UTG_Graph::LogVars(FString InTab) const
 {
 	FString Tab = InTab + "    ";
 
@@ -1150,7 +1187,7 @@ FString UTG_Graph::LogVars(FString InTab)
 	return LogMessage;
 }
 
-FString UTG_Graph::LogTraversal(FString InTab)
+FString UTG_Graph::LogTraversal(FString InTab) const
 {
 	FString Tab = InTab + "    ";
 
diff --git a/Engine/Plugins/Experimental/TextureGraph/Source/TextureGraph/Private/TG_GraphEvaluation.cpp b/Engine/Plugins/Experimental/TextureGraph/Source/TextureGraph/Private/TG_GraphEvaluation.cpp
--- a/Engine/Plugins/Experimental/TextureGraph/Source/TextureGraph/Private/TG_GraphEvaluation.cpp
+++ b/Engine/Plugins/Experimental/TextureGraph/Source/TextureGraph/Private/TG_GraphEvaluation.cpp
@@ -92,7 +92,7 @@ void FTG_Evaluation::TransferVarToPin(UTG_Pin* InPin, FTG_EvaluationContext* Con
 			}
 		}
 
-		Context->Inputs.VarArguments.Add(Arg.GetName(), { Var, Arg });
+		Context->Inputs.VarArguments.Add(Arg.GetName(), { *Var, Arg });
 	}
 	else if (Arg.IsOutput())
 	{
@@ -100,7 +100,7 @@ void FTG_Evaluation::TransferVarToPin(UTG_Pin* InPin, FTG_EvaluationContext* Con
 		/// This will get copied to anyway when the expression runs. 
 		if (Var->IsArray())
 			Var->Reset();
-		Context->Outputs.VarArguments.Add(Arg.GetName(), { Var, Arg });
+		Context->Outputs.VarArguments.Add(Arg.GetName(), { *Var, Arg });
 	}
 }
 
@@ -584,7 +584,7 @@ void FTG_Evaluation::EvaluateGraph(UTG_Graph* InGraph, FTG_EvaluationContext* In
 			auto TheVar = InGraph->GetVar(ParamId);
 			if (TheVar)
 			{
-				V.Value.Var->CopyTo(TheVar);
+				V.Value.Var.CopyTo(TheVar);
 				//TheVar->ShareData(*V.Value.Var);
 				EvalContext.ConnectedInputParamIds.Add(ParamId);
 			}
@@ -616,7 +616,7 @@ void FTG_Evaluation::EvaluateGraph(UTG_Graph* InGraph, FTG_EvaluationContext* In
 			auto TheVar = InGraph->GetVar(ParamId);
 			if (TheVar)
 			{
-				V.Value.Var->ShareData(*TheVar);
+				V.Value.Var.ShareData(*TheVar);
 			}
 		}
 	}
@@ -696,7 +696,7 @@ void FTG_Evaluation::EvaluateNodeArray(UTG_Node* InNode, const TArray<TObjectPtr
 
 				FTG_Evaluation::VarConverterInfo Info;
 				Info.Context = InContext;
-				Info.InVar = OutputVar.Value.Var;
+				Info.InVar = &OutputVar.Value.Var;
 				Info.OutVar = &VariantVar; // Generate the converted version of the connected input into the CoonvertedVar of the Pin
 				Info.Index = -1;
 
@@ -707,7 +707,7 @@ void FTG_Evaluation::EvaluateNodeArray(UTG_Node* InNode, const TArray<TObjectPtr
 			}
 			else
 			{
-				ArrayOutput->Set(VarIndex, OutputVar.Value.Var->GetAs<FTG_Variant>());
+				ArrayOutput->Set(VarIndex, OutputVar.Value.Var.GetAs<FTG_Variant>());
 			}
 		}
 	}
@@ -717,9 +717,9 @@ void FTG_Evaluation::EvaluateNodeArray(UTG_Node* InNode, const TArray<TObjectPtr
 	for (auto& OutputVar : InContext->Outputs.VarArguments)
 	{
 		FTG_VariantArray* ArrayOutput = ArrayOutputs.Find(OutputVar.Key);
-		OutputVar.Value.Var->ResetAs<FTG_VariantArray>();
-		OutputVar.Value.Var->SetAs<FTG_VariantArray>(*ArrayOutput);
-		OutputVar.Value.Var->SetArray();
+		OutputVar.Value.Var.ResetAs<FTG_VariantArray>();
+		OutputVar.Value.Var.SetAs<FTG_VariantArray>(*ArrayOutput);
+		OutputVar.Value.Var.SetArray();
 	}
 }
 
diff --git a/Engine/Plugins/Experimental/TextureGraph/Source/TextureGraph/Private/TG_HelperFunctions.cpp b/Engine/Plugins/Experimental/TextureGraph/Source/TextureGraph/Private/TG_HelperFunctions.cpp
--- a/Engine/Plugins/Experimental/TextureGraph/Source/TextureGraph/Private/TG_HelperFunctions.cpp
+++ b/Engine/Plugins/Experimental/TextureGraph/Source/TextureGraph/Private/TG_HelperFunctions.cpp
@@ -94,7 +94,7 @@ void FTG_HelperFunctions::EnsureOutputIsTexture(MixUpdateCyclePtr Cycle, UTG_Nod
 	}
 }
 
-JobBatchPtr FTG_HelperFunctions::InitExportBatch(UTextureGraph* InTextureGraph, FString ExportPath, FString AssetName, FExportSettings& TargetExportSettings,
+JobBatchPtr FTG_HelperFunctions::InitExportBatch(UTextureGraphBase* InTextureGraph, FString ExportPath, FString AssetName, FExportSettings& TargetExportSettings,
 	bool OverrideExportPath, bool OverwriteTextures, bool ExportAllOutputs, bool bSave)
 {
 	FString ErrorMessage = "";
@@ -197,7 +197,7 @@ JobBatchPtr FTG_HelperFunctions::InitExportBatch(UTextureGraph* InTextureGraph,
 	return Batch;
 }
 
-AsyncBool FTG_HelperFunctions::ExportAsync(UTextureGraph* InTextureGraph, FString ExportPath, FString AssetName, FExportSettings& TargetExportSettings, bool OverrideExportPath, bool OverwriteTextures /*= true*/, bool ExportAllOutputs /*= false*/,bool bSave /*= true*/)
+AsyncBool FTG_HelperFunctions::ExportAsync(UTextureGraphBase* InTextureGraph, FString ExportPath, FString AssetName, FExportSettings& TargetExportSettings, bool OverrideExportPath, bool OverwriteTextures /*= true*/, bool ExportAllOutputs /*= false*/,bool bSave /*= true*/)
 {
 	JobBatchPtr Batch = InitExportBatch(InTextureGraph, ExportPath, AssetName, TargetExportSettings, OverrideExportPath, OverwriteTextures, ExportAllOutputs, bSave);
 
@@ -227,7 +227,7 @@ AsyncBool FTG_HelperFunctions::ExportAsync(UTextureGraph* InTextureGraph, FStrin
 		});
 }
 
-void FTG_HelperFunctions::InitTargets(UTextureGraph* InTextureGraph)
+void FTG_HelperFunctions::InitTargets(UTextureGraphBase* InTextureGraph)
 {
 	TextureGraphEngine::RegisterErrorReporter(InTextureGraph, std::make_shared<FTextureGraphErrorReporter>());
 	/// Now run the update Cycle
@@ -254,7 +254,7 @@ void FTG_HelperFunctions::InitTargets(UTextureGraph* InTextureGraph)
 	}
 }
 
-JobBatchPtr FTG_HelperFunctions::InitRenderBatch(UTextureGraph* InTextureGraph, JobBatchPtr ExistingBatch /* = nullptr */)
+JobBatchPtr FTG_HelperFunctions::InitRenderBatch(UTextureGraphBase* InTextureGraph, JobBatchPtr ExistingBatch /* = nullptr */)
 {
 	/// Now run the update Cycle
 	JobBatchPtr NewBatch = nullptr;
@@ -278,7 +278,7 @@ JobBatchPtr FTG_HelperFunctions::InitRenderBatch(UTextureGraph* InTextureGraph,
 	return Batch;
 }
 
-AsyncBool FTG_HelperFunctions::RenderAsync(UTextureGraph* InTextureGraph, JobBatchPtr ExistingBatch /* = nullptr */)
+AsyncBool FTG_HelperFunctions::RenderAsync(UTextureGraphBase* InTextureGraph, JobBatchPtr ExistingBatch /* = nullptr */)
 {
 	JobBatchPtr Batch = InitRenderBatch(InTextureGraph, ExistingBatch);
 
diff --git a/Engine/Plugins/Experimental/TextureGraph/Source/TextureGraph/Private/TG_Texture.cpp b/Engine/Plugins/Experimental/TextureGraph/Source/TextureGraph/Private/TG_Texture.cpp
--- a/Engine/Plugins/Experimental/TextureGraph/Source/TextureGraph/Private/TG_Texture.cpp
+++ b/Engine/Plugins/Experimental/TextureGraph/Source/TextureGraph/Private/TG_Texture.cpp
@@ -1,7 +1,10 @@
 // Copyright Epic Games, Inc. All Rights Reserved.
 
 #include "TG_Texture.h"
+
+#include "TG_CustomVersion.h"
 #include "TG_Var.h"
+#include "Model/StaticImageResource.h"
 
 template <> FString TG_Var_LogValue(FTG_Texture& Value)
 {
@@ -10,4 +13,22 @@ template <> FString TG_Var_LogValue(FTG_Texture& Value)
 
 	return FString::Printf(TEXT("FTG_Texture <0x%0*x> %dx%d"), 8, Value.RasterBlob.get(), Value.RasterBlob->GetWidth(), Value.RasterBlob->GetHeight());
 
+}
+
+bool FTG_Texture::operator==(const FTG_Texture& RHS) const
+{
+	return (this->RasterBlob == RHS.RasterBlob &&
+			this->TexturePath == RHS.TexturePath);
+}
+
+void FTG_Texture::ResetTexturePath()
+{
+	TexturePath.Empty();
+}
+
+FArchive& operator<<(FArchive& Ar, FTG_Texture& T)
+{
+	Ar << T.TexturePath;
+	Ar << T.Descriptor;
+	return Ar;
 }
\ No newline at end of file
diff --git a/Engine/Plugins/Experimental/TextureGraph/Source/TextureGraph/Private/TG_Var.cpp b/Engine/Plugins/Experimental/TextureGraph/Source/TextureGraph/Private/TG_Var.cpp
--- a/Engine/Plugins/Experimental/TextureGraph/Source/TextureGraph/Private/TG_Var.cpp
+++ b/Engine/Plugins/Experimental/TextureGraph/Source/TextureGraph/Private/TG_Var.cpp
@@ -4,6 +4,7 @@
 
 #include <functional>
 
+#include "TG_CustomVersion.h"
 #include "TG_Texture.h"
 #include "TG_OutputSettings.h"
 #include "TG_Variant.h"
@@ -227,13 +228,16 @@ void Generic_Struct_Serializer(FTG_Var::VarPropertySerialInfo& Info)
 	{
 		if (!Info.Var->IsEmpty() && Info.Var->IsValid())
 		{
-			*PropertyValue = Info.Var->GetAs<T_ValueType>();
+			T_ValueType* VarValue = &Info.Var->GetAs<T_ValueType>();
+			StructProperty->SetValue_InContainer(Info.Owner, VarValue);
 		}
 		else
 		{
 			T_ValueType* DefaultValue = StructProperty->ContainerPtrToValuePtrForDefaults<T_ValueType>(StructProperty->Struct, Info.Owner, Info.ClampedIndex());
 			if (DefaultValue)
-				*PropertyValue = *DefaultValue;
+			{
+				StructProperty->SetValue_InContainer(Info.Owner, DefaultValue);
+			}
 		}
 	}
 	else
@@ -589,7 +593,7 @@ void FTG_Var::CopyFrom(UTG_Expression* Owner, FTG_Argument& Arg, int Index)
 
 void FTG_Var::CopyTo(FTG_Var* InVar)
 {
-	if (Concept && InVar)
+	if (Concept.IsValid() && InVar)
 	{
 		InVar->Concept = Concept->Clone();
 	}
@@ -625,7 +629,37 @@ void Generic_Simple_ArSerializer(FTG_Var::VarArchiveSerialInfo& Info)
 
 void VarArchiveSerializer_FTG_Texture(FTG_Var::VarArchiveSerialInfo& Info)
 {
-	//Generic_Simple_ArSerializer<FTG_Texture>(Info);
+	Info.Ar.UsingCustomVersion(FTG_CustomVersion::GUID);
+	if (Info.Var->IsEmpty())
+	{
+		Info.Var->ResetAs<FTG_Texture>();
+	}
+	FTG_Texture& Texture = Info.Var->EditAs<FTG_Texture>();
+	
+	if (Info.Ar.IsSaving())
+	{
+		UE_LOG(LogTextureGraph, VeryVerbose, TEXT("        Save Var %s: %s"), *Info.Var->GetId().ToString(), *Info.Var->LogValue());
+		// Convert the TObjectPtr<UTexture> to FSoftObjectPath to save the asset path
+		FSoftObjectPath AssetPath = FSoftObjectPath(Texture.TexturePath);
+		AssetPath.SerializePath(Info.Ar);
+	}
+	
+	if (Info.Ar.IsLoading())
+	{
+		UE_LOG(LogTextureGraph, VeryVerbose, TEXT("        Loaded Var %s: %s"), *Info.Var->GetId().ToString(), *Info.Var->LogValue());
+		int32 Version = Info.Ar.CustomVer(FTG_CustomVersion::GUID);
+		if (Version < FTG_CustomVersion::TGTextureAddedTexturePath)
+		{
+			Texture.TexturePath = FString(); // Default value for older assets
+		}
+		else
+		{
+			// Load the asset path and resolve it back to a UTexture pointer
+			FSoftObjectPath AssetPath;
+			AssetPath.SerializePath(Info.Ar);
+			Texture.TexturePath = AssetPath.GetAssetPathString();
+		}
+	}
 }
 
 void VarArchiveSerializer_FTG_VariantArray(FTG_Var::VarArchiveSerialInfo& Info)
@@ -665,7 +699,46 @@ void VarArchiveSerializer_bool(FTG_Var::VarArchiveSerialInfo& Info)
 {
 	Generic_Simple_ArSerializer<bool>(Info);
 }
+void VarArchiveSerializer_UTexture(FTG_Var::VarArchiveSerialInfo& Info)
+{
+	//Generic_Simple_ArSerializer<TObjectPtr<UTexture>>(Info);
+	if (Info.Var->IsEmpty())
+	{
+		Info.Var->ResetAs<TObjectPtr<UTexture>>();
+	}
+	TObjectPtr<UTexture>& Texture = Info.Var->EditAs<TObjectPtr<UTexture>>();
+	
+	if (Info.Ar.IsSaving())
+	{
+		UE_LOG(LogTextureGraph, VeryVerbose, TEXT("        Save Var %s: %s"), *Info.Var->GetId().ToString(), *Info.Var->LogValue());
+		// Convert the TObjectPtr<UTexture> to FSoftObjectPath to save the asset path
+		FSoftObjectPath AssetPath = Texture != nullptr ? FSoftObjectPath(Texture.Get()) : FSoftObjectPath();
+		AssetPath.SerializePath(Info.Ar);
+	}
+	
+	// Info.Ar << Value;
+	// TSoftObjectPtr SoftValue = Value.Get();
+	// SoftValue->Serialize(Info.Ar);
+	if (Info.Ar.IsLoading())
+	{
+		UE_LOG(LogTextureGraph, VeryVerbose, TEXT("        Loaded Var %s: %s"), *Info.Var->GetId().ToString(), *Info.Var->LogValue());
+		// Load the asset path and resolve it back to a UTexture pointer
+		FSoftObjectPath AssetPath;
+		AssetPath.SerializePath(Info.Ar);
+
+		// Resolve the asset path into a UTexture object (loaded asynchronously or synchronously)
+		if (!AssetPath.IsNull())
+		{
+			// Load the texture synchronously (can be async depending on your needs)
+			Texture = Cast<UTexture>(AssetPath.TryLoad());
+		}
+		else
+		{
+			Texture = nullptr;
+		}
+	}
 
+}
 //////////////////////////////////////////////////////////////////////////////////////////////////////////////////////////////////////
 FTG_Var::VarArchiveSerializerMap FTG_Var::DefaultArchiveSerializers
 (
@@ -675,6 +748,7 @@ FTG_Var::VarArchiveSerializerMap FTG_Var::DefaultArchiveSerializers
 		VAR_ARCHIVE_SERIALIZER_DEF(uint32),
 		VAR_ARCHIVE_SERIALIZER_DEF(float),
 		VAR_ARCHIVE_SERIALIZER_DEF(FTG_Texture),
+		VAR_ARCHIVE_SERIALIZER("TObjectPtr<UTexture>", VarArchiveSerializer_UTexture),
 		VAR_ARCHIVE_SERIALIZER_DEF(FVector4f),
 		VAR_ARCHIVE_SERIALIZER_DEF(FVector2f),
 		VAR_ARCHIVE_SERIALIZER_DEF(FLinearColor),
diff --git a/Engine/Plugins/Experimental/TextureGraph/Source/TextureGraph/Private/TextureGraph.cpp b/Engine/Plugins/Experimental/TextureGraph/Source/TextureGraph/Private/TextureGraph.cpp
--- a/Engine/Plugins/Experimental/TextureGraph/Source/TextureGraph/Private/TextureGraph.cpp
+++ b/Engine/Plugins/Experimental/TextureGraph/Source/TextureGraph/Private/TextureGraph.cpp
@@ -20,6 +20,84 @@
 #include "Transform/Mix/T_InvalidateTiles.h"
 #include "Transform/Mix/T_UpdateTargets.h"
 
+/** ---------------------------------------
+ *	Texture Graph Instance implementation
+ * ---------------------------------------**/
+
+void UTextureGraphInstance::Construct(FString InName)
+{
+	Super::Construct(InName);
+	RuntimeGraph = TStrongObjectPtr<UTG_Graph>(NewObject<UTG_Graph>(GetTransientPackage(), NAME_None,RF_Transactional));
+	RuntimeGraph->Construct(InName);
+	
+	const UTG_Node* OutputNode = RuntimeGraph->CreateExpressionNode(UTG_Expression_Output::StaticClass());
+
+	UTG_Expression_Output* OutputExpression = Cast<UTG_Expression_Output>(OutputNode->GetExpression());
+	OutputExpression->InitializeOutputSettings();
+
+	Settings->GetViewportSettings().InitDefaultSettings(OutputNode->GetNodeName());
+}
+
+void UTextureGraphInstance::CopyParamsToRuntimeGraph()
+{
+	check (RuntimeGraph);
+	RuntimeGraph->SetInputParamsFromVarMap(InputParams);
+}
+
+void UTextureGraphInstance::Update(MixUpdateCyclePtr InCycle)
+{
+	// Graph Evaluate
+	SceneTargetUpdatePtr Target = std::make_shared<MixTargetUpdate>(InCycle->GetMix(), 0);
+	InCycle->AddTarget(Target);
+	
+	T_InvalidateTiles::Create(InCycle, 0);
+
+	// Now Evaluate the Graph!
+	FTG_EvaluationContext EvaluationContext;
+	EvaluationContext.Inputs = InputParams;
+	EvaluationContext.Cycle = InCycle;
+	EvaluationContext.Cycle->PushMix(this);
+	
+	if (Graph()->Validate(InCycle))
+		Graph()->Evaluate(&EvaluationContext);
+	
+	InputParams = Graph()->GetInputParamsVarMap();
+	
+	EvaluationContext.Cycle->PopMix();
+	
+	/// This will be the final result of the rendering
+	T_UpdateTargets::Create(InCycle, 0, true);
+	
+}
+
+void UTextureGraphInstance::SetParent(const TObjectPtr<UTextureGraph>& Parent)
+{
+	this->ParentTextureGraph = Parent;
+	if (Parent == nullptr)
+	{
+		RuntimeGraph = nullptr;
+		Settings = nullptr;
+		InputParams.Empty();
+	}
+	else
+	{
+		RuntimeGraph = TStrongObjectPtr<UTG_Graph>(DuplicateObject(Parent->Graph(), GetTransientPackage()));
+		Settings = DuplicateObject(Parent->GetSettings(), this);
+	}
+}
+
+void UTextureGraphInstance::Initialize()
+{
+	if (ParentTextureGraph)
+	{
+		RuntimeGraph = TStrongObjectPtr<UTG_Graph>(DuplicateObject(ParentTextureGraph->Graph(), GetTransientPackage()));
+		CopyParamsToRuntimeGraph();
+	}
+}
+
+/** ---------------------------------------
+ *	Texture Graph implementation
+ * ---------------------------------------**/
 
 bool UTextureGraph::CheckRecursiveDependency(const UTextureGraph* InTextureGraph) const
 {
@@ -36,10 +114,26 @@ bool UTextureGraph::CheckRecursiveDependency(const UTextureGraph* InTextureGraph
 				(Package == SecondPackage && !bIsTransientPackage);
 	});
 }
+
+void UTextureGraph::Construct(FString InName)
+{
+	Super::Construct(InName);
+	TextureGraph = NewObject<UTG_Graph>(this, NAME_None,RF_Transactional);
+	TextureGraph->Construct(InName);
+
+	const UTG_Node* OutputNode = TextureGraph->CreateExpressionNode(UTG_Expression_Output::StaticClass());
+
+	UTG_Expression_Output* OutputExpression = Cast<UTG_Expression_Output>(OutputNode->GetExpression());
+	OutputExpression->InitializeOutputSettings();
+
+	Settings->GetViewportSettings().InitDefaultSettings(OutputNode->GetNodeName());
+}
+
 bool UTextureGraph::HasCyclicDependency() const
 {
 	return CheckRecursiveDependency(this);
 }
+
 void UTextureGraph::GatherAllDependentGraphs(TArray<UTextureGraph*>& DependentGraphs) const
 {
 	Graph()->ForEachNodes(
@@ -84,27 +178,21 @@ bool UTextureGraph::IsDependentOn(const UTextureGraph* InTextureGraph) const
 	return false;
 }
 
-void UTextureGraph::Construct(FString InName)
+/** ---------------------------------------
+ *	Texture Graph Base implementation
+ * ---------------------------------------**/
+
+void UTextureGraphBase::Construct(FString InName)
 {
 	// On the first new texture script we set the engine in run mode
 	TextureGraphEngine::SetRunEngine();
 
 	Settings = NewObject<UMixSettings>(this);
 	bInvalidateTextures = false;
-	
-	TextureGraph = NewObject<UTG_Graph>(this, NAME_None,RF_Transactional);
-	TextureGraph->Construct(InName);
-
-	const UTG_Node* OutputNode = TextureGraph->CreateExpressionNode(UTG_Expression_Output::StaticClass());
-
-	UTG_Expression_Output* OutputExpression = Cast<UTG_Expression_Output>(OutputNode->GetExpression());
-	OutputExpression->InitializeOutputSettings();
-
-	Settings->GetViewportSettings().InitDefaultSettings(OutputNode->GetNodeName());
 }
 
 
-void UTextureGraph::Serialize(FArchive& Ar)
+void UTextureGraphBase::Serialize(FArchive& Ar)
 {
 	Super::Serialize(Ar);
 	Ar.UsingCustomVersion(FTG_CustomVersion::GUID);
@@ -117,7 +205,7 @@ void UTextureGraph::Serialize(FArchive& Ar)
 		*FString::FromInt(Version));
 }
 
-void UTextureGraph::PostLoad()
+void UTextureGraphBase::PostLoad()
 {
 	// On the first script load we set the engine in run mode as well
 	TextureGraphEngine::SetRunEngine();
@@ -134,10 +222,10 @@ void UTextureGraph::PostLoad()
 	// Fallback to default material.
 	if(!Settings->GetViewportSettings().Material)
 	{
-		FTG_Ids OutputIds = TextureGraph->GetOutputParamIds();
+		FTG_Ids OutputIds = Graph()->GetOutputParamIds();
 		if (!OutputIds.IsEmpty())
 		{
-			Settings->GetViewportSettings().InitDefaultSettings(TextureGraph->GetNode(OutputIds[0])->GetNodeName());
+			Settings->GetViewportSettings().InitDefaultSettings(Graph()->GetNode(OutputIds[0])->GetNodeName());
 		}
 		else // the TG doesn't have any output? that's weird but could happen, let's simply create the material regardless targeting a name
 		{
@@ -146,13 +234,13 @@ void UTextureGraph::PostLoad()
 	}
 }
 
-void UTextureGraph::PreSave(FObjectPreSaveContext SaveContext)
+void UTextureGraphBase::PreSave(FObjectPreSaveContext SaveContext)
 {
 	Super::PreSave(SaveContext);
 	UE_LOG(LogTextureGraph, Verbose, TEXT("PreSave Script: %s"), *GetName());
 }
 
-void UTextureGraph::Update(MixUpdateCyclePtr InCycle)
+void UTextureGraphBase::Update(MixUpdateCyclePtr InCycle)
 {
 	// Graph Evaluate
 	SceneTargetUpdatePtr Target = std::make_shared<MixTargetUpdate>(InCycle->GetMix(), 0);
@@ -186,18 +274,18 @@ void UTextureGraph::Update(MixUpdateCyclePtr InCycle)
 	T_UpdateTargets::Create(InCycle, 0, true);
 }
 
-void UTextureGraph::PostMeshLoad()
+void UTextureGraphBase::PostMeshLoad()
 {
 	FModelInvalidateInfo InvalidateInfo;
 	Invalidate(InvalidateInfo);
 }
 
-void UTextureGraph::FlushInvalidations()
+void UTextureGraphBase::FlushInvalidations()
 {
 	TextureGraphEngine::GetMixManager()->FlushMix(this);
 }
 
-void UTextureGraph::TriggerUpdate(bool Tweaking)
+void UTextureGraphBase::TriggerUpdate(bool Tweaking)
 {
 	FModelInvalidateInfo InvalidateInfo;
 
@@ -208,7 +296,7 @@ void UTextureGraph::TriggerUpdate(bool Tweaking)
 	TextureGraphEngine::GetMixManager()->InvalidateMix(this, InvalidateInfo.Details);
 }
 
-void UTextureGraph::InvalidateAll()
+void UTextureGraphBase::InvalidateAll()
 {
 	FInvalidationDetails Details = FInvalidationDetails().All();
 	Details.Mix = this;
@@ -218,7 +306,7 @@ void UTextureGraph::InvalidateAll()
 	TextureGraphEngine::GetMixManager()->InvalidateMix(this, Details);
 }
 
-void UTextureGraph::UpdateGlobalTGSettings()
+void UTextureGraphBase::UpdateGlobalTGSettings()
 {
 	check(GetSettings());
 	GetSettings()->SetWidth(GetMaxWidth());
@@ -230,7 +318,7 @@ void UTextureGraph::UpdateGlobalTGSettings()
 	GetSettings()->SetTextureFormat(TextureFormat);
 }
 
-EResolution UTextureGraph::GetMaxWidth()
+EResolution UTextureGraphBase::GetMaxWidth()
 {
 	EResolution MaxWidth = EResolution::Auto;
 	Graph()->ForEachOutputSettings( [MaxWidth](const FTG_OutputSettings& OutSettings) mutable
@@ -241,7 +329,7 @@ EResolution UTextureGraph::GetMaxWidth()
 	return MaxWidth;
 }
 
-EResolution UTextureGraph::GetMaxHeight()
+EResolution UTextureGraphBase::GetMaxHeight()
 {
 	EResolution MaxHeight = EResolution::Auto;
 	Graph()->ForEachOutputSettings( [MaxHeight](const FTG_OutputSettings& OutSettings) mutable
@@ -252,7 +340,7 @@ EResolution UTextureGraph::GetMaxHeight()
 	return MaxHeight;
 }
 
-int32 UTextureGraph::GetMaxBufferChannels()
+int32 UTextureGraphBase::GetMaxBufferChannels()
 {
 	uint32 MaxBufferChannels = 0;
 	Graph()->ForEachOutputSettings( [MaxBufferChannels](const FTG_OutputSettings& OutSettings) mutable
@@ -265,7 +353,7 @@ int32 UTextureGraph::GetMaxBufferChannels()
 	return MaxBufferChannels;
 }
 
-BufferFormat UTextureGraph::GetMaxBufferFormat()
+BufferFormat UTextureGraphBase::GetMaxBufferFormat()
 {
 	BufferFormat MaxBufferFormat = BufferFormat::Auto;
 	Graph()->ForEachOutputSettings( [MaxBufferFormat](const FTG_OutputSettings& OutSettings) mutable
@@ -278,7 +366,7 @@ BufferFormat UTextureGraph::GetMaxBufferFormat()
 	return MaxBufferFormat;
 }
 
-void UTextureGraph::Log() const
+void UTextureGraphBase::Log() const
 {
-	TextureGraph->Log();
+	Graph()->Log();
 }
diff --git a/Engine/Plugins/Experimental/TextureGraph/Source/TextureGraph/Public/Expressions/Input/TG_Expression_Texture.h b/Engine/Plugins/Experimental/TextureGraph/Source/TextureGraph/Public/Expressions/Input/TG_Expression_Texture.h
--- a/Engine/Plugins/Experimental/TextureGraph/Source/TextureGraph/Public/Expressions/Input/TG_Expression_Texture.h
+++ b/Engine/Plugins/Experimental/TextureGraph/Source/TextureGraph/Public/Expressions/Input/TG_Expression_Texture.h
@@ -21,19 +21,27 @@ public:
 	virtual void Evaluate(FTG_EvaluationContext* InContext) override;
 
 	virtual bool Validate(MixUpdateCyclePtr	Cycle) override;
-	
+	void SetSourceInternal();
+
 	// The output of the node, which is the loaded texture asset
 	UPROPERTY(meta = (TGType = "TG_Output", PinDisplayName = "", HideInnerPropertiesInNode))
 	FTG_Texture Output;
 
 	// The source asset to be used to generate the Output
-	UPROPERTY(EditAnywhere, BlueprintReadWrite, Category = NoCategory, meta = (TGType = "TG_Input", TGPinNotConnectable) )
+	UPROPERTY(EditAnywhere, Setter, BlueprintReadWrite, Category = NoCategory, meta = (TGType = "TG_Input", TGPinNotConnectable) )
 	TObjectPtr<UTexture> Source;
 
 	// The input texture that was loaded from the asset
-	UPROPERTY(meta = (TGType = "TG_InputParam"))
+	UPROPERTY(EditAnywhere, Setter, Category = NoCategory, meta = (TGType = "TG_InputParam"))
 	FTG_Texture Texture;
 
+	void SetSource(UTexture* InSource);
+	void SetTextureInternal();
+	void SetTexture(const FTG_Texture& InTexture);
+	
+#if WITH_EDITOR
+	void PostEditChangeProperty(FPropertyChangedEvent& PropertyChangedEvent) override;
+#endif
 	class ULayerChannel* Channel;
 	virtual FTG_Name GetDefaultName() const override { return TEXT("Texture"); }
 	virtual void SetTitleName(FName NewName) override;
diff --git a/Engine/Plugins/Experimental/TextureGraph/Source/TextureGraph/Public/TG_CustomVersion.h b/Engine/Plugins/Experimental/TextureGraph/Source/TextureGraph/Public/TG_CustomVersion.h
--- a/Engine/Plugins/Experimental/TextureGraph/Source/TextureGraph/Public/TG_CustomVersion.h
+++ b/Engine/Plugins/Experimental/TextureGraph/Source/TextureGraph/Public/TG_CustomVersion.h
@@ -12,6 +12,9 @@ struct FTG_CustomVersion
 		// Before any version changes were made
 		BeforeCustomVersionWasAdded = 0,
 
+		// Added TexturePath to FTG_Texture
+		TGTextureAddedTexturePath,
+
 		// -----<new versions can be added above this line>-------------------------------------------------
 		VersionPlusOne,
 		LatestVersion = VersionPlusOne - 1
diff --git a/Engine/Plugins/Experimental/TextureGraph/Source/TextureGraph/Public/TG_Graph.h b/Engine/Plugins/Experimental/TextureGraph/Source/TextureGraph/Public/TG_Graph.h
--- a/Engine/Plugins/Experimental/TextureGraph/Source/TextureGraph/Public/TG_Graph.h
+++ b/Engine/Plugins/Experimental/TextureGraph/Source/TextureGraph/Public/TG_Graph.h
@@ -240,8 +240,9 @@ public:
 	// along with their names optionally
 	int GetAllOutputParamValues(TArray<FTG_Variant>& OutVariants, TArray<FName>* OutNames = nullptr) const;
 
-
-	//////////////////////////////////////////////////////////////////////////
+	FTG_VarMap GetInputParamsVarMap();
+    void SetInputParamsFromVarMap(FTG_VarMap InMap);
+    //////////////////////////////////////////////////////////////////////////
 	// Traverseal of the graph and evaluation 
 	//////////////////////////////////////////////////////////////////////////
 
@@ -255,7 +256,7 @@ public:
 	TArray<FTG_Id> GatherAllSourceNodes(const UTG_Node* Node) const;
 
 	//// Graph methods to traverse the graph!
-
+	
 	// Validate internal checks, warnings and errors, Returns true if Graph is valid and has no validation errors
 	bool Validate(MixUpdateCyclePtr	Cycle);
 
@@ -316,13 +317,13 @@ public:
 	//// Logging Utilities
 
 	// Log the graph description to the log output buffer
-	void Log();
+	void Log() const;
 
 	static constexpr int32 LogHeaderWidth = 32; // width in number of chars for any header token during logging
-	FString LogNodes(FString InTab = TEXT(""));
-	FString LogParams(FString InTab = TEXT(""));
-	FString LogTraversal(FString InTab = TEXT(""));
-	FString LogVars(FString InTab = TEXT(""));
+	FString LogNodes(FString InTab = TEXT("")) const;
+	FString LogParams(FString InTab = TEXT("")) const;
+	FString LogTraversal(FString InTab = TEXT("")) const;
+	FString LogVars(FString InTab = TEXT("")) const;
 	static FString LogCall(const TArray<FTG_Id>& PinInputs, const  TArray<FTG_Id>& PinOutputs, int32 InputLogWidth = 10);
 
 	//// Testing zone
diff --git a/Engine/Plugins/Experimental/TextureGraph/Source/TextureGraph/Public/TG_GraphEvaluation.h b/Engine/Plugins/Experimental/TextureGraph/Source/TextureGraph/Public/TG_GraphEvaluation.h
--- a/Engine/Plugins/Experimental/TextureGraph/Source/TextureGraph/Public/TG_GraphEvaluation.h
+++ b/Engine/Plugins/Experimental/TextureGraph/Source/TextureGraph/Public/TG_GraphEvaluation.h
@@ -4,26 +4,57 @@
 
 #include "Model/Mix/MixInterface.h"
 #include "TG_Var.h"
-
+#include "TG_GraphEvaluation.generated.h"
 class UTG_Graph;
 class UTG_Node;
 
-struct TEXTUREGRAPH_API FTG_VarMap
+	
+USTRUCT()
+struct TEXTUREGRAPH_API FVarArgument
+{
+	GENERATED_BODY()
+
+	UPROPERTY()
+	FTG_Var Var;
+
+	UPROPERTY()
+	FTG_Argument Argument;
+
+	// Custom Serialize method for FVarArgument
+	bool Serialize(FArchive& Ar)
+	{
+		FTG_Argument::StaticStruct()->SerializeItem(Ar, ((void*) &Argument), nullptr);
+		// Argument.Serialize(Ar);
+		Argument.SetPersistentSelfVar();
+
+		Var.Serialize(Ar, FTG_Id(), Argument);
+		
+		return true;
+	}
+};
+template<>
+struct TStructOpsTypeTraits<FVarArgument> : public TStructOpsTypeTraitsBase2<FVarArgument>
 {
-	struct FVarArgument
+	enum
 	{
-		FTG_Var* Var;
-		FTG_Argument Argument;
+		WithSerializer = true,  // Enables the use of a custom Serialize method
 	};
+};
+USTRUCT()
+struct TEXTUREGRAPH_API FTG_VarMap
+{
+	GENERATED_BODY()
+
+	UPROPERTY()
 	TMap<FName, FVarArgument> VarArguments;
 
-	FTG_Var* GetVar(FName Name) const
+	FTG_Var* GetVar(FName Name)
 	{
 		auto Element = VarArguments.Find(Name);
 		
 		if (Element)
 		{
-			return Element->Var;
+			return &Element->Var;
 		}
 
 		return nullptr;
diff --git a/Engine/Plugins/Experimental/TextureGraph/Source/TextureGraph/Public/TG_HelperFunctions.h b/Engine/Plugins/Experimental/TextureGraph/Source/TextureGraph/Public/TG_HelperFunctions.h
--- a/Engine/Plugins/Experimental/TextureGraph/Source/TextureGraph/Public/TG_HelperFunctions.h
+++ b/Engine/Plugins/Experimental/TextureGraph/Source/TextureGraph/Public/TG_HelperFunctions.h
@@ -10,22 +10,22 @@
 #include "TG_Node.h"
 #include "2D/TextureHelper.h"
 
-class UTextureGraph;
+class UTextureGraphBase;
 
 class TEXTUREGRAPH_API FTG_HelperFunctions
 {
 
 public:
-	static void	InitTargets(UTextureGraph* InTextureGraph);
+	static void	InitTargets(UTextureGraphBase* InTextureGraph);
 	static TArray<BlobPtr> GetTexturedOutputs(const UTG_Node* Node, FTG_EvaluationContext* TextureConversionContext = nullptr);
 
 	static void EnsureOutputIsTexture(MixUpdateCyclePtr Cycle, UTG_Node* OutputNode);
 
-	static JobBatchPtr InitExportBatch(UTextureGraph* InTextureGraph, FString ExportPath, FString AssetName, FExportSettings& TargetExportSettings, bool OverrideExportPath, bool OverwriteTextures, bool ExportAllOutputs, bool bSave);
-	static AsyncBool ExportAsync(UTextureGraph* InTextureGraph, FString ExportPath, FString AssetName, FExportSettings& TargetExportSettings, bool OverrideExportPath, bool OverwriteTextures = true, bool ExportAllOutputs = false, bool bSave = true);
+	static JobBatchPtr InitExportBatch(UTextureGraphBase* InTextureGraph, FString ExportPath, FString AssetName, FExportSettings& TargetExportSettings, bool OverrideExportPath, bool OverwriteTextures, bool ExportAllOutputs, bool bSave);
+	static AsyncBool ExportAsync(UTextureGraphBase* InTextureGraph, FString ExportPath, FString AssetName, FExportSettings& TargetExportSettings, bool OverrideExportPath, bool OverwriteTextures = true, bool ExportAllOutputs = false, bool bSave = true);
 
-	static JobBatchPtr InitRenderBatch(UTextureGraph* InTextureGraph, JobBatchPtr ExistingBatch = nullptr);
-	static AsyncBool RenderAsync(UTextureGraph* InTextureGraph, JobBatchPtr ExistingBatch = nullptr);
+	static JobBatchPtr InitRenderBatch(UTextureGraphBase* InTextureGraph, JobBatchPtr ExistingBatch = nullptr);
+	static AsyncBool RenderAsync(UTextureGraphBase* InTextureGraph, JobBatchPtr ExistingBatch = nullptr);
 	template <typename T_Type>
 	static TArray<T_Type> GetOutputsOfType(const UTG_Node* Node)
 	{
diff --git a/Engine/Plugins/Experimental/TextureGraph/Source/TextureGraph/Public/TG_Texture.h b/Engine/Plugins/Experimental/TextureGraph/Source/TextureGraph/Public/TG_Texture.h
--- a/Engine/Plugins/Experimental/TextureGraph/Source/TextureGraph/Public/TG_Texture.h
+++ b/Engine/Plugins/Experimental/TextureGraph/Source/TextureGraph/Public/TG_Texture.h
@@ -77,23 +77,25 @@ struct TEXTUREGRAPH_API FTG_Texture
 public:
 	TiledBlobPtr RasterBlob = nullptr;
 	
+	UPROPERTY(EditAnywhere, Category = "TextureDescriptor")
+	FString TexturePath;
+	
 	UPROPERTY(EditAnywhere, Category = "TextureDescriptor")
 	FTG_TextureDescriptor Descriptor;
 
 	FTG_Texture() = default;
-	FTG_Texture(TiledBlobPtr RHS) : RasterBlob(RHS) {}
+	FTG_Texture(TiledBlobPtr RHS) : RasterBlob(RHS) { }
 	
 	operator TiledBlobPtr() { return RasterBlob; }
 	operator TiledBlobRef() { return RasterBlob; }
 	operator bool() const { return !!RasterBlob; }
-	friend FArchive& operator<<(FArchive& Ar, FTG_Texture& T)
-	{
-		Ar << T.Descriptor;
-		return Ar;
-	}
-	FTG_Texture& operator = (TiledBlobRef RHS) { RasterBlob = RHS; return *this; }
-	FTG_Texture& operator = (TiledBlobPtr RHS) { RasterBlob = RHS; return *this; }
+	friend FArchive& operator<<(FArchive& Ar, FTG_Texture& T);
+	FTG_Texture& operator = (TiledBlobRef RHS) { RasterBlob = RHS; ResetTexturePath(); return *this; }
+	FTG_Texture& operator = (TiledBlobPtr RHS) { RasterBlob = RHS; ResetTexturePath(); return *this; }
+	bool operator == (const FTG_Texture& RHS) const;
 	FORCEINLINE TiledBlob* operator -> () const { return RasterBlob.get(); }
+	void ResetTexturePath();
+	
 	//////////////////////////////////////////////////////////////////////////
 	/// Static functions
 	//////////////////////////////////////////////////////////////////////////
diff --git a/Engine/Plugins/Experimental/TextureGraph/Source/TextureGraph/Public/TextureGraph.h b/Engine/Plugins/Experimental/TextureGraph/Source/TextureGraph/Public/TextureGraph.h
--- a/Engine/Plugins/Experimental/TextureGraph/Source/TextureGraph/Public/TextureGraph.h
+++ b/Engine/Plugins/Experimental/TextureGraph/Source/TextureGraph/Public/TextureGraph.h
@@ -4,32 +4,27 @@
 
 #include "Model/Mix/MixInterface.h"
 #include "TG_OutputSettings.h"
+#include "TG_GraphEvaluation.h"
 #include "TextureGraph.generated.h"
 
 class UTG_Graph;
 using ErrorReportMap = TMap<int32, TArray<FTextureGraphErrorReport>>;
 
-UCLASS(BlueprintType)
-class TEXTUREGRAPH_API UTextureGraph : public UMixInterface
+UCLASS(abstract, BlueprintType, MinimalAPI)
+class UTextureGraphBase : public UMixInterface
 {
 	GENERATED_BODY()
 
 protected:
-	UPROPERTY(EditAnywhere, Category="TG_GraphParams")
-	TObjectPtr<class UTG_Graph> TextureGraph;
-
-	void GatherAllDependentGraphs(TArray<UTextureGraph*>& DependentGraphs) const;
 	
-	bool CheckRecursiveDependency(const UTextureGraph* InTextureGraph) const;
-	
-public:
+	EResolution GetMaxWidth();
 
-	bool IsDependentOn(const UTextureGraph* TextureGraph) const;
-	bool HasCyclicDependency() const;
-	
-	// Construct the script giving it its name Initialize to a default one output script
-	virtual void Construct(FString Name);
+	EResolution GetMaxHeight();
+
+	int32 GetMaxBufferChannels();
 
+	BufferFormat GetMaxBufferFormat();
+	
 	// Override Serialize method of UObject
 	virtual void Serialize(FArchive& Ar) override;
 
@@ -38,27 +33,71 @@ public:
 
 	// Override PreSave method of UObject
 	virtual void PreSave(FObjectPreSaveContext SaveContext) override;
+	void PostMeshLoad() override;
 
-	virtual const UTG_Graph* Graph() const { return TextureGraph; }
-	virtual UTG_Graph* Graph() { return TextureGraph; }
-
+public:
+	virtual const UTG_Graph* Graph() const { return nullptr; }
+	virtual UTG_Graph* Graph() { return nullptr; }
+	
+	// Construct the script giving it its name Initialize to a default one output script
+	virtual void Construct(FString Name);
+	
 	void InvalidateAll() override;
 	void Update(MixUpdateCyclePtr InCycle) override;
 
-	void PostMeshLoad() override;
 
-	void TriggerUpdate(bool Tweaking);
-	void FlushInvalidations();
-
-	void UpdateGlobalTGSettings();
+	TEXTUREGRAPH_API void TriggerUpdate(bool Tweaking);
+	TEXTUREGRAPH_API void FlushInvalidations();
+	TEXTUREGRAPH_API void UpdateGlobalTGSettings();
+	TEXTUREGRAPH_API void Log() const;
+};
 
-	EResolution GetMaxWidth();
+UCLASS(BlueprintType)
+class TEXTUREGRAPH_API UTextureGraph: public UTextureGraphBase
+{
+	GENERATED_BODY()
 
-	EResolution GetMaxHeight();
+public:
+	UPROPERTY()
+	TObjectPtr<class UTG_Graph> TextureGraph;
+	
+protected:
+	void GatherAllDependentGraphs(TArray<UTextureGraph*>& DependentGraphs) const;
+	
+	bool CheckRecursiveDependency(const UTextureGraph* InTextureGraph) const;
+	
+public:
+	virtual void Construct(FString Name) override;
+	virtual const UTG_Graph* Graph() const { return TextureGraph; }
+	virtual UTG_Graph* Graph() { return TextureGraph; }
+	bool IsDependentOn(const UTextureGraph* TextureGraph) const;
+	bool HasCyclicDependency() const;
+};
 
-	int32 GetMaxBufferChannels();
+UCLASS(BlueprintType)
+class TEXTUREGRAPH_API UTextureGraphInstance: public UTextureGraphBase
+{
+	friend struct FTG_InstanceImpl;
+	GENERATED_BODY()
+	
+private:
+	TStrongObjectPtr<UTG_Graph>	RuntimeGraph;		// The runtime graph instance to use for the sub-graph
+public:
+	
+	UPROPERTY(EditAnywhere, Category=TextureGraphInstance)
+	FTG_VarMap InputParams;
+	
+	/** Parent Texture Graph. */
+	UPROPERTY(EditAnywhere, BlueprintReadOnly, Category=TextureGraphInstance, AssetRegistrySearchable)
+	TObjectPtr<UTextureGraph> ParentTextureGraph;
 
-	BufferFormat GetMaxBufferFormat();
+	virtual void Construct(FString Name) override;
 
-	void Log() const;
-};
+	void CopyParamsToRuntimeGraph();
+	void Update(MixUpdateCyclePtr InCycle) override;
+	void SetParent(const TObjectPtr<UTextureGraph>& Parent);
+	
+	virtual const UTG_Graph* Graph() const { return RuntimeGraph.Get(); }
+	virtual UTG_Graph* Graph() { return RuntimeGraph.Get(); }
+	void Initialize();
+};
\ No newline at end of file
diff --git a/Engine/Plugins/Experimental/TextureGraph/Source/TextureGraphEditor/Private/AssetDefinition_TextureGraph.cpp b/Engine/Plugins/Experimental/TextureGraph/Source/TextureGraphEditor/Private/AssetDefinition_TextureGraph.cpp
--- a/Engine/Plugins/Experimental/TextureGraph/Source/TextureGraphEditor/Private/AssetDefinition_TextureGraph.cpp
+++ b/Engine/Plugins/Experimental/TextureGraph/Source/TextureGraphEditor/Private/AssetDefinition_TextureGraph.cpp
@@ -10,6 +10,7 @@
 #include "ToolMenus.h"
 #include "IAssetTools.h"
 #include "ToolMenuSection.h"
+#include "TG_InstanceFactory.h"
 #include "AssetRegistry/AssetRegistryHelpers.h"
 
 #define LOCTEXT_NAMESPACE "UAssetDefinition_TextureGraph"
@@ -30,6 +31,7 @@ FAssetOpenSupport UAssetDefinition_TextureGraph::GetAssetOpenSupport(const FAsse
 	return FAssetOpenSupport(OpenSupportArgs.OpenMethod,OpenSupportArgs.OpenMethod == EAssetOpenMethod::Edit || OpenSupportArgs.OpenMethod == EAssetOpenMethod::View); 
 }
 
+
 // Menu Extensions
 //--------------------------------------------------------------------
 
@@ -50,6 +52,22 @@ namespace MenuExtension_TextureGraph
 			
 		}
 	}
+	static void ExecuteNewTextureGraphInstance(const FToolMenuContext& InContext)
+	{
+		if (const UContentBrowserAssetContextMenuContext* CBContext = UContentBrowserAssetContextMenuContext::FindContextWithAssets(InContext))
+		{
+			
+			IAssetTools::Get().CreateAssetsFrom<UTextureGraph>(
+				CBContext->LoadSelectedObjects<UTextureGraph>(), UTextureGraphInstance::StaticClass(), TEXT("_Inst"), [](UTextureGraph* SourceObject)
+				{
+					UTG_InstanceFactory* Factory = NewObject<UTG_InstanceFactory>();
+					Factory->InitialParent = SourceObject;
+					return Factory;
+				}
+			);
+			
+		}
+	}
 
 	static FDelayedAutoRegisterHelper DelayedAutoRegister(EDelayedRegisterRunPhase::EndOfEngineInit, []{ 
 		UToolMenus::RegisterStartupCallback(FSimpleMulticastDelegate::FDelegate::CreateLambda([]()
@@ -73,6 +91,20 @@ namespace MenuExtension_TextureGraph
 						}
 					}
 				}));
+
+				Section.AddDynamicEntry(NAME_None, FNewToolMenuSectionDelegate::CreateLambda([](FToolMenuSection& InSection)
+				{
+					if (const UContentBrowserAssetContextMenuContext* Context = UContentBrowserAssetContextMenuContext::FindContextWithAssets(InSection))
+					{
+						{
+							const TAttribute<FText> Label = LOCTEXT("TextureGraph_NewInstance", "Create Texture Graph Instance");
+							const TAttribute<FText> ToolTip = LOCTEXT("TextureGraph_NewInstanceTooltip", "Creates a parameterized Texture Graph Instance using this Texture Graph as a base");
+							const FSlateIcon Icon = FSlateIcon();
+							const FToolMenuExecuteAction UIAction = FToolMenuExecuteAction::CreateStatic(&ExecuteNewTextureGraphInstance);
+							InSection.AddMenuEntry("TextureGraph_NewInstance", Label, ToolTip, Icon, UIAction);
+						}
+					}
+				}));
 			}
 
 			
diff --git a/Engine/Plugins/Experimental/TextureGraph/Source/TextureGraphEditor/Private/AssetDefinition_TextureGraphInstance.cpp b/Engine/Plugins/Experimental/TextureGraph/Source/TextureGraphEditor/Private/AssetDefinition_TextureGraphInstance.cpp
new file mode 100644
--- /dev/null
+++ b/Engine/Plugins/Experimental/TextureGraph/Source/TextureGraphEditor/Private/AssetDefinition_TextureGraphInstance.cpp
@@ -0,0 +1,83 @@
+// Copyright Epic Games, Inc. All Rights Reserved.
+
+#include "AssetDefinition_TextureGraphInstance.h"
+#include "TextureGraphEditorModule.h"
+#include "TextureGraph.h"
+
+#include "AssetRegistry/AssetIdentifier.h"
+#include "ContentBrowserMenuContexts.h"
+#include "ToolMenu.h"
+#include "ToolMenus.h"
+#include "IAssetTools.h"
+#include "ToolMenuSection.h"
+#include "AssetRegistry/AssetRegistryHelpers.h"
+
+#define LOCTEXT_NAMESPACE "UAssetDefinition_TextureGraphInstance"
+
+EAssetCommandResult UAssetDefinition_TextureGraphInstance::OpenAssets(const FAssetOpenArgs& OpenArgs) const
+{
+	for (UTextureGraphInstance* TextureGraphInstance : OpenArgs.LoadObjects<UTextureGraphInstance>())
+	{
+		FTextureGraphEditorModule* TextureEditorModule = &FModuleManager::LoadModuleChecked<FTextureGraphEditorModule>("TextureGraphEditor");
+		TextureEditorModule->CreateTextureGraphInstanceEditor(OpenArgs.GetToolkitMode(), OpenArgs.ToolkitHost, TextureGraphInstance);
+	}
+	
+	return EAssetCommandResult::Handled;
+}
+
+FAssetOpenSupport UAssetDefinition_TextureGraphInstance::GetAssetOpenSupport(const FAssetOpenSupportArgs& OpenSupportArgs) const
+{
+	return FAssetOpenSupport(OpenSupportArgs.OpenMethod,OpenSupportArgs.OpenMethod == EAssetOpenMethod::Edit || OpenSupportArgs.OpenMethod == EAssetOpenMethod::View); 
+}
+
+// Menu Extensions
+//--------------------------------------------------------------------
+
+namespace MenuExtension_TextureGraphInstance
+{
+
+	static void ExecuteExportTextureGraphInstance(const FToolMenuContext& InContext)
+	{
+		if (const UContentBrowserAssetContextMenuContext* CBContext = UContentBrowserAssetContextMenuContext::FindContextWithAssets(InContext))
+		{
+			const TArray<UTextureGraphInstance*> &TextureGraphs = CBContext->LoadSelectedObjects<UTextureGraphInstance>();
+			const FTextureGraphEditorModule& TextureGraphEditorModule = FModuleManager::Get().LoadModuleChecked<FTextureGraphEditorModule>(TEXT("TextureGraphEditor"));
+			
+			for (auto TextureGraphIt = TextureGraphs.CreateConstIterator(); TextureGraphIt; ++TextureGraphIt)
+			{
+				TextureGraphEditorModule.GetTextureExporter()->SetTextureGraphToExport((*TextureGraphIt));
+			}
+			
+		}
+	}
+
+	static FDelayedAutoRegisterHelper DelayedAutoRegister(EDelayedRegisterRunPhase::EndOfEngineInit, []{ 
+		UToolMenus::RegisterStartupCallback(FSimpleMulticastDelegate::FDelegate::CreateLambda([]()
+		{
+			FToolMenuOwnerScoped OwnerScoped(UE_MODULE_NAME);
+
+			{
+				UToolMenu* Menu = UE::ContentBrowser::ExtendToolMenu_AssetContextMenu(UTextureGraphInstance::StaticClass());
+		
+				FToolMenuSection& Section = Menu->FindOrAddSection("GetAssetActions");
+				Section.AddDynamicEntry(NAME_None, FNewToolMenuSectionDelegate::CreateLambda([](FToolMenuSection& InSection)
+				{
+					if (const UContentBrowserAssetContextMenuContext* Context = UContentBrowserAssetContextMenuContext::FindContextWithAssets(InSection))
+					{
+						{
+							const TAttribute<FText> Label = LOCTEXT("TextureGraphInstance_Export", "Export Texture Graph Instance");
+							const TAttribute<FText> ToolTip = LOCTEXT("Texture_ExportTextureGraphInstanceTooltip", "Allows Exporting Texture Graph Instance.");
+							const FSlateIcon Icon = FSlateIcon();
+							const FToolMenuExecuteAction UIAction = FToolMenuExecuteAction::CreateStatic(&ExecuteExportTextureGraphInstance);
+							InSection.AddMenuEntry("Texture_Export", Label, ToolTip, Icon, UIAction);
+						}
+					}
+				}));
+			}
+
+			
+		}));
+	});
+}
+
+#undef LOCTEXT_NAMESPACE
diff --git a/Engine/Plugins/Experimental/TextureGraph/Source/TextureGraphEditor/Private/AssetDefinition_TextureGraphInstance.h b/Engine/Plugins/Experimental/TextureGraph/Source/TextureGraphEditor/Private/AssetDefinition_TextureGraphInstance.h
new file mode 100644
--- /dev/null
+++ b/Engine/Plugins/Experimental/TextureGraph/Source/TextureGraphEditor/Private/AssetDefinition_TextureGraphInstance.h
@@ -0,0 +1,29 @@
+// Copyright Epic Games, Inc. All Rights Reserved.
+
+#pragma once
+
+#include "TextureGraph.h"
+#include "AssetDefinitionDefault.h"
+
+#include "AssetDefinition_TextureGraphInstance.generated.h"
+
+UCLASS()
+class UAssetDefinition_TextureGraphInstance : public UAssetDefinitionDefault
+{
+	GENERATED_BODY()
+
+public:
+	// UAssetDefinition Begin
+	virtual FText GetAssetDisplayName() const override { return NSLOCTEXT("AssetTypeActions", "AssetTypeActions_TextureGraphInstance", "Texture Graph Instance"); }
+	virtual FLinearColor GetAssetColor() const override { return FLinearColor(FColor::Emerald); }
+	virtual TSoftClassPtr<UObject> GetAssetClass() const override { return UTextureGraphInstance::StaticClass(); }
+	virtual TConstArrayView<FAssetCategoryPath> GetAssetCategories() const override
+	{
+		static const auto Categories = { EAssetCategoryPaths::Texture };
+		return Categories;
+	}
+
+	virtual FAssetOpenSupport GetAssetOpenSupport(const FAssetOpenSupportArgs& OpenSupportArgs) const override;
+	virtual EAssetCommandResult OpenAssets(const FAssetOpenArgs& OpenArgs) const override;
+	// UAssetDefinition End
+};
diff --git a/Engine/Plugins/Experimental/TextureGraph/Source/TextureGraphEditor/Private/Customizations/TG_MaterialMappingInfoCustomization.h b/Engine/Plugins/Experimental/TextureGraph/Source/TextureGraphEditor/Private/Customizations/TG_MaterialMappingInfoCustomization.h
--- a/Engine/Plugins/Experimental/TextureGraph/Source/TextureGraphEditor/Private/Customizations/TG_MaterialMappingInfoCustomization.h
+++ b/Engine/Plugins/Experimental/TextureGraph/Source/TextureGraphEditor/Private/Customizations/TG_MaterialMappingInfoCustomization.h
@@ -46,7 +46,7 @@ public:
 	{
 		check(MixSettings)
 
-		UTextureGraph* TextureGraph = Cast<UTextureGraph>(MixSettings->Mix());
+		UTextureGraphBase* TextureGraph = Cast<UTextureGraphBase>(MixSettings->Mix());
 
 		OutComboBoxStrings.Add(MakeShared<FString>("None"));		
 		TextureGraph->Graph()->ForEachNodes(
diff --git a/Engine/Plugins/Experimental/TextureGraph/Source/TextureGraphEditor/Private/Customizations/TG_OutputSettingsCustomization.h b/Engine/Plugins/Experimental/TextureGraph/Source/TextureGraphEditor/Private/Customizations/TG_OutputSettingsCustomization.h
--- a/Engine/Plugins/Experimental/TextureGraph/Source/TextureGraphEditor/Private/Customizations/TG_OutputSettingsCustomization.h
+++ b/Engine/Plugins/Experimental/TextureGraph/Source/TextureGraphEditor/Private/Customizations/TG_OutputSettingsCustomization.h
@@ -239,7 +239,7 @@ public:
 
 		PropertyHandle->GetOuterObjects(ExternalObjects);
 		UObject* OutermostObject = ExternalObjects[0]->GetOutermostObject();
-		UTextureGraph* TextureGraph = CastChecked<UTextureGraph>(OutermostObject);
+		UTextureGraphBase* TextureGraph = CastChecked<UTextureGraphBase>(OutermostObject);
 
 		FTG_Id OutputNodeId = OutputExpressionInfo->OutputNodeID;
 
diff --git a/Engine/Plugins/Experimental/TextureGraph/Source/TextureGraphEditor/Private/Customizations/TG_ParameterCustomization.h b/Engine/Plugins/Experimental/TextureGraph/Source/TextureGraphEditor/Private/Customizations/TG_ParameterCustomization.h
--- a/Engine/Plugins/Experimental/TextureGraph/Source/TextureGraphEditor/Private/Customizations/TG_ParameterCustomization.h
+++ b/Engine/Plugins/Experimental/TextureGraph/Source/TextureGraphEditor/Private/Customizations/TG_ParameterCustomization.h
@@ -64,7 +64,7 @@ public:
 					TArray<UObject*> ExternalObjects;
 					ExternalObjects.Add(ExpressionPtr);
 
-					FName PinName;
+					FName PinName = PinPtr->GetArgumentName();
 
 					// In case of texture expression
 					// Param is TG_Texture but it cannot be modified here
@@ -72,11 +72,16 @@ public:
 					UTG_Expression_Texture* TextureExpression = Cast<UTG_Expression_Texture>(ExpressionPtr);
 					if(TextureExpression)
 					{
-						PinName = GET_MEMBER_NAME_CHECKED(UTG_Expression_Texture, Source);
-					}
-					else
-					{
-						PinName = PinPtr->GetArgumentName();
+						if (PinName == GET_MEMBER_NAME_CHECKED(UTG_Expression_Texture, Source))
+						{
+							// skip the Source pin to be shown directly
+							return;
+						}
+						else if (PinName == GET_MEMBER_NAME_CHECKED(UTG_Expression_Texture, Texture))
+						{
+							// allow the Texture pin instead with the correct parameter alias name but internally use the Source pin.
+							PinName = GET_MEMBER_NAME_CHECKED(UTG_Expression_Texture, Source);
+						}
 					}
 					
 					IDetailPropertyRow* ParameterDetailsPropertyRow = ChildBuilder.AddExternalObjectProperty(ExternalObjects, PinName, FAddPropertyParams().HideRootObjectNode(true));
diff --git a/Engine/Plugins/Experimental/TextureGraph/Source/TextureGraphEditor/Private/STG_EditorViewport.cpp b/Engine/Plugins/Experimental/TextureGraph/Source/TextureGraphEditor/Private/STG_EditorViewport.cpp
--- a/Engine/Plugins/Experimental/TextureGraph/Source/TextureGraphEditor/Private/STG_EditorViewport.cpp
+++ b/Engine/Plugins/Experimental/TextureGraph/Source/TextureGraphEditor/Private/STG_EditorViewport.cpp
@@ -46,7 +46,7 @@
 class FTG_EditorViewportClient : public FEditorViewportClient
 {
 public:
-	FTG_EditorViewportClient(TObjectPtr<UTextureGraph> InTextureGraphPtr, FAdvancedPreviewScene& InPreviewScene, const TSharedRef<STG_EditorViewport>& InTG_AssetEditorViewport);
+	FTG_EditorViewportClient(TObjectPtr<UTextureGraphBase> InTextureGraphPtr, FAdvancedPreviewScene& InPreviewScene, const TSharedRef<STG_EditorViewport>& InTG_AssetEditorViewport);
 
 	// FEditorViewportClient interface
 	virtual bool InputKey(const FInputKeyEventArgs& EventArgs) override;
@@ -69,14 +69,14 @@ public:
 private:
 
 	/** Pointer to the Edited Texture Graph */
-	TObjectPtr<UTextureGraph> TextureGraphPtr;
+	TObjectPtr<UTextureGraphBase> TextureGraphPtr;
 
 	/** Preview Scene - uses advanced preview settings */
 	class FAdvancedPreviewScene* AdvancedPreviewScene;
 };
 
 
-FTG_EditorViewportClient::FTG_EditorViewportClient(TObjectPtr<UTextureGraph> InTextureGraphPtr, FAdvancedPreviewScene& InPreviewScene, const TSharedRef<STG_EditorViewport>& InTG_EditorViewport)
+FTG_EditorViewportClient::FTG_EditorViewportClient(TObjectPtr<UTextureGraphBase> InTextureGraphPtr, FAdvancedPreviewScene& InPreviewScene, const TSharedRef<STG_EditorViewport>& InTG_EditorViewport)
 	: FEditorViewportClient(nullptr, &InPreviewScene, StaticCastSharedRef<SEditorViewport>(InTG_EditorViewport))
 	, TextureGraphPtr(InTextureGraphPtr)
 {
@@ -250,12 +250,6 @@ void STG_EditorViewport::Construct(const FArguments& InArgs)
 	SEditorViewport::Construct(SEditorViewport::FArguments());
 
 	PreviewMeshComponent = nullptr;
-	
-	ERHIFeatureLevel::Type FeatureLevel = GEditor->PreviewPlatform.GetEffectivePreviewFeatureLevel();
-	if (FeatureLevel <= ERHIFeatureLevel::ES3_1)
-	{
-		PreviewMeshComponent->SetMobility(EComponentMobility::Static);
-	}
 
 	UAssetViewerSettings* Settings = UAssetViewerSettings::Get();
 	const int32 ProfileIndex = PreviewScene->GetCurrentProfileIndex();
@@ -269,6 +263,12 @@ void STG_EditorViewport::Construct(const FArguments& InArgs)
 	GenerateRendermodeToolbar();
 }
 
+void STG_EditorViewport::SetTextureGraph(const TObjectPtr<UTextureGraphBase>& InTextureGraph)
+{
+	TextureGraphPtr = InTextureGraph;
+	GenerateRendermodeToolbar();
+}
+
 STG_EditorViewport::STG_EditorViewport()
 	: PreviewScene(MakeShareable(new FAdvancedPreviewScene(FPreviewScene::ConstructionValues())))
 {
@@ -443,7 +443,7 @@ void STG_EditorViewport::SetRenderMode(FName InRenderModeName)
 #endif
 }
 
-void STG_EditorViewport::InitRenderModes(UTextureGraph* InTextureGraph)
+void STG_EditorViewport::InitRenderModes(UTextureGraphBase* InTextureGraph)
 {
 	if (!RenderModeMgr)
 		RenderModeMgr = MakeShared<TG_RenderModeManager>();
@@ -462,7 +462,6 @@ void STG_EditorViewport::InitPreviewMesh()
 {
 	UStaticMesh* Primitive = GUnrealEd->GetThumbnailManager()->EditorCube;
 	SetPreviewAsset(Primitive);
-	TextureGraphPtr->GetSettings()->SetPreviewMesh(Primitive);	
 			
 }
 
@@ -487,7 +486,6 @@ void STG_EditorViewport::OnSetPreviewPrimitive(EThumbnailPrimType PrimType, bool
 		if (Primitive != nullptr)
 		{
 			SetPreviewAsset(Primitive);
-			TextureGraphPtr->GetSettings()->SetPreviewMesh(Primitive);	
 			
 			RefreshViewport();
 		}
@@ -512,7 +510,6 @@ void STG_EditorViewport::OnSetPreviewMeshFromSelection()
 				if (ComponentClass->IsChildOf(UMeshComponent::StaticClass()))
 				{
 					SetPreviewAsset(TestAsset);
-					TextureGraphSettings->SetPreviewMesh(Cast<UStaticMesh>(TestAsset));
 					bFoundPreviewMesh = true;
 				}
 			}
@@ -669,6 +666,10 @@ bool STG_EditorViewport::SetPreviewAsset(UObject* InAsset)
 			const FRotator PlaneRotation(-90.0f, 180.0f, 0.0f);
 			Transform.SetRotation(FQuat(PlaneRotation));
 		}
+		if (TextureGraphPtr)
+		{
+			TextureGraphPtr->GetSettings()->SetPreviewMesh(StaticMesh);
+		}
 	}
 	else if (InAsset != nullptr)
 	{
diff --git a/Engine/Plugins/Experimental/TextureGraph/Source/TextureGraphEditor/Private/STG_EditorViewport.h b/Engine/Plugins/Experimental/TextureGraph/Source/TextureGraphEditor/Private/STG_EditorViewport.h
--- a/Engine/Plugins/Experimental/TextureGraph/Source/TextureGraphEditor/Private/STG_EditorViewport.h
+++ b/Engine/Plugins/Experimental/TextureGraph/Source/TextureGraphEditor/Private/STG_EditorViewport.h
@@ -14,6 +14,7 @@
 #include "UObject/GCObject.h"
 #include "Widgets/DeclarativeSyntaxSupport.h"
 
+class UTextureGraphBase;
 class SVerticalBox;
 class SRichTextBlock;
 class ITG_Editor;
@@ -27,7 +28,7 @@ class STG_EditorViewport : public SAssetEditorViewport, public FGCObject, public
 public:
 	SLATE_BEGIN_ARGS(STG_EditorViewport)
 	{}
-	SLATE_ARGUMENT(UTextureGraph*, InTextureGraph)
+	SLATE_ARGUMENT(UTextureGraphBase*, InTextureGraph)
 		SLATE_END_ARGS()
 
 
@@ -37,7 +38,7 @@ private:
 	FName									RenderModeName;	///Current Selected render mode name
 	TArray<FName>							RenderModesList;
 	TSharedPtr<class FEditorViewportClient> EditorViewportClient;	/// Editor viewport client
-	TObjectPtr<UTextureGraph>				TextureGraphPtr;		/// Pointer to the Texture Graph being shown
+	TObjectPtr<UTextureGraphBase>		TextureGraphPtr;		/// Pointer to the Texture Graph being shown
 	
 	EViewModeIndex							CurrentViewMode;		/// The currently selected view mode
 	TSharedPtr<FAdvancedPreviewScene>		PreviewScene;			/// The scene for this viewport
@@ -56,7 +57,7 @@ protected:
 	virtual FReply							OnDrop( const FGeometry& MyGeometry, const FDragDropEvent& DragDropEvent ) override;
 public:
 	void									Construct(const FArguments& InArgs);
-	void									SetTextureGraph(const TObjectPtr<UTextureGraph>& InTextureGraph){ TextureGraphPtr = InTextureGraph; }
+	void									SetTextureGraph(const TObjectPtr<UTextureGraphBase>& InTextureGraph);
 
 	STG_EditorViewport();
 	virtual									~STG_EditorViewport() override;
@@ -76,7 +77,7 @@ public:
 	TArray<FName>							GetRenderModesList() const { return RenderModesList;};
 	bool									IsRenderModeEnabled(FName RenderModeTypeArg) const;
 	void									SetRenderMode(FName RenderModeTypeArg);
-	void									InitRenderModes(UTextureGraph* InTextureGraph);
+	void									InitRenderModes(UTextureGraphBase* InTextureGraph);
 	void									UpdateRenderMode();
 	FText									GetRenderModeName() { return FText::FromName(RenderModeName); }
 	FName									GetRenderModeFName() { return RenderModeName; }
diff --git a/Engine/Plugins/Experimental/TextureGraph/Source/TextureGraphEditor/Private/TG_Exporter.cpp b/Engine/Plugins/Experimental/TextureGraph/Source/TextureGraphEditor/Private/TG_Exporter.cpp
--- a/Engine/Plugins/Experimental/TextureGraph/Source/TextureGraphEditor/Private/TG_Exporter.cpp
+++ b/Engine/Plugins/Experimental/TextureGraph/Source/TextureGraphEditor/Private/TG_Exporter.cpp
@@ -17,7 +17,7 @@
 #include "Framework/Commands/UICommandInfo.h"
 #include "Framework/Commands/UICommandList.h"
 #include "Framework/Docking/LayoutService.h"
-#include "Framework/Docking/TabManager.h"
+
 #include "Framework/MultiBox/MultiBoxBuilder.h"
 #include "HAL/Platform.h"
 #include "Internationalization/Internationalization.h"
@@ -52,6 +52,7 @@
 #include "STG_NodePreview.h"
 #include "TG_HelperFunctions.h"
 #include "PropertyEditorModule.h"
+#include "Components/StaticMeshComponent.h"
 #include "DetailsViewArgs.h"
 #include "IDetailsView.h"
 
@@ -60,109 +61,6 @@ class SWindow;
 
 #define LOCTEXT_NAMESPACE "TextureGraphExporter"
 
-struct FTG_ExporterCommands : public TCommands<FTG_ExporterCommands>
-{
-	FTG_ExporterCommands()
-		: TCommands<FTG_ExporterCommands>(
-			TEXT("TextureGraphExporter"), // Context name for fast lookup
-			LOCTEXT("TextureGraphDebugger", "Texture Graph Exporter"), // Localized context name for displaying
-			NAME_None, // Parent
-			FCoreStyle::Get().GetStyleSetName() // Icon Style Set
-		)
-	{
-	}
-
-	// TCommand<> interface
-	virtual void RegisterCommands() override;
-	// End of TCommand<> interface
-
-	TSharedPtr<FUICommandInfo> ShowOutputPreview;
-	TSharedPtr<FUICommandInfo> Show3DPreview;
-	TSharedPtr<FUICommandInfo> Show3DPreviewSettings;
-	TSharedPtr<FUICommandInfo> ShowParameters;
-	TSharedPtr<FUICommandInfo> ShowExportSettings;
-
-};
-
-
-struct FTG_ExporterImpl: public FTickableGameObject, public FGCObject
-{
-public:
-	FTG_ExporterImpl();
-	virtual ~FTG_ExporterImpl() override;
-
-	/** Function registered with tab manager to create the Texture Graph Exporter */
-	TSharedRef<SDockTab> CreateTGExporterTab(const FSpawnTabArgs& Args);
-
-	void Cleanup();
-	/** Sets the Texture Graph to be exported in the Exporter */
-	void SetTextureGraphToExport(UTextureGraph* InTextureGraph);
-	void OnGraphChanged(UTG_Graph* InGraph, UTG_Node* InNode, bool Tweaking);
-
-	TSharedPtr<FTabManager>							TGExporterTabManager;
-	TSharedPtr<FTabManager::FLayout>				TGExporterLayout;
-
-	TObjectPtr<UTextureGraph>						TextureGraphPtr;
-	TWeakPtr<STG_NodePreviewWidget>					NodePreviewPtr;
-	
-	TSharedPtr<class IDetailsView>					ParametersView;
-	TSharedPtr<class IDetailsView>					ExportSettingsView;
-	TSharedPtr<class IDetailsView>					PreviewSettingsView;
-// Tracking the active viewports in this editor.
-	TWeakPtr<class FEditorViewportTabContent>		ViewportTabContentPtr;
-
-	TObjectPtr<UTG_Parameters>						Parameters;
-	TObjectPtr<UTG_ExportSettings>					ExportSettings;
-	FExportSettings                                 TargetExportSettings;
-	
-	TArray<TSharedPtr<FName>>						OutputNodesList;
-	TSharedPtr<SComboBox<TSharedPtr<FName>>>		OutputNodesComboBoxWidget;
-	UTG_Node*										SelectedNode = nullptr;
-	// Inherited via FTickableGameObject
-	virtual void									Tick(float DeltaTime) override;
-
-	virtual ETickableTickType						GetTickableTickType() const override { return ETickableTickType::Always; }
-
-	virtual bool									IsTickableWhenPaused() const override { return true; }
-
-	virtual bool									IsTickableInEditor() const override { return true; }
-
-	virtual TStatId									GetStatId() const override;
-
-	// Inherited via FGCObject
-	virtual void									AddReferencedObjects(FReferenceCollector& Collector) override;
-	virtual FString									GetReferencerName() const override { return TEXT("FTextureGraphExporter");}
-
-private:
-	TSharedRef<SDockTab>							SpawnTab_Viewport(const FSpawnTabArgs& Args);
-	TSharedRef<SDockTab>							SpawnTab_ParameterDefaults(const FSpawnTabArgs& Args);
-	TSharedRef<SDockTab>							SpawnTab_NodePreview(const FSpawnTabArgs& Args);
-	TSharedRef<SDockTab>							SpawnTab_ExportSettings(const FSpawnTabArgs& Args);
-	TSharedRef<SDockTab>							SpawnTab_PreviewSettings(const FSpawnTabArgs& Args);
-	void 											RegisterTabSpawners(const TSharedPtr<FTabManager>);
-	void  											UnregisterTabSpawners(const TSharedPtr<FTabManager>);
-	TSharedPtr<class STG_EditorViewport>			GetEditorViewport() const;
-	void 											SetViewportPreviewMesh();	
-	bool 											SetPreviewAsset(UObject* InAsset);
-	TSharedRef<class IDetailsView>					GetPreviewSettingsView() const { return PreviewSettingsView.ToSharedRef(); }
-	TSharedRef<class IDetailsView>					GetExportSettingsView() const { return ExportSettingsView.ToSharedRef(); }
-	TSharedRef<class IDetailsView>					GetParametersView() const { return ParametersView.ToSharedRef(); }
-	void											OnViewportMaterialChanged();
-	void											OnMaterialMappingChanged();
-	void											UpdateExportSettingsUI();
-	void											UpdateParametersUI();
-	virtual void									RefreshViewport();
-	void											OnRenderingDone(UMixInterface* TextureGraph, const FInvalidationDetails* Details);
-	TSharedRef<SWidget>								GenerateOutputComboItem(TSharedPtr<FName> String);
-	void											OnOutputSelectionChanged(TSharedPtr<FName> String, ESelectInfo::Type Arg);
-	FReply											OnExportClicked(EAppReturnType::Type ButtonID);
-
-	// prevent copying:
-	FTG_ExporterImpl(const FTG_ExporterImpl&);
-	FTG_ExporterImpl(FTG_ExporterImpl&&);
-	FTG_ExporterImpl& operator=(FTG_ExporterImpl const&);
-	FTG_ExporterImpl& operator=(FTG_ExporterImpl&&);
-};
 
 void FTG_ExporterCommands::RegisterCommands()
 {
@@ -172,379 +70,115 @@ void FTG_ExporterCommands::RegisterCommands()
 	UI_COMMAND(ShowParameters, "Parameters", "Toggles visibility of the Parameters window", EUserInterfaceActionType::Check, FInputChord());
 	UI_COMMAND(ShowExportSettings, "Export Settings", "Toggles visibility of the Export Settings window", EUserInterfaceActionType::Check, FInputChord());
 }
-FTG_ExporterImpl::FTG_ExporterImpl()
+FTG_InstanceImpl::FTG_InstanceImpl() :
+	TargetExportSettings(MakeShared<FExportSettings>())
 {
-	const IWorkspaceMenuStructure& MenuStructure = WorkspaceMenu::GetMenuStructure();
-
-	FTG_ExporterCommands::Register();
-	FGlobalTabmanager::Get()->RegisterNomadTabSpawner(FTG_EditorTabs::TextureExporterTabId, FOnSpawnTab::CreateRaw(this, &FTG_ExporterImpl::CreateTGExporterTab))
-		.SetDisplayName(NSLOCTEXT("TextureGraphExporter", "TabTitle", "Texture Graph Exporter"))
-		.SetTooltipText(NSLOCTEXT("TextureGraphExporter", "TooltipText", "Open the Texture Graph Exporter tab."))
-		.SetGroup(MenuStructure.GetDeveloperToolsMiscCategory())
-		.SetIcon(FSlateIcon(FAppStyle::GetAppStyleSetName(), "ClassIcon.Texture2D"));
 }
 
-FTG_ExporterImpl::~FTG_ExporterImpl()
+FTG_InstanceImpl::~FTG_InstanceImpl()
 {
-	if (FSlateApplication::IsInitialized())
+	Cleanup();
+	DetailsView.Reset();
+	ExportSettingsView.Reset();
+	PreviewSettingsView.Reset();
+	ParametersView.Reset();
+
+	// cleanup UI
+	if (Parameters->IsValidLowLevelFast())
 	{
-		FGlobalTabmanager::Get()->UnregisterNomadTabSpawner(FTG_EditorTabs::TextureExporterTabId);
+		Parameters->Parameters.Empty();
+		Parameters = nullptr;
 	}
-
-	if (TGExporterTabManager.IsValid())
+	if (ExportSettings->IsValidLowLevelFast())
 	{
-		UnregisterTabSpawners(TGExporterTabManager);
-		FGlobalTabmanager::Get()->UnregisterTabSpawner(FTG_EditorTabs::TextureExporterTabId);
-		TGExporterLayout = TSharedPtr<FTabManager::FLayout>();
-		TGExporterTabManager = TSharedPtr<FTabManager>();
-		
-		
-		Cleanup();
-		
-		ExportSettingsView.Reset();
-		PreviewSettingsView.Reset();
-		ParametersView.Reset();
-		
-		// cleanup UI
-		if (Parameters->IsValidLowLevelFast())
-		{
-			Parameters->Parameters.Empty();
-			Parameters = nullptr;
-		}
-		if (ExportSettings->IsValidLowLevelFast())
-		{
-			ExportSettings->OutputExpressionsInfos.Empty();
-			ExportSettings = nullptr;
-		}
-		OutputNodesComboBoxWidget.Reset();
+		ExportSettings->OutputExpressionsInfos.Empty();
+		ExportSettings = nullptr;
 	}
-	FTG_ExporterCommands::Unregister();
+	OutputNodesComboBoxWidget.Reset();
 }
 
-void FTG_ExporterImpl::Cleanup()
+void FTG_InstanceImpl::Cleanup()
 {
-	if (TextureGraphPtr->IsValidLowLevelFast())
+	if (TextureGraphInstance && TextureGraphInstance->IsValidLowLevelFast())
 	{
 		// cleanup events
-		TextureGraphPtr->GetSettings()->GetViewportSettings().OnViewportMaterialChangedEvent.RemoveAll(this);
-		TextureGraphPtr->GetSettings()->GetViewportSettings().OnMaterialMappingChangedEvent.RemoveAll(this);
-		TextureGraphPtr->GetSettings()->OnPreviewMeshChangedEvent.RemoveAll(this);
-		TextureGraphPtr->Graph()->OnGraphChangedDelegate.RemoveAll(this);
-		TextureGraphPtr->OnRenderDone.Unbind();
-		TextureGraphPtr = nullptr;
+		TextureGraphInstance->GetSettings()->GetViewportSettings().OnViewportMaterialChangedEvent.RemoveAll(this);
+		TextureGraphInstance->GetSettings()->GetViewportSettings().OnMaterialMappingChangedEvent.RemoveAll(this);
+		TextureGraphInstance->GetSettings()->OnPreviewMeshChangedEvent.RemoveAll(this);
+		TextureGraphInstance->Graph()->OnGraphChangedDelegate.RemoveAll(this);
+		TextureGraphInstance->OnRenderDone.Unbind();
+		TextureGraphInstance = nullptr;
 	}
 
 	OutputNodesList.Empty();
 	SelectedNode = nullptr;
 }
-TSharedRef<SDockTab> FTG_ExporterImpl::CreateTGExporterTab(const FSpawnTabArgs& Args)
-{
-	const TSharedRef<SDockTab> NomadTab = SNew(SDockTab)
-		.TabRole(ETabRole::NomadTab)
-		.Label(NSLOCTEXT("TextureGraphExporter", "TabTitle", "Texture Graph Exporter"));
-
-	TGExporterTabManager = FGlobalTabmanager::Get()->NewTabManager(NomadTab);
-	// on persist layout will handle saving layout if the editor is shut down:
-	TGExporterTabManager->SetOnPersistLayout(
-		FTabManager::FOnPersistLayout::CreateStatic(
-			[](const TSharedRef<FTabManager::FLayout>& InLayout)
-			{
-				if (InLayout->GetPrimaryArea().Pin().IsValid())
-				{
-					FLayoutSaveRestore::SaveToConfig(GEditorLayoutIni, InLayout);
-				}
-			}
-		)
-	);
-
-
-	TWeakPtr<FTabManager> TGExporterTabManagerWeak = TGExporterTabManager;
-	// On tab close will save the layout if the exporter window itself is closed,
-	// this handler also cleans up any floating controls. If we don't close
-	// all areas we need to add some logic to the tab manager to reuse existing tabs:
-	NomadTab->SetOnTabClosed(SDockTab::FOnTabClosedCallback::CreateStatic(
-		[](TSharedRef<SDockTab> Self, TWeakPtr<FTabManager> TabManager)
-		{
-			TSharedPtr<FTabManager> OwningTabManager = TabManager.Pin();
-			if (OwningTabManager.IsValid())
-			{
-				FLayoutSaveRestore::SaveToConfig(GEditorLayoutIni, OwningTabManager->PersistLayout());
-				OwningTabManager->CloseAllAreas();
-			}
-		}
-		, TGExporterTabManagerWeak
-	));
-	
-	FPropertyEditorModule& PropertyEditorModule = FModuleManager::GetModuleChecked<FPropertyEditorModule>("PropertyEditor");
-	FDetailsViewArgs ParameterViewArgs;
-	ParameterViewArgs.NameAreaSettings = FDetailsViewArgs::HideNameArea;
-	ParameterViewArgs.bHideSelectionTip = true;
-	ParameterViewArgs.ColumnWidth = 0.70;
-	ParametersView = PropertyEditorModule.CreateDetailView(ParameterViewArgs);
-
-	// Settings details view
-	FDetailsViewArgs ExportSettingsViewArgs;
-	ExportSettingsViewArgs.NameAreaSettings = FDetailsViewArgs::HideNameArea;
-	ExportSettingsViewArgs.bHideSelectionTip = true;
-	ExportSettingsView = PropertyEditorModule.CreateDetailView(ExportSettingsViewArgs);
-
-	// Settings details view
-	FDetailsViewArgs SettingsViewArgs;
-	SettingsViewArgs.NameAreaSettings = FDetailsViewArgs::HideNameArea;
-	SettingsViewArgs.bHideSelectionTip = true;
-	PreviewSettingsView = PropertyEditorModule.CreateDetailView(SettingsViewArgs);
-	
-	RegisterTabSpawners(TGExporterTabManager);
-
-	
-	TGExporterLayout = FTabManager::NewLayout("Standalone_TextureGraphExporter_Layout_v1")
-	->AddArea
-	(
-		FTabManager::NewPrimaryArea()->SetOrientation(Orient_Vertical)
-			->Split
-			(
-				FTabManager::NewSplitter()->SetOrientation(Orient_Horizontal)->SetSizeCoefficient(0.9f)
-				->Split
-				(
-					FTabManager::NewSplitter()->SetOrientation(Orient_Vertical)->SetSizeCoefficient(0.25f)
-					->Split
-					(
-						FTabManager::NewStack()
-						->AddTab(FTG_EditorTabs::ParameterDefaultsTabId, ETabState::OpenedTab)
-					)
-					->Split
-					(
-						FTabManager::NewStack()
-						->AddTab(FTG_EditorTabs::OutputTabId, ETabState::OpenedTab)
-					)
-				)
-				->Split
-				(
-					FTabManager::NewSplitter()->SetOrientation(Orient_Vertical)->SetSizeCoefficient(0.25f)
-					->Split
-					(
-						FTabManager::NewStack()
-						->SetHideTabWell(true)
-						->AddTab(FTG_EditorTabs::NodePreviewTabId, ETabState::OpenedTab)
-					)
-					->Split
-					(
-						FTabManager::NewStack()
-						->AddTab(FTG_EditorTabs::ViewportTabId, ETabState::OpenedTab)
-						->AddTab(FTG_EditorTabs::PreviewSettingsTabId, ETabState::OpenedTab)
-						->SetForegroundTab(FTG_EditorTabs::ViewportTabId)
-					)
-				)
-			)
-		);
-
-	
-	TGExporterLayout = FLayoutSaveRestore::LoadFromConfig(GEditorLayoutIni, TGExporterLayout.ToSharedRef());
-	
-	TSharedRef<SWidget> TabContents = TGExporterTabManager->RestoreFrom(TGExporterLayout.ToSharedRef(), TSharedPtr<SWindow>()).ToSharedRef();
-	
-	// build command list for tab restoration menu:
-	TSharedPtr<FUICommandList> CommandList = MakeShareable(new FUICommandList());
-	
-	TWeakPtr<FTabManager> TGExportManagerWeak = TGExporterTabManager;
-	
-	const auto ToggleTabVisibility = [](TWeakPtr<FTabManager> InTGExportManagerWeak, FName InTabName)
-	{
-		TSharedPtr<FTabManager> InTGExportManager = InTGExportManagerWeak.Pin();
-		if (InTGExportManager.IsValid())
-		{
-			TSharedPtr<SDockTab> ExistingTab = InTGExportManager->FindExistingLiveTab(InTabName);
-			if (ExistingTab.IsValid())
-			{
-				ExistingTab->RequestCloseTab();
-			}
-			else
-			{
-				InTGExportManager->TryInvokeTab(InTabName);
-			}
-		}
-	};
-	
-	const auto IsTabVisible = [](TWeakPtr<FTabManager> InTGExportManagerWeak, FName InTabName)
-	{
-		TSharedPtr<FTabManager> InTGExportManager = InTGExportManagerWeak.Pin();
-		if (InTGExportManager.IsValid())
-		{
-			return InTGExportManager->FindExistingLiveTab(InTabName).IsValid();
-		}
-		return false;
-	};
-	
-	bool bViewportIsOff = !ViewportTabContentPtr.IsValid();
-	TSharedPtr<SDockTab> ViewportTab; 
-	// check here if 3d viewport is turned off, we need to turn it on temporarily to initialize our systems correctly
-	if (bViewportIsOff)
-	{
-		ViewportTab = TGExporterTabManager->TryInvokeTab(FTG_EditorTabs::ViewportTabId);	
-	}
-	
-	// Set the preview mesh for the material.  This call must occur after the toolbar is initialized.
-	SetViewportPreviewMesh();
-
-	CommandList->MapAction(
-		FTG_ExporterCommands::Get().Show3DPreview,
-		FExecuteAction::CreateStatic(
-			ToggleTabVisibility,
-			TGExportManagerWeak,
-			FTG_EditorTabs::ViewportTabId
-		),
-		FCanExecuteAction::CreateStatic(
-			[]() { return true; }
-		),
-		FIsActionChecked::CreateStatic(
-			IsTabVisible,
-			TGExportManagerWeak,
-			FTG_EditorTabs::ViewportTabId
-		)
-	);
-	
-	CommandList->MapAction(
-		FTG_ExporterCommands::Get().Show3DPreviewSettings,
-		FExecuteAction::CreateStatic(
-			ToggleTabVisibility,
-			TGExportManagerWeak,
-			FTG_EditorTabs::PreviewSettingsTabId
-		),
-		FCanExecuteAction::CreateStatic(
-			[]() { return true; }
-		),
-		FIsActionChecked::CreateStatic(
-			IsTabVisible,
-			TGExportManagerWeak,
-			FTG_EditorTabs::PreviewSettingsTabId
-		)
-	);
-
-	CommandList->MapAction(
-			FTG_ExporterCommands::Get().ShowParameters,
-			FExecuteAction::CreateStatic(
-				ToggleTabVisibility,
-				TGExportManagerWeak,
-				FTG_EditorTabs::ParameterDefaultsTabId
-			),
-			FCanExecuteAction::CreateStatic(
-				[]() { return true; }
-			),
-			FIsActionChecked::CreateStatic(
-				IsTabVisible,
-				TGExportManagerWeak,
-				FTG_EditorTabs::ParameterDefaultsTabId
-			)
-		);
-
-	CommandList->MapAction(
-			FTG_ExporterCommands::Get().ShowOutputPreview,
-			FExecuteAction::CreateStatic(
-				ToggleTabVisibility,
-				TGExportManagerWeak,
-				FTG_EditorTabs::NodePreviewTabId
-			),
-			FCanExecuteAction::CreateStatic(
-				[]() { return true; }
-			),
-			FIsActionChecked::CreateStatic(
-				IsTabVisible,
-				TGExportManagerWeak,
-				FTG_EditorTabs::NodePreviewTabId
-			)
-		);
-
-	CommandList->MapAction(
-			FTG_ExporterCommands::Get().ShowExportSettings,
-			FExecuteAction::CreateStatic(
-				ToggleTabVisibility,
-				TGExportManagerWeak,
-				FTG_EditorTabs::OutputTabId
-			),
-			FCanExecuteAction::CreateStatic(
-				[]() { return true; }
-			),
-			FIsActionChecked::CreateStatic(
-				IsTabVisible,
-				TGExportManagerWeak,
-				FTG_EditorTabs::OutputTabId
-			)
-		);
-	
-	FMenuBarBuilder MenuBarBuilder(CommandList);
-	MenuBarBuilder.AddPullDownMenu(
-		LOCTEXT("WindowMenuLabel", "Window"),
-		FText::GetEmpty(),
-		FNewMenuDelegate::CreateLambda([](FMenuBuilder& Builder) {
-			Builder.AddMenuEntry(FTG_ExporterCommands::Get().ShowOutputPreview);
-			Builder.AddMenuEntry(FTG_ExporterCommands::Get().ShowParameters);
-			Builder.AddMenuEntry(FTG_ExporterCommands::Get().Show3DPreviewSettings);
-			Builder.AddMenuEntry(FTG_ExporterCommands::Get().Show3DPreview);
-			Builder.AddMenuEntry(FTG_ExporterCommands::Get().ShowExportSettings);
-			})
-	);
-	
-	
-	TSharedRef<SWidget> MenuBarWidget = MenuBarBuilder.MakeWidget();
-	
-	NomadTab->SetContent(
-		SNew(SVerticalBox)
-		+ SVerticalBox::Slot()
-		.AutoHeight()
-		[
-			MenuBarWidget
-		]
-		+SVerticalBox::Slot()
-		[
-			SNew(SBorder)
-			.BorderImage(FAppStyle::GetBrush("ToolPanel.DarkGroupBorder"))
-			.Padding(FMargin(0.f, 2.f))
-			[
-				TabContents
-			]
-		]
-	);
-	
-	
-	// Tell tab-manager about the multi-box for platforms with a global menu bar
-	TGExporterTabManager->SetMenuMultiBox(MenuBarBuilder.GetMultiBox(), MenuBarWidget);
-	
-	return NomadTab;
-}
 
-
-void FTG_ExporterImpl::RegisterTabSpawners(const TSharedPtr<FTabManager> InTabManager)
+void FTG_InstanceImpl::RegisterTabSpawners(const TSharedPtr<FTabManager> InTabManager)
 {
-	InTabManager->RegisterTabSpawner(FTG_EditorTabs::ViewportTabId, FOnSpawnTab::CreateRaw(this, &FTG_ExporterImpl::SpawnTab_Viewport))
+	InTabManager->RegisterTabSpawner(FTG_EditorTabs::ViewportTabId, FOnSpawnTab::CreateRaw(this, &FTG_InstanceImpl::SpawnTab_Viewport))
 		.SetDisplayName(LOCTEXT("ViewportTab", "3D Preview"))
 		.SetIcon(FSlateIcon(FAppStyle::GetAppStyleSetName(), "LevelEditor.Tabs.Viewports"));
-
-	InTabManager->RegisterTabSpawner(FTG_EditorTabs::ParameterDefaultsTabId, FOnSpawnTab::CreateRaw(this, &FTG_ExporterImpl::SpawnTab_ParameterDefaults))
+	
+	InTabManager->RegisterTabSpawner(FTG_EditorTabs::ParameterDefaultsTabId, FOnSpawnTab::CreateRaw(this, &FTG_InstanceImpl::SpawnTab_ParameterDefaults))
 		.SetDisplayName(LOCTEXT("ParametersTab", "Parameters"))
 		.SetIcon(FSlateIcon(FAppStyle::GetAppStyleSetName(), "LevelEditor.Tabs.Details"));
 
-	InTabManager->RegisterTabSpawner(FTG_EditorTabs::NodePreviewTabId, FOnSpawnTab::CreateRaw(this, &FTG_ExporterImpl::SpawnTab_NodePreview))
+	InTabManager->RegisterTabSpawner(FTG_EditorTabs::NodePreviewTabId, FOnSpawnTab::CreateRaw(this, &FTG_InstanceImpl::SpawnTab_NodePreview))
 		.SetDisplayName(LOCTEXT("NodePreviewTab", "Node Preview"))
 		.SetIcon(FSlateIcon(FAppStyle::GetAppStyleSetName(), "LevelEditor.Tabs.Viewports"));
 
-	InTabManager->RegisterTabSpawner(FTG_EditorTabs::PreviewSettingsTabId, FOnSpawnTab::CreateRaw(this, &FTG_ExporterImpl::SpawnTab_PreviewSettings))
+	InTabManager->RegisterTabSpawner(FTG_EditorTabs::PreviewSettingsTabId, FOnSpawnTab::CreateRaw(this, &FTG_InstanceImpl::SpawnTab_PreviewSettings))
 		.SetDisplayName(LOCTEXT("PreviewSettingsTab", "3D Preview Settings"))
 		.SetIcon(FSlateIcon(FAppStyle::GetAppStyleSetName(), "LevelEditor.Tabs.Details"));
 
-	InTabManager->RegisterTabSpawner(FTG_EditorTabs::OutputTabId, FOnSpawnTab::CreateRaw(this, &FTG_ExporterImpl::SpawnTab_ExportSettings))
+	InTabManager->RegisterTabSpawner(FTG_EditorTabs::PreviewSceneSettingsTabId, FOnSpawnTab::CreateRaw(this, &FTG_InstanceImpl::SpawnTab_PreviewSettings))
+			.SetDisplayName(LOCTEXT("PreviewSceneSettingsTab", "Preview Scene Settings"))
+			.SetIcon(FSlateIcon(FAppStyle::GetAppStyleSetName(), "LevelEditor.Tabs.Details"));
+	
+	InTabManager->RegisterTabSpawner(FTG_EditorTabs::OutputTabId, FOnSpawnTab::CreateRaw(this, &FTG_InstanceImpl::SpawnTab_ExportSettings))
 		.SetDisplayName(LOCTEXT("ExportSettingsTab", "Export Settings"))
 		.SetIcon(FSlateIcon(FAppStyle::GetAppStyleSetName(), "LevelEditor.Tabs.Details"));
 }
 
-void FTG_ExporterImpl::UnregisterTabSpawners(const TSharedPtr<FTabManager> InTabManager)
+void FTG_InstanceImpl::UnregisterTabSpawners(const TSharedPtr<FTabManager> InTabManager)
 {
 	InTabManager->UnregisterTabSpawner(FTG_EditorTabs::ViewportTabId);
 	InTabManager->UnregisterTabSpawner(FTG_EditorTabs::ParameterDefaultsTabId);
 	InTabManager->UnregisterTabSpawner(FTG_EditorTabs::NodePreviewTabId);
 	InTabManager->UnregisterTabSpawner(FTG_EditorTabs::OutputTabId);
 	InTabManager->UnregisterTabSpawner(FTG_EditorTabs::PreviewSettingsTabId);
+	InTabManager->UnregisterTabSpawner(FTG_EditorTabs::PreviewSceneSettingsTabId);
 }
-TSharedRef<SDockTab> FTG_ExporterImpl::SpawnTab_PreviewSettings(const FSpawnTabArgs& Args)
+
+void FTG_InstanceImpl::SetMesh(class UMeshComponent* InPreviewMesh, class UWorld* InWorld)
+{
+	TextureGraphInstance->SetEditorMesh(Cast<UStaticMeshComponent>(InPreviewMesh), InWorld).then([this]()
+		{
+			GetEditorViewport()->InitRenderModes(TextureGraphInstance);
+		});
+}
+
+bool FTG_InstanceImpl::SetPreviewAsset(UObject* InAsset)
+{
+	if (GetEditorViewport().IsValid())
+	{
+		return GetEditorViewport()->SetPreviewAsset(InAsset);
+	}
+	return false;
+}
+TSharedPtr<STG_EditorViewport> FTG_InstanceImpl::GetEditorViewport() const
+{
+	if (ViewportTabContent.IsValid())
+	{
+		// we can use static cast here b/c we know in this editor we will have a static mesh viewport 
+		return StaticCastSharedPtr<STG_EditorViewport>(ViewportTabContent->GetFirstViewport());
+	}
+
+	return nullptr;
+}
+
+TSharedRef<SDockTab> FTG_InstanceImpl::SpawnTab_PreviewSettings(const FSpawnTabArgs& Args)
 {
 	check(Args.GetTabId() == FTG_EditorTabs::PreviewSettingsTabId);
 
@@ -553,14 +187,14 @@ TSharedRef<SDockTab> FTG_ExporterImpl::SpawnTab_PreviewSettings(const FSpawnTabA
 			PreviewSettingsView.ToSharedRef()
 		];
 
-	if (TextureGraphPtr)
+	if (TextureGraphInstance)
 	{
-		GetPreviewSettingsView()->SetObject(TextureGraphPtr->GetSettings(), true);
+		GetPreviewSettingsView()->SetObject(TextureGraphInstance->GetSettings(), true);
 	}
 
 	return SettingsTab.ToSharedRef();
 }
-TSharedRef<SDockTab> FTG_ExporterImpl::SpawnTab_ExportSettings(const FSpawnTabArgs& Args)
+TSharedRef<SDockTab> FTG_InstanceImpl::SpawnTab_ExportSettings(const FSpawnTabArgs& Args)
 {
 	check(Args.GetTabId() == FTG_EditorTabs::OutputTabId);
 
@@ -589,19 +223,19 @@ TSharedRef<SDockTab> FTG_ExporterImpl::SpawnTab_ExportSettings(const FSpawnTabAr
 			[
 				SNew(SPrimaryButton)
 				.Text(LOCTEXT("Export", "Export"))
-				.OnClicked_Raw(this, &FTG_ExporterImpl::OnExportClicked, EAppReturnType::Ok)
+				.OnClicked_Raw(this, &FTG_InstanceImpl::OnExportClicked, EAppReturnType::Ok)
 			]
 			
 		];
 
-	if (TextureGraphPtr)
+	if (TextureGraphInstance)
 	{
 		GetExportSettingsView()->SetObject(ExportSettings, true);
 	}
 
 	return SettingsTab.ToSharedRef();
 }
-TSharedRef<SDockTab> FTG_ExporterImpl::SpawnTab_Viewport(const FSpawnTabArgs& Args)
+TSharedRef<SDockTab> FTG_InstanceImpl::SpawnTab_Viewport(const FSpawnTabArgs& Args)
 {
 	check(Args.GetTabId() == FTG_EditorTabs::ViewportTabId);
 
@@ -611,25 +245,22 @@ TSharedRef<SDockTab> FTG_ExporterImpl::SpawnTab_Viewport(const FSpawnTabArgs& Ar
 	AssetEditorViewportFactoryFunction MakeViewportFunc = [this](const FAssetEditorViewportConstructionArgs& InArgs)
 	{
 		return SNew(STG_EditorViewport)
-			.InTextureGraph(TextureGraphPtr);
+			.InTextureGraph(TextureGraphInstance);
 	};
 
 	// Create a new tab
-	TSharedRef<FEditorViewportTabContent> ViewportTabContent = MakeShared<FEditorViewportTabContent>();
-	//ViewportTabContent->OnViewportTabContentLayoutChanged().AddRaw(this, &FTG_ExporterImpl::OnEditorLayoutChanged);
+	ViewportTabContent = MakeShared<FEditorViewportTabContent>();
 
 	const FString LayoutId = FString("TG_EditorViewport");
 	ViewportTabContent->Initialize(MakeViewportFunc, DockableTab, LayoutId);
 
-	ViewportTabContentPtr = ViewportTabContent;
-	
 	// This call must occur after the toolbar is initialized.
 	SetViewportPreviewMesh();
 	
 	return DockableTab;
 }
 
-TSharedRef<SDockTab> FTG_ExporterImpl::SpawnTab_ParameterDefaults(const FSpawnTabArgs& Args)
+TSharedRef<SDockTab> FTG_InstanceImpl::SpawnTab_ParameterDefaults(const FSpawnTabArgs& Args)
 {
 	check(Args.GetTabId() == FTG_EditorTabs::ParameterDefaultsTabId);
 
@@ -643,7 +274,7 @@ TSharedRef<SDockTab> FTG_ExporterImpl::SpawnTab_ParameterDefaults(const FSpawnTa
 }
 
 
-TSharedRef<SDockTab> FTG_ExporterImpl::SpawnTab_NodePreview(const FSpawnTabArgs& Args)
+TSharedRef<SDockTab> FTG_InstanceImpl::SpawnTab_NodePreview(const FSpawnTabArgs& Args)
 {
 	check(Args.GetTabId() == FTG_EditorTabs::NodePreviewTabId);
 
@@ -659,8 +290,8 @@ TSharedRef<SDockTab> FTG_ExporterImpl::SpawnTab_NodePreview(const FSpawnTabArgs&
 			[
 				SAssignNew(OutputNodesComboBoxWidget, SComboBox<TSharedPtr<FName>>)
 				.OptionsSource(&OutputNodesList)
-				.OnGenerateWidget_Raw(this, &FTG_ExporterImpl::GenerateOutputComboItem)
-				.OnSelectionChanged_Raw(this, &FTG_ExporterImpl::OnOutputSelectionChanged)
+				.OnGenerateWidget_Raw(this, &FTG_InstanceImpl::GenerateOutputComboItem)
+				.OnSelectionChanged_Raw(this, &FTG_InstanceImpl::OnOutputSelectionChanged)
 				[
 					SNew(STextBlock)
 					.Text_Lambda([this] ()
@@ -684,12 +315,13 @@ TSharedRef<SDockTab> FTG_ExporterImpl::SpawnTab_NodePreview(const FSpawnTabArgs&
 				NodePreview
 			]
 		];
+
 }
-void FTG_ExporterImpl::SetViewportPreviewMesh()
+void FTG_InstanceImpl::SetViewportPreviewMesh()
 {
-	if (TextureGraphPtr)
+	if (TextureGraphInstance)
 	{
-		TObjectPtr<UStaticMesh> PreviewMesh = TextureGraphPtr->GetSettings()->GetPreviewMesh();
+		TObjectPtr<UStaticMesh> PreviewMesh = TextureGraphInstance->GetSettings()->GetPreviewMesh();
 		// Set the preview mesh for the material.  
 		if (!PreviewMesh || !SetPreviewAsset(PreviewMesh))
 		{
@@ -698,39 +330,32 @@ void FTG_ExporterImpl::SetViewportPreviewMesh()
 		}
 	}
 }
-TSharedPtr<STG_EditorViewport> FTG_ExporterImpl::GetEditorViewport() const
-{
-	if (ViewportTabContentPtr.IsValid())
-	{
-		// we can use static cast here b/c we know in this editor we will have a static mesh viewport 
-		return StaticCastSharedPtr<STG_EditorViewport>(ViewportTabContentPtr.Pin()->GetFirstViewport());
-	}
 
-	return nullptr;
-}
-bool FTG_ExporterImpl::SetPreviewAsset(UObject* InAsset)
+void FTG_InstanceImpl::UpdatePreviewMesh()
 {
-	if (GetEditorViewport().IsValid())
+	if (GetEditorViewport())
 	{
-		return GetEditorViewport()->SetPreviewAsset(InAsset);
+		GetEditorViewport()->SetTextureGraph(TextureGraphInstance);
+		// OnViewportMaterialChanged();
+		SetViewportPreviewMesh();
 	}
-	return false;
 }
+
 // Function to generate combo box items
-TSharedRef<SWidget> FTG_ExporterImpl::GenerateOutputComboItem(TSharedPtr<FName> InItem)
+TSharedRef<SWidget> FTG_InstanceImpl::GenerateOutputComboItem(TSharedPtr<FName> InItem)
 {
 	return SNew(STextBlock).Text(FText::FromName(*InItem));
 }
 
 // Function called when the selection changes
-void FTG_ExporterImpl::OnOutputSelectionChanged(TSharedPtr<FName> SelectedItem, ESelectInfo::Type SelectInfo)
+void FTG_InstanceImpl::OnOutputSelectionChanged(TSharedPtr<FName> SelectedItem, ESelectInfo::Type SelectInfo)
 {
 	if (SelectedItem.IsValid())
 	{
 		FName SelectedNodeName = *SelectedItem;
 		FTG_Id SelectedNodeId = FTG_Id::INVALID;
 
-		TextureGraphPtr->Graph()->ForEachNodes([this, &SelectedNodeId, SelectedNodeName](const UTG_Node* Node, uint32 Index)
+		TextureGraphInstance->Graph()->ForEachNodes([this, &SelectedNodeId, SelectedNodeName](const UTG_Node* Node, uint32 Index)
 		{
 			if (Node && Node->GetExpression()->IsA(UTG_Expression_Output::StaticClass()))
 			{
@@ -746,7 +371,7 @@ void FTG_ExporterImpl::OnOutputSelectionChanged(TSharedPtr<FName> SelectedItem,
 			}
 		});
 		
-		SelectedNode = TextureGraphPtr->Graph()->GetNode(SelectedNodeId);
+		SelectedNode = TextureGraphInstance->Graph()->GetNode(SelectedNodeId);
 		if (NodePreviewPtr.IsValid())
 		{
 			NodePreviewPtr.Pin()->SelectionChanged(SelectedNode);
@@ -755,95 +380,157 @@ void FTG_ExporterImpl::OnOutputSelectionChanged(TSharedPtr<FName> SelectedItem,
 	}
 }
 
-FReply FTG_ExporterImpl::OnExportClicked(EAppReturnType::Type ButtonID)
+FReply FTG_InstanceImpl::OnExportClicked(EAppReturnType::Type ButtonID)
 {
-	if (TextureGraphPtr)
+	if (TextureGraphInstance)
 	{
-		FTG_HelperFunctions::ExportAsync(TextureGraphPtr, "", "", this->TargetExportSettings, false);
+		FTG_HelperFunctions::ExportAsync(TextureGraphInstance, "", "", *TargetExportSettings, false);
 	}
 	return FReply::Handled();
 }
-void FTG_ExporterImpl::SetTextureGraphToExport(UTextureGraph* InTextureGraph)
+
+void FTG_InstanceImpl::SetTextureGraphToExport(UTextureGraphInstance* InTextureGraph)
 {
 	// clear out previous handles
 	Cleanup();
 
-	// overwrite the original TextureGraph in place by constructing a new one with the same name
-	FObjectDuplicationParameters Params = InitStaticDuplicateObjectParams(InTextureGraph, GetTransientPackage(), NAME_None,
-	~RF_Standalone, UTextureGraph::StaticClass(), EDuplicateMode::Normal, EInternalObjectFlags::None);
+	TextureGraphInstance = InTextureGraph;
 	
-	TextureGraphPtr = Cast<UTextureGraph>(StaticDuplicateObjectEx(Params));
 	//Exporter gets notified when rendering is done
-	TextureGraphPtr->OnRenderDone.BindRaw(this, &FTG_ExporterImpl::OnRenderingDone);
-
-	// force open export window
-	FGlobalTabmanager::Get()->TryInvokeTab(FTG_EditorTabs::TextureExporterTabId);
-	
+	TextureGraphInstance->OnRenderDone.BindRaw(this, &FTG_InstanceImpl::OnRenderingDone);
 	
 	UpdateParametersUI();
 
 	UpdateExportSettingsUI();
-	
+
 	// Update list of output nodes in 2d View
 	OutputNodesList.Empty();
-	TextureGraphPtr->Graph()->ForEachNodes(
-			[&](const UTG_Node* node, uint32 index)
+	TextureGraphInstance->Graph()->ForEachNodes(
+		[&](const UTG_Node* node, uint32 index)
+		{
+			if(UTG_Expression_Output* OutputExpression = Cast<UTG_Expression_Output>(node->GetExpression()))
+			{
+				OutputNodesList.Add(MakeShared<FName>(OutputExpression->GetTitleName()));			
+			}
+		});
+	
+	if (OutputNodesComboBoxWidget.IsValid())
 	{
-		if(UTG_Expression_Output* OutputExpression = Cast<UTG_Expression_Output>(node->GetExpression()))
+		if (!OutputNodesList.IsEmpty())
 		{
-			OutputNodesList.Add(MakeShared<FName>(OutputExpression->GetTitleName()));			
+			OutputNodesComboBoxWidget->SetSelectedItem(OutputNodesList[0]);
 		}
-	});
-	if (!OutputNodesList.IsEmpty())
-	{
-		OutputNodesComboBoxWidget->SetSelectedItem(OutputNodesList[0]);
+		OutputNodesComboBoxWidget->RefreshOptions();
 	}
-	OutputNodesComboBoxWidget->RefreshOptions();
 	
-	FViewportSettings& ViewportSettings = TextureGraphPtr->GetSettings()->GetViewportSettings();
+	FViewportSettings& ViewportSettings = TextureGraphInstance->GetSettings()->GetViewportSettings();
 
-	ViewportSettings.OnViewportMaterialChangedEvent.AddRaw(this, &FTG_ExporterImpl::OnViewportMaterialChanged);
-	ViewportSettings.OnMaterialMappingChangedEvent.AddRaw(this, &FTG_ExporterImpl::OnMaterialMappingChanged);
-	TextureGraphPtr->GetSettings()->OnPreviewMeshChangedEvent.AddRaw(this, &FTG_ExporterImpl::SetViewportPreviewMesh);
-	TextureGraphPtr->Graph()->OnGraphChangedDelegate.AddRaw(this, &FTG_ExporterImpl::OnGraphChanged);
-	GetPreviewSettingsView()->SetObject(TextureGraphPtr->GetSettings(), true);
+	ViewportSettings.OnViewportMaterialChangedEvent.AddRaw(this, &FTG_InstanceImpl::OnViewportMaterialChanged);
+	ViewportSettings.OnMaterialMappingChangedEvent.AddRaw(this, &FTG_InstanceImpl::OnMaterialMappingChanged);
+	TextureGraphInstance->GetSettings()->OnPreviewMeshChangedEvent.AddRaw(this, &FTG_InstanceImpl::SetViewportPreviewMesh);
+	TextureGraphInstance->Graph()->OnGraphChangedDelegate.AddRaw(this, &FTG_InstanceImpl::OnGraphChanged);
+	GetPreviewSettingsView()->SetObject(TextureGraphInstance->GetSettings(), true);
 
-	GetEditorViewport()->SetTextureGraph(TextureGraphPtr);
-	OnViewportMaterialChanged();
-	SetViewportPreviewMesh();
-	
-	TextureGraphEngine::RegisterErrorReporter(TextureGraphPtr, std::make_shared<FTextureGraphErrorReporter>());
+	UpdatePreviewMesh();
+	OnGraphChanged(TextureGraphInstance->Graph(), nullptr, false);
+	TextureGraphEngine::RegisterErrorReporter(TextureGraphInstance, std::make_shared<FTextureGraphErrorReporter>());
 
 }
-void FTG_ExporterImpl::OnGraphChanged(UTG_Graph* InGraph, UTG_Node* InNode, bool Tweaking)
+void FTG_InstanceImpl::OnGraphChanged(UTG_Graph* InGraph, UTG_Node* InNode, bool Tweaking)
+{
+	if (TextureGraphInstance)
+	{
+		TextureGraphInstance->TriggerUpdate(Tweaking);
+		if (InNode && InNode->IsA<UTG_Expression_Output>())
+		{
+			TextureGraphInstance->UpdateGlobalTGSettings();
+		}
+		RefreshViewport();
+	}
+}
+
+TSharedRef<FTabManager::FLayout> FTG_InstanceImpl::GetDefaultLayout()
+{
+	return FTabManager::NewLayout("Standalone_TextureGraphExporter_Layout_v1")
+	->AddArea
+	(
+		FTabManager::NewPrimaryArea()->SetOrientation(Orient_Vertical)
+			->Split
+			(
+				FTabManager::NewSplitter()->SetOrientation(Orient_Horizontal)->SetSizeCoefficient(0.9f)
+				->Split
+				(
+					FTabManager::NewSplitter()->SetOrientation(Orient_Vertical)->SetSizeCoefficient(0.25f)
+					->Split
+					(
+						FTabManager::NewStack()
+						->AddTab(FTG_EditorTabs::ParameterDefaultsTabId, ETabState::OpenedTab)
+					)
+					->Split
+					(
+						FTabManager::NewStack()
+						->AddTab(FTG_EditorTabs::OutputTabId, ETabState::OpenedTab)
+					)
+				)
+				->Split
+				(
+					FTabManager::NewSplitter()->SetOrientation(Orient_Vertical)->SetSizeCoefficient(0.25f)
+					->Split
+					(
+						FTabManager::NewStack()
+						->SetHideTabWell(true)
+						->AddTab(FTG_EditorTabs::NodePreviewTabId, ETabState::OpenedTab)
+					)
+					->Split
+					(
+						FTabManager::NewStack()
+						->AddTab(FTG_EditorTabs::ViewportTabId, ETabState::OpenedTab)
+						->AddTab(FTG_EditorTabs::PreviewSettingsTabId, ETabState::OpenedTab)
+						->SetForegroundTab(FTG_EditorTabs::ViewportTabId)
+					)
+				)
+			)
+		);
+}
+
+void FTG_InstanceImpl::Initialize()
 {
-	if (TextureGraphPtr)
-	{
-		TextureGraphPtr->TriggerUpdate(Tweaking);
 	
-		// UpdateParametersUI();
+	FPropertyEditorModule& PropertyEditorModule = FModuleManager::GetModuleChecked<FPropertyEditorModule>("PropertyEditor");
+	FDetailsViewArgs ParameterViewArgs;
+	ParameterViewArgs.NameAreaSettings = FDetailsViewArgs::HideNameArea;
+	ParameterViewArgs.bHideSelectionTip = true;
+	ParameterViewArgs.ColumnWidth = 0.70;
+	ParametersView = PropertyEditorModule.CreateDetailView(ParameterViewArgs);
 
-		if (InNode->IsA<UTG_Expression_Output>())
-		{
-			TextureGraphPtr->UpdateGlobalTGSettings();
-		}
-		RefreshViewport();
-	}
-	// OutputView->ForceRefresh();
+	// Settings details view
+	FDetailsViewArgs ExportSettingsViewArgs;
+	ExportSettingsViewArgs.NameAreaSettings = FDetailsViewArgs::HideNameArea;
+	ExportSettingsViewArgs.bHideSelectionTip = true;
+	ExportSettingsView = PropertyEditorModule.CreateDetailView(ExportSettingsViewArgs);
+
+	// Settings details view
+	FDetailsViewArgs SettingsViewArgs;
+	SettingsViewArgs.NameAreaSettings = FDetailsViewArgs::HideNameArea;
+	SettingsViewArgs.bHideSelectionTip = true;
+	PreviewSettingsView = PropertyEditorModule.CreateDetailView(SettingsViewArgs);
+		
+	
 }
 
-void FTG_ExporterImpl::OnRenderingDone(UMixInterface* TextureGraph, const FInvalidationDetails* Details)
+void FTG_InstanceImpl::OnRenderingDone(UMixInterface* TextureGraph, const FInvalidationDetails* Details)
 {
-	if (TextureGraph != nullptr && TextureGraph == TextureGraphPtr && NodePreviewPtr.IsValid())
+	if (TextureGraph != nullptr && TextureGraph == TextureGraphInstance && NodePreviewPtr.IsValid())
 	{
 		// refresh node preview
 		NodePreviewPtr.Pin()->Update();
+		GetEditorViewport()->UpdateRenderMode();
 	}
 }
-void FTG_ExporterImpl::OnViewportMaterialChanged()
+void FTG_InstanceImpl::OnViewportMaterialChanged()
 {
 	const UTG_Node* FirstTargetNode = nullptr;
-	TextureGraphPtr->Graph()->ForEachNodes([&](const UTG_Node* node, uint32 index)
+	TextureGraphInstance->Graph()->ForEachNodes([&](const UTG_Node* node, uint32 index)
 		{
 			if (UTG_Expression_Output* OutputExpression = dynamic_cast<UTG_Expression_Output*>(node->GetExpression()))
 			{
@@ -854,7 +541,7 @@ void FTG_ExporterImpl::OnViewportMaterialChanged()
 			}
 		});
 
-	FViewportSettings& ViewportSettings = TextureGraphPtr->GetSettings()->GetViewportSettings();
+	FViewportSettings& ViewportSettings = TextureGraphInstance->GetSettings()->GetViewportSettings();
 
 	if (FirstTargetNode && ViewportSettings.MaterialMappingInfos.Num() > 0)
 	{
@@ -862,20 +549,20 @@ void FTG_ExporterImpl::OnViewportMaterialChanged()
 	}
 
 	GetEditorViewport()->GenerateRendermodeToolbar();
-	GetEditorViewport()->InitRenderModes(TextureGraphPtr);
+	GetEditorViewport()->InitRenderModes(TextureGraphInstance);
 }
 
-void FTG_ExporterImpl::OnMaterialMappingChanged()
+void FTG_InstanceImpl::OnMaterialMappingChanged()
 {
 	GetEditorViewport()->UpdateRenderMode();
 }
 
-void FTG_ExporterImpl::UpdateExportSettingsUI()
+void FTG_InstanceImpl::UpdateExportSettingsUI()
 {
 	// recreate export settings UI
-	ExportSettings = NewObject<UTG_ExportSettings>(TextureGraphPtr);
+	ExportSettings = NewObject<UTG_ExportSettings>(TextureGraphInstance);
 	
-	UTG_Graph* Graph = TextureGraphPtr->Graph();
+	UTG_Graph* Graph = TextureGraphInstance->Graph();
 
 	Graph->ForEachNodes([&](const UTG_Node* node, uint32 index)
 	{
@@ -887,9 +574,9 @@ void FTG_ExporterImpl::UpdateExportSettingsUI()
 
 	GetExportSettingsView()->SetObject(ExportSettings);
 }
-void FTG_ExporterImpl::UpdateParametersUI()
+void FTG_InstanceImpl::UpdateParametersUI()
 {
-	UTG_Graph* Graph = TextureGraphPtr->Graph();
+	UTG_Graph* Graph = TextureGraphInstance->Graph();
 	FTG_Ids Ids = Graph->GetParamIds();
 
 	// Create a new object to set for the view.
@@ -909,24 +596,24 @@ void FTG_ExporterImpl::UpdateParametersUI()
 
 	GetParametersView()->SetObject(Parameters);
 }
-void FTG_ExporterImpl::Tick(float DeltaTime)
+void FTG_InstanceImpl::Tick(float DeltaTime)
 {
 	RefreshViewport();
 }
 
-TStatId FTG_ExporterImpl::GetStatId() const
+TStatId FTG_InstanceImpl::GetStatId() const
 {
 	return TStatId();
 }
 
-void FTG_ExporterImpl::AddReferencedObjects(FReferenceCollector& Collector)
+void FTG_InstanceImpl::AddReferencedObjects(FReferenceCollector& Collector)
 {
 	Collector.AddReferencedObject(Parameters);
-	Collector.AddReferencedObject(TextureGraphPtr);
+	Collector.AddReferencedObject(TextureGraphInstance);
 	Collector.AddReferencedObject(ExportSettings);
 }
 
-void FTG_ExporterImpl::RefreshViewport()
+void FTG_InstanceImpl::RefreshViewport()
 {
 	if (GetEditorViewport().IsValid())
 	{
@@ -935,18 +622,267 @@ void FTG_ExporterImpl::RefreshViewport()
 }
 
 ///////////////////////////////////////////////////////////////////////////////////
-FTG_Exporter::FTG_Exporter()
-	: Impl(MakeUnique<FTG_ExporterImpl>())
+FTG_ExporterUtility::FTG_ExporterUtility()
+	: Impl(MakeUnique<FTG_InstanceImpl>())
+{
+	const IWorkspaceMenuStructure& MenuStructure = WorkspaceMenu::GetMenuStructure();
+
+	FGlobalTabmanager::Get()->RegisterNomadTabSpawner(FTG_EditorTabs::TextureExporterTabId, FOnSpawnTab::CreateRaw(this, &FTG_ExporterUtility::CreateTGExporterTab))
+		.SetDisplayName(NSLOCTEXT("TextureGraphExporter", "TabTitle", "Texture Graph Exporter"))
+		.SetTooltipText(NSLOCTEXT("TextureGraphExporter", "TooltipText", "Open the Texture Graph Exporter tab."))
+		.SetGroup(MenuStructure.GetDeveloperToolsMiscCategory())
+		.SetIcon(FSlateIcon(FAppStyle::GetAppStyleSetName(), "ClassIcon.Texture2D"));
+}
+
+FTG_ExporterUtility::~FTG_ExporterUtility()
 {
+	if (FSlateApplication::IsInitialized())
+	{
+		FGlobalTabmanager::Get()->UnregisterNomadTabSpawner(FTG_EditorTabs::TextureExporterTabId);
+	}
+
+	if (TGExporterTabManager.IsValid())
+	{
+		Impl->UnregisterTabSpawners(TGExporterTabManager);
+		FGlobalTabmanager::Get()->UnregisterTabSpawner(FTG_EditorTabs::TextureExporterTabId);
+	}
 }
 
-FTG_Exporter::~FTG_Exporter()
+TSharedRef<SDockTab> FTG_ExporterUtility::CreateTGExporterTab(const FSpawnTabArgs& Args)
 {
+	const TSharedRef<SDockTab> NomadTab = SNew(SDockTab)
+		.TabRole(ETabRole::NomadTab)
+		.Label(NSLOCTEXT("TextureGraphExporter", "TabTitle", "Texture Graph Exporter"));
+
+	TGExporterTabManager = FGlobalTabmanager::Get()->NewTabManager(NomadTab);
+	// on persist layout will handle saving layout if the editor is shut down:
+	TGExporterTabManager->SetOnPersistLayout(
+		FTabManager::FOnPersistLayout::CreateStatic(
+			[](const TSharedRef<FTabManager::FLayout>& InLayout)
+			{
+				if (InLayout->GetPrimaryArea().Pin().IsValid())
+				{
+					FLayoutSaveRestore::SaveToConfig(GEditorLayoutIni, InLayout);
+				}
+			}
+		)
+	);
+
+
+	TWeakPtr<FTabManager> TGExporterTabManagerWeak = TGExporterTabManager;
+	// On tab close will save the layout if the exporter window itself is closed,
+	// this handler also cleans up any floating controls. If we don't close
+	// all areas we need to add some logic to the tab manager to reuse existing tabs:
+	NomadTab->SetOnTabClosed(SDockTab::FOnTabClosedCallback::CreateLambda(
+		[this](TSharedRef<SDockTab> Self, TWeakPtr<FTabManager> TabManager)
+		{
+			TSharedPtr<FTabManager> OwningTabManager = TabManager.Pin();
+			if (OwningTabManager.IsValid())
+			{
+				FLayoutSaveRestore::SaveToConfig(GEditorLayoutIni, OwningTabManager->PersistLayout());
+				OwningTabManager->CloseAllAreas();
+			}
+			Impl->Cleanup();
+			
+		}
+		, TGExporterTabManagerWeak
+	));
+	Impl->RegisterTabSpawners(TGExporterTabManager);
+	Impl->Initialize();
+	
+	TGExporterLayout = FLayoutSaveRestore::LoadFromConfig(GEditorLayoutIni, Impl->GetDefaultLayout());
+	TSharedRef<SWidget> TabContents = TGExporterTabManager->RestoreFrom(TGExporterLayout.ToSharedRef(), TSharedPtr<SWindow>()).ToSharedRef();
+	
+	// build command list for tab restoration menu:
+	TSharedPtr<FUICommandList> CommandList = MakeShareable(new FUICommandList());
+	
+	TWeakPtr<FTabManager> TGExportManagerWeak = TGExporterTabManager;
+	
+	const auto ToggleTabVisibility = [](TWeakPtr<FTabManager> InTGExportManagerWeak, FName InTabName)
+	{
+		TSharedPtr<FTabManager> InTGExportManager = InTGExportManagerWeak.Pin();
+		if (InTGExportManager.IsValid())
+		{
+			TSharedPtr<SDockTab> ExistingTab = InTGExportManager->FindExistingLiveTab(InTabName);
+			if (ExistingTab.IsValid())
+			{
+				ExistingTab->RequestCloseTab();
+			}
+			else
+			{
+				InTGExportManager->TryInvokeTab(InTabName);
+			}
+		}
+	};
+	
+	const auto IsTabVisible = [](TWeakPtr<FTabManager> InTGExportManagerWeak, FName InTabName)
+	{
+		TSharedPtr<FTabManager> InTGExportManager = InTGExportManagerWeak.Pin();
+		if (InTGExportManager.IsValid())
+		{
+			return InTGExportManager->FindExistingLiveTab(InTabName).IsValid();
+		}
+		return false;
+	};
+
+	auto CurrentViewportTab = TGExporterTabManager->FindExistingLiveTab(FTG_EditorTabs::ViewportTabId);
+	bool bViewportIsOff = CurrentViewportTab == nullptr;
+	
+	// check here if 3d viewport is turned off, we need to turn it on temporarily to initialize our systems correctly
+	if (bViewportIsOff)
+	{
+		CurrentViewportTab = TGExporterTabManager->TryInvokeTab(FTG_EditorTabs::ViewportTabId);	
+	}
+	
+	Impl->SetViewportPreviewMesh();
+
+	CommandList->MapAction(
+		FTG_ExporterCommands::Get().Show3DPreview,
+		FExecuteAction::CreateStatic(
+			ToggleTabVisibility,
+			TGExportManagerWeak,
+			FTG_EditorTabs::ViewportTabId
+		),
+		FCanExecuteAction::CreateStatic(
+			[]() { return true; }
+		),
+		FIsActionChecked::CreateStatic(
+			IsTabVisible,
+			TGExportManagerWeak,
+			FTG_EditorTabs::ViewportTabId
+		)
+	);
+	
+	CommandList->MapAction(
+		FTG_ExporterCommands::Get().Show3DPreviewSettings,
+		FExecuteAction::CreateStatic(
+			ToggleTabVisibility,
+			TGExportManagerWeak,
+			FTG_EditorTabs::PreviewSettingsTabId
+		),
+		FCanExecuteAction::CreateStatic(
+			[]() { return true; }
+		),
+		FIsActionChecked::CreateStatic(
+			IsTabVisible,
+			TGExportManagerWeak,
+			FTG_EditorTabs::PreviewSettingsTabId
+		)
+	);
+
+	CommandList->MapAction(
+			FTG_ExporterCommands::Get().ShowParameters,
+			FExecuteAction::CreateStatic(
+				ToggleTabVisibility,
+				TGExportManagerWeak,
+				FTG_EditorTabs::ParameterDefaultsTabId
+			),
+			FCanExecuteAction::CreateStatic(
+				[]() { return true; }
+			),
+			FIsActionChecked::CreateStatic(
+				IsTabVisible,
+				TGExportManagerWeak,
+				FTG_EditorTabs::ParameterDefaultsTabId
+			)
+		);
+
+	CommandList->MapAction(
+			FTG_ExporterCommands::Get().ShowOutputPreview,
+			FExecuteAction::CreateStatic(
+				ToggleTabVisibility,
+				TGExportManagerWeak,
+				FTG_EditorTabs::NodePreviewTabId
+			),
+			FCanExecuteAction::CreateStatic(
+				[]() { return true; }
+			),
+			FIsActionChecked::CreateStatic(
+				IsTabVisible,
+				TGExportManagerWeak,
+				FTG_EditorTabs::NodePreviewTabId
+			)
+		);
+
+	CommandList->MapAction(
+			FTG_ExporterCommands::Get().ShowExportSettings,
+			FExecuteAction::CreateStatic(
+				ToggleTabVisibility,
+				TGExportManagerWeak,
+				FTG_EditorTabs::OutputTabId
+			),
+			FCanExecuteAction::CreateStatic(
+				[]() { return true; }
+			),
+			FIsActionChecked::CreateStatic(
+				IsTabVisible,
+				TGExportManagerWeak,
+				FTG_EditorTabs::OutputTabId
+			)
+		);
+	
+	FMenuBarBuilder MenuBarBuilder(CommandList);
+	MenuBarBuilder.AddPullDownMenu(
+		LOCTEXT("WindowMenuLabel", "Window"),
+		FText::GetEmpty(),
+		FNewMenuDelegate::CreateLambda([](FMenuBuilder& Builder) {
+			Builder.AddMenuEntry(FTG_ExporterCommands::Get().ShowOutputPreview);
+			Builder.AddMenuEntry(FTG_ExporterCommands::Get().ShowParameters);
+			Builder.AddMenuEntry(FTG_ExporterCommands::Get().Show3DPreviewSettings);
+			Builder.AddMenuEntry(FTG_ExporterCommands::Get().Show3DPreview);
+			Builder.AddMenuEntry(FTG_ExporterCommands::Get().ShowExportSettings);
+			})
+	);
+	
+	
+	TSharedRef<SWidget> MenuBarWidget = MenuBarBuilder.MakeWidget();
+	
+	
+	NomadTab->SetContent(
+		SNew(SVerticalBox)
+		+ SVerticalBox::Slot()
+		.AutoHeight()
+		[
+			MenuBarWidget
+		]
+		+SVerticalBox::Slot()
+		[
+			SNew(SBorder)
+			.BorderImage(FAppStyle::GetBrush("ToolPanel.DarkGroupBorder"))
+			.Padding(FMargin(0.f, 2.f))
+			[
+				TabContents
+			]
+		]
+	);
+	
+	
+	// Tell tab-manager about the multi-box for platforms with a global menu bar
+	TGExporterTabManager->SetMenuMultiBox(MenuBarBuilder.GetMultiBox(), MenuBarWidget);
+	
+	return NomadTab;
 }
 
-void FTG_Exporter::SetTextureGraphToExport(UTextureGraph* InTextureGraph)
+void FTG_ExporterUtility::SetTextureGraphToExport(UTextureGraphBase* InTextureGraph)
 {
-	Impl->SetTextureGraphToExport(InTextureGraph);
+	if (InTextureGraph->IsA<UTextureGraph>())
+	{
+		TextureGraphInstance = NewObject<UTextureGraphInstance>();
+		TextureGraphInstance->Construct(FString());
+		TextureGraphInstance->SetParent(Cast<UTextureGraph>(InTextureGraph));
+	}
+	else if (InTextureGraph->IsA<UTextureGraphInstance>())
+	{
+		TextureGraphInstance = Cast<UTextureGraphInstance>(InTextureGraph);
+	}
+	
+	// force open export window
+	FGlobalTabmanager::Get()->TryInvokeTab(FTG_EditorTabs::TextureExporterTabId);
+	
+	
+	
+	Impl->SetTextureGraphToExport(TextureGraphInstance);
+	
 }
 
 #undef LOCTEXT_NAMESPACE 
diff --git a/Engine/Plugins/Experimental/TextureGraph/Source/TextureGraphEditor/Private/TG_InstanceEditor.cpp b/Engine/Plugins/Experimental/TextureGraph/Source/TextureGraphEditor/Private/TG_InstanceEditor.cpp
new file mode 100644
--- /dev/null
+++ b/Engine/Plugins/Experimental/TextureGraph/Source/TextureGraphEditor/Private/TG_InstanceEditor.cpp
@@ -0,0 +1,210 @@
+// Copyright Epic Games, Inc. All Rights Reserved.
+#include "TG_InstanceEditor.h"
+#include "Misc/MessageDialog.h"
+
+#include "TextureGraphEditorModule.h"
+#include "TextureGraph.h"
+#include "TG_EditorTabs.h"
+#include "STG_EditorViewport.h"
+#include "Components/MeshComponent.h"
+#include "Components/StaticMeshComponent.h"
+#include "DetailsViewArgs.h"
+#include "Editor.h"
+
+#include "Framework/Application/SlateApplication.h"
+#include "PropertyEditorModule.h"
+#include "ScopedTransaction.h"
+#include "ToolMenus.h"
+#include "Widgets/Docking/SDockTab.h"
+#include "Widgets/Layout/SBox.h"
+#include "Widgets/Text/STextBlock.h"
+#include <Toolkits/AssetEditorToolkit.h>
+
+#include "EdGraph/TG_EditorErrorReporter.h"
+#include "Model/Mix/MixInterface.h"
+#include "UObject/MetaData.h"
+#include "UObject/ObjectSaveContext.h"
+#include "TG_Exporter.h"
+DEFINE_LOG_CATEGORY(LogTextureGraphInstanceEditor);
+#define LOCTEXT_NAMESPACE "TG_InstanceEditor"
+
+void FTG_InstanceEditor::RefreshViewport()
+{
+	InstanceImpl->RefreshViewport();
+}
+
+void FTG_InstanceEditor::RefreshTool()
+{
+	RefreshViewport();
+}
+void FTG_InstanceEditor::SetMesh(class UMeshComponent* InPreviewMesh, class UWorld* InWorld)
+{
+	// TextureGraphInstance->SetEditorMesh(Cast<UStaticMeshComponent>(InPreviewMesh), InWorld).then([this]()
+	// 	{
+	// 		InstanceImpl->GetEditorViewport()->InitRenderModes(TextureGraphInstance);
+	// 	});
+}
+
+void FTG_InstanceEditor::RegisterTabSpawners(const TSharedRef<class FTabManager>& InTabManager)
+{
+	WorkspaceMenuCategory = InTabManager->AddLocalWorkspaceMenuCategory(LOCTEXT("WorkspaceMenu_TG_Editor", "Texture Graph Instance Editor"));
+	auto WorkspaceMenuCategoryRef = WorkspaceMenuCategory.ToSharedRef();
+
+	FAssetEditorToolkit::RegisterTabSpawners(InTabManager);
+
+	InTabManager->RegisterTabSpawner(FTG_EditorTabs::PropertiesTabId, FOnSpawnTab::CreateSP(this, &FTG_InstanceEditor::SpawnTab_TG_Properties))
+		.SetDisplayName(LOCTEXT("DetailsTab", "Details"))
+		.SetGroup( WorkspaceMenuCategoryRef )
+		.SetIcon(FSlateIcon(FAppStyle::GetAppStyleSetName(), "LevelEditor.Tabs.Details"));
+	
+	InstanceImpl->RegisterTabSpawners(InTabManager);
+}
+
+void FTG_InstanceEditor::UnregisterTabSpawners(const TSharedRef<class FTabManager>& InTabManager)
+{
+	InstanceImpl->UnregisterTabSpawners(InTabManager);
+	InTabManager->UnregisterTabSpawner(FTG_EditorTabs::PropertiesTabId);
+	FAssetEditorToolkit::UnregisterTabSpawners(InTabManager);
+}
+TSharedRef<SDockTab> FTG_InstanceEditor::SpawnTab_TG_Properties(const FSpawnTabArgs& Args)
+{
+	check(Args.GetTabId() == FTG_EditorTabs::PropertiesTabId);
+
+	TSharedPtr<SDockTab> DetailsTab = SNew(SDockTab)
+		[
+			DetailsView.ToSharedRef()
+		];
+	return DetailsTab.ToSharedRef();
+}
+
+void FTG_InstanceEditor::OnTextureGraphPreSave(UObject* Object, FObjectPreSaveContext SaveContext)
+{
+	if (TextureGraphInstance == Object)
+		TextureGraphInstance->TriggerUpdate(false);
+}
+
+void FTG_InstanceEditor::InitEditor(const EToolkitMode::Type Mode, const TSharedPtr< class IToolkitHost >& InitToolkitHost, UTextureGraphInstance* InTextureGraph)
+{
+	FPropertyEditorModule& PropertyEditorModule = FModuleManager::GetModuleChecked<FPropertyEditorModule>("PropertyEditor");
+
+	// Initialize the TextureGraphInstance to create the runtime graph
+	if (InTextureGraph)
+	{
+		InTextureGraph->Initialize();
+	}
+	
+	FDetailsViewArgs DetailsViewArgs;
+	DetailsViewArgs.NameAreaSettings = FDetailsViewArgs::HideNameArea;
+	DetailsViewArgs.bShowObjectLabel = false;
+	DetailsViewArgs.bHideSelectionTip = true;
+	DetailsViewArgs.NotifyHook = this;
+	DetailsViewArgs.ColumnWidth = 0.70;
+	DetailsView = PropertyEditorModule.CreateDetailView(DetailsViewArgs);
+	
+	InstanceImpl = MakeUnique<FTG_InstanceImpl>();
+	InstanceImpl->Initialize();
+
+	TextureGraphInstance = InTextureGraph;
+	FAssetEditorToolkit::InitAssetEditor(Mode, InitToolkitHost, TG_EditorAppIdentifier, InstanceImpl->GetDefaultLayout(), /*bCreateDefaultToolbar*/ true, /*bCreateDefaultStandaloneMenu*/ true, InTextureGraph);
+    
+	InstanceImpl->SetTextureGraphToExport(InTextureGraph);
+	
+	DetailsView->SetObject(TextureGraphInstance, true);
+
+// 	// Add analytics tag
+// 	if (FEngineAnalytics::IsAvailable())
+// 	{
+// 		SessionStartTime = FDateTime::Now();
+// 		FEngineAnalytics::GetProvider().RecordEvent(TEXT("Editor.Usage.TextureGraph.SessionStarted"));
+// 	}
+}
+
+FTG_InstanceEditor::FTG_InstanceEditor()
+{
+}
+
+FTG_InstanceEditor::~FTG_InstanceEditor()
+{
+	DetailsView.Reset();
+	InstanceImpl.Reset();
+	TextureGraphInstance = nullptr;
+}
+
+void FTG_InstanceEditor::NotifyPostChange(const FPropertyChangedEvent& PropertyChangedEvent, FProperty* PropertyThatChanged)
+{
+	if(PropertyThatChanged->GetName() == GET_MEMBER_NAME_CHECKED(UTextureGraphInstance, ParentTextureGraph))
+	{
+		TextureGraphInstance->SetParent(TextureGraphInstance->ParentTextureGraph);
+		InstanceImpl->SetTextureGraphToExport(TextureGraphInstance);
+	}
+}
+
+FLinearColor FTG_InstanceEditor::GetWorldCentricTabColorScale() const
+{
+	return FLinearColor(0.3f, 0.2f, 0.5f, 0.5f);
+}
+
+FName FTG_InstanceEditor::GetToolkitFName() const
+{
+	return FName("TG_Editor");
+}
+
+FText FTG_InstanceEditor::GetBaseToolkitName() const
+{
+	return LOCTEXT("AppLabel", "TG_ Editor");
+}
+
+FString FTG_InstanceEditor::GetWorldCentricTabPrefix() const
+{
+	return LOCTEXT("WorldCentricTabPrefix", "TG_ ").ToString();
+}
+
+void FTG_InstanceEditor::AddReferencedObjects(FReferenceCollector& Collector)
+{
+	Collector.AddReferencedObject(TextureGraphInstance);
+}
+
+class UMixInterface* FTG_InstanceEditor::GetTextureGraphInterface() const
+{
+	return TextureGraphInstance;
+}
+
+void FTG_InstanceEditor::Tick(float DeltaTime)
+{
+	RefreshViewport();
+}
+
+TStatId FTG_InstanceEditor::GetStatId() const
+{
+	return TStatId();
+}
+
+FText FTG_InstanceEditor::GetOriginalObjectName() const
+{
+	return FText::FromString(GetEditingObjects()[0]->GetName());
+}
+
+void FTG_InstanceEditor::OnClose()
+{
+	// // Add analytics tag
+	// if (FEngineAnalytics::IsAvailable())
+	// {
+	// 	TArray<FAnalyticsEventAttribute> Attributes;
+	// 	Attributes.Add(FAnalyticsEventAttribute(TEXT("TimeActive.Seconds"),  (FDateTime::Now() - SessionStartTime).GetTotalSeconds()));
+	// 			
+	// 	// Send Analytics event 
+	// 	FEngineAnalytics::GetProvider().RecordEvent(TEXT("Editor.Usage.TextureGraph.SessionEnded"), Attributes);
+	// 			
+	// }
+	if (TextureGraphInstance)
+	{
+		/// We need to flush any invalidations coming for this graph. This is because if the user decided to save
+		/// graph on exit then this queues a mix update that never gets finished as the engine is being shutdown
+		/// and results in a cleanup assertion in Device.cpp
+		TextureGraphInstance->FlushInvalidations();
+	}
+	
+	InstanceImpl = nullptr;
+}
+
+#undef LOCTEXT_NAMESPACE
diff --git a/Engine/Plugins/Experimental/TextureGraph/Source/TextureGraphEditor/Private/TG_InstanceEditor.h b/Engine/Plugins/Experimental/TextureGraph/Source/TextureGraphEditor/Private/TG_InstanceEditor.h
new file mode 100644
--- /dev/null
+++ b/Engine/Plugins/Experimental/TextureGraph/Source/TextureGraphEditor/Private/TG_InstanceEditor.h
@@ -0,0 +1,112 @@
+// Copyright Epic Games, Inc. All Rights Reserved.
+
+#pragma once
+
+#include "CoreMinimal.h"
+#include "AdvancedPreviewSceneModule.h"
+#include "ITG_Editor.h"
+#include "EdGraph/EdGraphPin.h"
+#include "EditorUndoClient.h"
+#include "Misc/NotifyHook.h"
+#include "AssetEditorViewportLayout.h"
+#include "EditorViewportTabContent.h"
+#include "IDetailsView.h"
+#include "Export/TextureExporter.h"
+
+#include "TG_Parameter.h"
+#include "EdGraph/TG_EdGraphNode.h"
+#include "STG_Palette.h"
+#include "GraphEditor.h"
+
+class UTextureGraphInstance;
+class UTextureGraph;
+class UTG_Graph;
+class UTG_Node;
+class UTG_EdGraph;
+class SDockTab;
+class FSpawnTabArgs;
+class FTabManager;
+
+class UTG_Expression;
+class FUICommandList;
+class FObjectPreSaveContext;
+struct FGraphAppearanceInfo;
+
+DECLARE_LOG_CATEGORY_EXTERN(LogTextureGraphInstanceEditor, Log, All);
+
+class FTG_InstanceEditor : public ITG_Editor, public FGCObject, public FTickableGameObject, public FEditorUndoClient, public FNotifyHook
+{
+protected:
+	virtual void RefreshViewport() override;
+	virtual void RefreshTool() override;
+
+public:
+	virtual void SetMesh(UMeshComponent* PreviewMesh, UWorld* World) override;
+
+	virtual void									RegisterTabSpawners(const TSharedRef<class FTabManager>& InTabManager) override;
+	virtual void									UnregisterTabSpawners(const TSharedRef<class FTabManager>& InTabManager) override;
+	void											OnTextureGraphPreSave(UObject* Object, FObjectPreSaveContext SaveContext);
+	
+	TSharedRef<SDockTab>							SpawnTab_TG_Properties(const FSpawnTabArgs& Args);
+	/**
+	 * Edits the specified static TSX Asset object
+	 *
+	 * @param	Mode								Asset editing mode for this editor (standalone or world-centric)
+	 * @param	InitToolkitHost						When Mode is WorldCentric, this is the level editor instance to spawn this editor within
+	 * @param	ObjectToEdit						The TSX Asset to edit
+	 */
+	void											InitEditor(const EToolkitMode::Type Mode, const TSharedPtr< class IToolkitHost >& InitToolkitHost, UTextureGraphInstance* TG_Script);
+
+	/** Constructor */
+													FTG_InstanceEditor();
+
+	virtual											~FTG_InstanceEditor() override;
+
+	/** Post edit change notify for properties. */
+	virtual void NotifyPostChange( const FPropertyChangedEvent& PropertyChangedEvent, FProperty* PropertyThatChanged ) override;
+	
+	// Inherited via ITG_Editor
+	virtual FLinearColor							GetWorldCentricTabColorScale() const override;
+	virtual FName									GetToolkitFName() const override;
+	virtual FText									GetBaseToolkitName() const override;
+	virtual FString									GetWorldCentricTabPrefix() const override;
+	
+	// Inherited via FGCObject
+	virtual void									AddReferencedObjects(FReferenceCollector& Collector) override;
+
+	virtual FString									GetReferencerName() const override { return TEXT("FTextureScriptEditor");}
+
+	class UMixInterface*							GetTextureGraphInterface() const override;
+
+	// Inherited via FTickableGameObject
+	virtual void									Tick(float DeltaTime) override;
+
+	virtual ETickableTickType						GetTickableTickType() const override { return ETickableTickType::Always; }
+
+	virtual bool									IsTickableWhenPaused() const override { return true; }
+
+	virtual bool									IsTickableInEditor() const override { return true; }
+
+	virtual TStatId									GetStatId() const override;
+
+	virtual FText									GetOriginalObjectName() const override;
+
+
+	// ~Begin FAssetEditorToolkit interface
+	virtual void									OnClose() override;
+	// ~End FAssetEditorToolkit interface
+
+	// Widget Accessors
+	TSharedRef<class IDetailsView>					GetDetailView() const { return DetailsView.ToSharedRef(); }
+private:
+	//	FDateTime										SessionStartTime;
+
+private:
+	/** Property View */
+	TSharedPtr<class IDetailsView>					DetailsView;
+	
+	TUniquePtr<struct FTG_InstanceImpl>				InstanceImpl;
+	
+	TObjectPtr<UTextureGraphInstance>				TextureGraphInstance;
+	
+};
diff --git a/Engine/Plugins/Experimental/TextureGraph/Source/TextureGraphEditor/Private/TG_InstanceFactory.cpp b/Engine/Plugins/Experimental/TextureGraph/Source/TextureGraphEditor/Private/TG_InstanceFactory.cpp
new file mode 100644
--- /dev/null
+++ b/Engine/Plugins/Experimental/TextureGraph/Source/TextureGraphEditor/Private/TG_InstanceFactory.cpp
@@ -0,0 +1,22 @@
+// Copyright Epic Games, Inc. All Rights Reserved.
+
+#include "TG_InstanceFactory.h"
+#include "Model/ModelObject.h"
+#include "TextureGraph.h"
+
+UTG_InstanceFactory::UTG_InstanceFactory()
+{
+	SupportedClass = UTextureGraphInstance::StaticClass();
+	bCreateNew = true;
+	bEditAfterNew = true;
+}
+
+UObject* UTG_InstanceFactory::FactoryCreateNew(UClass* Class, UObject* InParent, FName Name, EObjectFlags Flags, UObject* Context, FFeedbackContext* Warn)
+{
+	UTextureGraphInstance* TextureGraphInstance = NewObject<UTextureGraphInstance>(InParent, Class, Name, Flags, Context);
+
+	check(TextureGraphInstance);
+	TextureGraphInstance->Construct(Name.GetPlainNameString());
+	TextureGraphInstance->SetParent(InitialParent);
+	return TextureGraphInstance;
+}
diff --git a/Engine/Plugins/Experimental/TextureGraph/Source/TextureGraphEditor/Private/TG_InstanceFactory.h b/Engine/Plugins/Experimental/TextureGraph/Source/TextureGraphEditor/Private/TG_InstanceFactory.h
new file mode 100644
--- /dev/null
+++ b/Engine/Plugins/Experimental/TextureGraph/Source/TextureGraphEditor/Private/TG_InstanceFactory.h
@@ -0,0 +1,19 @@
+// Copyright Epic Games, Inc. All Rights Reserved.
+
+#pragma once
+
+#include "Factories/Factory.h"
+#include "TG_InstanceFactory.generated.h"
+
+
+UCLASS()
+class UTG_InstanceFactory : public UFactory
+{
+	GENERATED_BODY()
+public:
+	UPROPERTY()
+	TObjectPtr<class UTextureGraph> InitialParent;
+	
+	UTG_InstanceFactory();
+	UObject* FactoryCreateNew(UClass* Class, UObject* InParent, FName Name, EObjectFlags Flags, UObject* Context, FFeedbackContext* Warn);
+};
diff --git a/Engine/Plugins/Experimental/TextureGraph/Source/TextureGraphEditor/Private/TG_RenderModeManager.cpp b/Engine/Plugins/Experimental/TextureGraph/Source/TextureGraphEditor/Private/TG_RenderModeManager.cpp
--- a/Engine/Plugins/Experimental/TextureGraph/Source/TextureGraphEditor/Private/TG_RenderModeManager.cpp
+++ b/Engine/Plugins/Experimental/TextureGraph/Source/TextureGraphEditor/Private/TG_RenderModeManager.cpp
@@ -22,14 +22,14 @@ TG_RenderModeManager::~TG_RenderModeManager()
 	
 }
 
-void TG_RenderModeManager::ChangeRenderMode(FName NewRenderMode, UTextureGraph* TextureGraph)
+void TG_RenderModeManager::ChangeRenderMode(FName NewRenderMode, UTextureGraphBase* TextureGraph)
 {
 	_lastRenderMode = _currentRenderMode;
 	_currentRenderMode = NewRenderMode;
 	UpdateRenderMode(TextureGraph);
 }
 
-void TG_RenderModeManager::UpdateRenderMode(UTextureGraph* TextureGraph)
+void TG_RenderModeManager::UpdateRenderMode(UTextureGraphBase* TextureGraph)
 {
 	UMixSettings* settings = TextureGraph->GetSettings();
 	size_t NumTargets = settings->NumTargets();
@@ -91,7 +91,7 @@ bool TG_RenderModeManager::IsCurrentRenderModelLit(int targetId /* = 0*/)
 }
 
 
-void TG_RenderModeManager::InitializeDefaultMaterials(int totalTargets, UTextureGraph* TextureGraph)
+void TG_RenderModeManager::InitializeDefaultMaterials(int totalTargets, UTextureGraphBase* TextureGraph)
 {
 	_renderModeMaterials.Empty();
 
diff --git a/Engine/Plugins/Experimental/TextureGraph/Source/TextureGraphEditor/Private/TG_RenderModeManager.h b/Engine/Plugins/Experimental/TextureGraph/Source/TextureGraphEditor/Private/TG_RenderModeManager.h
--- a/Engine/Plugins/Experimental/TextureGraph/Source/TextureGraphEditor/Private/TG_RenderModeManager.h
+++ b/Engine/Plugins/Experimental/TextureGraph/Source/TextureGraphEditor/Private/TG_RenderModeManager.h
@@ -5,7 +5,7 @@
 #include <memory>
 #include "TG_RenderModeManager.generated.h"
 
-class UTextureGraph;
+class UTextureGraphBase;
 
 UENUM()
 enum class ERenderModes : uint8
@@ -42,14 +42,14 @@ protected:
 	TMap<int, MaterialMap>			_renderModeMaterials;
 	TArray<FName>					_renderModesString;
 
-	void							InitializeDefaultMaterials(int totalTargets, UTextureGraph* TextureGraph);
+	void							InitializeDefaultMaterials(int totalTargets, UTextureGraphBase* TextureGraph);
 	RenderMaterial_BPPtr			GetTargetRenderModeMaterial(int targetId, FName renderMode);
 	void							BindBlobToMaterial(RenderMaterialPtr renderMaterial, TiledBlobPtr blobToBind, const FName& targetName);
-	virtual void					UpdateRenderMode(UTextureGraph* TextureGraph);
+	virtual void					UpdateRenderMode(UTextureGraphBase* TextureGraph);
 public:
 									TG_RenderModeManager();
 	virtual							~TG_RenderModeManager();		
-	virtual void					ChangeRenderMode(FName NewRenderMode, UTextureGraph* TextureGraph);
+	virtual void					ChangeRenderMode(FName NewRenderMode, UTextureGraphBase* TextureGraph);
 	void							Clear();
 	bool							IsCurrentRenderModelLit(int targetId = 0);
 	//////////////////////////////////////////////////////////////////////////
diff --git a/Engine/Plugins/Experimental/TextureGraph/Source/TextureGraphEditor/Private/TextureGraphEditorModule.cpp b/Engine/Plugins/Experimental/TextureGraph/Source/TextureGraphEditor/Private/TextureGraphEditorModule.cpp
--- a/Engine/Plugins/Experimental/TextureGraph/Source/TextureGraphEditor/Private/TextureGraphEditorModule.cpp
+++ b/Engine/Plugins/Experimental/TextureGraph/Source/TextureGraphEditor/Private/TextureGraphEditorModule.cpp
@@ -7,7 +7,9 @@
 #include "EdGraph/TG_EditorGraphPanelPinFactory.h"
 #include "TextureGraphEngine.h"
 #include "AssetDefinition_TextureGraph.h"
+#include "AssetDefinition_TextureGraphInstance.h"
 #include "TG_Editor.h"
+#include "TG_InstanceEditor.h"
 #include "TG_EditorCommands.h"
 #include "TG_Style.h"
 #include "EdGraph/TG_EditorErrorReporter.h"
@@ -45,6 +47,7 @@ void FTextureGraphEditorModule::StartupModule()
 	// Register slate style overrides
 	FTG_Style::Register();
 	FTG_EditorCommands::Register();
+	FTG_ExporterCommands::Register();
 	
 	GraphNodeFactory = MakeShareable(new FTG_EditorGraphNodeFactory());
 	FEdGraphUtilities::RegisterVisualNodeFactory(GraphNodeFactory);
@@ -52,7 +55,7 @@ void FTextureGraphEditorModule::StartupModule()
 	GraphPanelPinFactory = MakeShared<FTG_EditorGraphPanelPinFactory>();
 	FEdGraphUtilities::RegisterVisualPinFactory(GraphPanelPinFactory);
 	
-	TG_Exporter = MakeUnique<FTG_Exporter>();
+	TG_Exporter = MakeUnique<FTG_ExporterUtility>();
 	
 	StartTextureGraphEngine();
 }
@@ -63,6 +66,7 @@ void FTextureGraphEditorModule::ShutdownModule()
 	// we call this function before unloading the module.
 	UnRegisterAllAssetTypeActions();
 	FTG_EditorCommands::Unregister();
+	FTG_ExporterCommands::Unregister();
 	
 	FPropertyEditorModule& PropertyEditorModule = FModuleManager::GetModuleChecked<FPropertyEditorModule>("PropertyEditor");
 	// Unregister the details customization
@@ -143,6 +147,13 @@ TSharedRef<ITG_Editor> FTextureGraphEditorModule::CreateTextureGraphEditor(const
 	
 	return NewEditor;
 }
+TSharedRef<ITG_Editor> FTextureGraphEditorModule::CreateTextureGraphInstanceEditor(const EToolkitMode::Type Mode, const TSharedPtr<IToolkitHost>& InitToolkitHost, UTextureGraphInstance* InTextureGraphInstance)
+{
+	TSharedRef<FTG_InstanceEditor> NewEditor(new FTG_InstanceEditor());
+	NewEditor->InitEditor(Mode, InitToolkitHost, InTextureGraphInstance);
+	
+	return NewEditor;
+}
 
 #undef LOCTEXT_NAMESPACE
 
diff --git a/Engine/Plugins/Experimental/TextureGraph/Source/TextureGraphEditor/Public/TG_Exporter.h b/Engine/Plugins/Experimental/TextureGraph/Source/TextureGraphEditor/Public/TG_Exporter.h
--- a/Engine/Plugins/Experimental/TextureGraph/Source/TextureGraphEditor/Public/TG_Exporter.h
+++ b/Engine/Plugins/Experimental/TextureGraph/Source/TextureGraphEditor/Public/TG_Exporter.h
@@ -3,27 +3,151 @@
 #pragma once
 
 #include "Templates/UniquePtr.h"
-
-class UTextureGraph;
-
-struct FTG_Exporter
+#include "Tickable.h"
+#include "Framework/Commands/Commands.h"
+#include "Framework/Docking/TabManager.h"
+#include "Widgets/Input/SComboBox.h"
+#define LOCTEXT_NAMESPACE "TextureGraphExporter"
+class UTextureGraphInstance;
+class UTextureGraphBase;
+class FUICommandInfo;
+struct FTG_ExporterUtility
 {
 	// Initializes the global state of the exporter (commands, tab spawners, etc):
-	FTG_Exporter();
+	FTG_ExporterUtility();
 
 	// Destructor declaration purely so that we can pimpl:
-	~FTG_Exporter();
+	~FTG_ExporterUtility();
 
 	/** Sets the current texture graph to be used with the exporter */
-	void SetTextureGraphToExport(UTextureGraph* InTextureGraph);
+	void SetTextureGraphToExport(UTextureGraphBase* InTextureGraph);
+
+	/** Function registered with tab manager to create the Texture Graph Exporter */
+	TSharedRef<SDockTab> CreateTGExporterTab(const FSpawnTabArgs& Args);
+	
+
+private:
+	TUniquePtr<struct FTG_InstanceImpl>	Impl;
+	TObjectPtr<UTextureGraphInstance> TextureGraphInstance;
+	TSharedPtr<FTabManager> TGExporterTabManager;
+	TSharedPtr<FTabManager::FLayout> TGExporterLayout;
+	
+	// prevent copying:
+	FTG_ExporterUtility(const FTG_ExporterUtility&);
+	FTG_ExporterUtility(FTG_ExporterUtility&&);
+	FTG_ExporterUtility& operator=(FTG_ExporterUtility const&);
+	FTG_ExporterUtility& operator=(FTG_ExporterUtility&&);
+};
+
+
+struct FTG_ExporterCommands : public TCommands<FTG_ExporterCommands>
+{
+	FTG_ExporterCommands()
+		: TCommands<FTG_ExporterCommands>(
+			TEXT("TextureGraphExporter"), // Context name for fast lookup
+			LOCTEXT("TextureGraphExporter", "Texture Graph Exporter"), // Localized context name for displaying
+			NAME_None, // Parent
+			FCoreStyle::Get().GetStyleSetName() // Icon Style Set
+		)
+	{
+	}
+
+	// TCommand<> interface
+	virtual void RegisterCommands() override;
+	// End of TCommand<> interface
+
+	TSharedPtr<FUICommandInfo> ShowOutputPreview;
+	TSharedPtr<FUICommandInfo> Show3DPreview;
+	TSharedPtr<FUICommandInfo> Show3DPreviewSettings;
+	TSharedPtr<FUICommandInfo> ShowParameters;
+	TSharedPtr<FUICommandInfo> ShowExportSettings;
+
+};
+class UTG_Node;
+class UTG_Graph;
+
+struct FTG_InstanceImpl: public FTickableGameObject, public FGCObject
+{
+	friend class FTG_Exporter;
+	friend class FTG_InstanceEditor;
+	friend class UTextureGraphInstance;
+public:
+													FTG_InstanceImpl();
+	virtual											~FTG_InstanceImpl() override;
+
+	void											Cleanup();
+	/** Sets the Texture Graph to be exported in the Exporter */
+	void 											SetTextureGraphToExport(UTextureGraphInstance* InTextureGraph);
+	void 											OnGraphChanged(UTG_Graph* InGraph, UTG_Node* InNode, bool Tweaking);
+	TSharedRef<FTabManager::FLayout>				GetDefaultLayout();
+	void 											Initialize();
+	void 											RegisterTabSpawners(const TSharedPtr<FTabManager>);
+	void 											UnregisterTabSpawners(const TSharedPtr<FTabManager>);
+	void 											SetMesh(class UMeshComponent* InPreviewMesh, class UWorld* InWorld);
+	void 											SetViewportPreviewMesh();
+	void											UpdatePreviewMesh();
+
+private:
+	TObjectPtr<UTextureGraphInstance>				TextureGraphInstance;
+	TWeakPtr<class STG_NodePreviewWidget>			NodePreviewPtr;
+	
+	TSharedPtr<class IDetailsView>					DetailsView;
+	TSharedPtr<class IDetailsView>					ParametersView;
+	TSharedPtr<class IDetailsView>					ExportSettingsView;
+	TSharedPtr<class IDetailsView>					PreviewSettingsView;
+// Tracking the active viewports in this editor.
+	TSharedPtr<class FEditorViewportTabContent>		ViewportTabContent;
+
+	TObjectPtr<class UTG_Parameters>				Parameters;
+	TObjectPtr<class UTG_ExportSettings>			ExportSettings;
+	TSharedPtr<class FExportSettings>				TargetExportSettings;
+	
+	TArray<TSharedPtr<FName>>						OutputNodesList;
+	TSharedPtr<SComboBox<TSharedPtr<FName>>>		OutputNodesComboBoxWidget;
+	UTG_Node*										SelectedNode = nullptr;
+	// Inherited via FTickableGameObject
+	virtual void									Tick(float DeltaTime) override;
+
+	virtual ETickableTickType						GetTickableTickType() const override { return ETickableTickType::Always; }
+
+	virtual bool									IsTickableWhenPaused() const override { return true; }
+
+	virtual bool									IsTickableInEditor() const override { return true; }
+
+	virtual TStatId									GetStatId() const override;
+
+	// Inherited via FGCObject
+	virtual void									AddReferencedObjects(FReferenceCollector& Collector) override;
+	virtual FString									GetReferencerName() const override { return TEXT("FTextureGraphExporter");}
 
 private:
-	TUniquePtr< struct FTG_ExporterImpl > Impl;
+	TSharedRef<SDockTab>							SpawnTab_Viewport(const FSpawnTabArgs& Args);
+	TSharedRef<SDockTab>							SpawnTab_ParameterDefaults(const FSpawnTabArgs& Args);
+	TSharedRef<SDockTab>							SpawnTab_NodePreview(const FSpawnTabArgs& Args);
+	TSharedRef<SDockTab>							SpawnTab_ExportSettings(const FSpawnTabArgs& Args);
+	TSharedRef<SDockTab>							SpawnTab_PreviewSettings(const FSpawnTabArgs& Args);
+	TSharedPtr<class STG_EditorViewport>			GetEditorViewport() const;
+	
+	bool 											SetPreviewAsset(UObject* InAsset);
+	TSharedRef<class IDetailsView>					GetDetailsView() const { return DetailsView.ToSharedRef(); }
+	TSharedRef<class IDetailsView>					GetPreviewSettingsView() const { return PreviewSettingsView.ToSharedRef(); }
+	TSharedRef<class IDetailsView>					GetExportSettingsView() const { return ExportSettingsView.ToSharedRef(); }
+	TSharedRef<class IDetailsView>					GetParametersView() const { return ParametersView.ToSharedRef(); }
+	void											OnViewportMaterialChanged();
+	void											OnMaterialMappingChanged();
+	void											UpdateExportSettingsUI();
+	void											UpdateParametersUI();
+	virtual void									RefreshViewport();
+	void											OnRenderingDone(class UMixInterface* TextureGraph, const struct FInvalidationDetails* Details);
+	TSharedRef<SWidget>								GenerateOutputComboItem(TSharedPtr<FName> String);
+	void											OnOutputSelectionChanged(TSharedPtr<FName> String, ESelectInfo::Type Arg);
+	FReply											OnExportClicked(EAppReturnType::Type ButtonID);
 
 	// prevent copying:
-	FTG_Exporter(const FTG_Exporter&);
-	FTG_Exporter(FTG_Exporter&&);
-	FTG_Exporter& operator=(FTG_Exporter const&);
-	FTG_Exporter& operator=(FTG_Exporter&&);
+	FTG_InstanceImpl(const FTG_InstanceImpl&);
+	FTG_InstanceImpl(FTG_InstanceImpl&&);
+	FTG_InstanceImpl& operator=(FTG_InstanceImpl const&);
+	FTG_InstanceImpl& operator=(FTG_InstanceImpl&&);
 };
 
+#undef LOCTEXT_NAMESPACE 
diff --git a/Engine/Plugins/Experimental/TextureGraph/Source/TextureGraphEditor/Public/TextureGraphEditorModule.h b/Engine/Plugins/Experimental/TextureGraph/Source/TextureGraphEditor/Public/TextureGraphEditorModule.h
--- a/Engine/Plugins/Experimental/TextureGraph/Source/TextureGraphEditor/Public/TextureGraphEditorModule.h
+++ b/Engine/Plugins/Experimental/TextureGraph/Source/TextureGraphEditor/Public/TextureGraphEditorModule.h
@@ -9,6 +9,7 @@ class IAssetTools;
 class IAssetTypeActions;
 class ITG_Editor;
 class UTextureGraph;
+class UTextureGraphInstance;
 class FTG_EditorGraphNodeFactory;
 class FTG_EditorGraphPanelPinFactory;
 typedef TArray<TSharedPtr<IAssetTypeActions>> AssetTypeActionsArray;
@@ -23,7 +24,7 @@ private:
 	FTickerDelegate					TickDelegate;
 	/** Handle of the delegate to run the Tick method in charge of running TextureGraphEngine update*/
 	FTSTicker::FDelegateHandle		TickDelegateHandle;
-	TUniquePtr<FTG_Exporter>		TG_Exporter;
+	TUniquePtr<FTG_ExporterUtility>		TG_Exporter;
 
 protected:
 	TSharedPtr<FTG_EditorGraphNodeFactory> GraphNodeFactory;
@@ -41,7 +42,8 @@ public:
 	void							RegisterAssetTypeAction(IAssetTools& AssetTools, TSharedRef<IAssetTypeActions> Action);
 	void							UnRegisterAllAssetTypeActions();
 	static TSharedRef<ITG_Editor>	CreateTextureGraphEditor(const EToolkitMode::Type Mode, const TSharedPtr< class IToolkitHost >& InitToolkitHost, UTextureGraph* InTextureGraph);
+	static TSharedRef<ITG_Editor>	CreateTextureGraphInstanceEditor(const EToolkitMode::Type Mode, const TSharedPtr< class IToolkitHost >& InitToolkitHost, UTextureGraphInstance* InTextureGraphInstance);
 
 	/** Returns a reference to the Blueprint Debugger state object */
-	const TUniquePtr<FTG_Exporter>& GetTextureExporter() const { return TG_Exporter; }
+	const TUniquePtr<FTG_ExporterUtility>& GetTextureExporter() const { return TG_Exporter; }
 };
