#UE Quick fix for cook error

diff --git a/Engine/Plugins/Experimental/TextureGraph/Source/TextureGraphEditor/Private/TG_Exporter.cpp b/Engine/Plugins/Experimental/TextureGraph/Source/TextureGraphEditor/Private/TG_Exporter.cpp
--- a/Engine/Plugins/Experimental/TextureGraph/Source/TextureGraphEditor/Private/TG_Exporter.cpp
+++ b/Engine/Plugins/Experimental/TextureGraph/Source/TextureGraphEditor/Private/TG_Exporter.cpp
@@ -84,12 +84,12 @@ FTG_InstanceImpl::~FTG_InstanceImpl()
 	ParametersView.Reset();
 
 	// cleanup UI
-	if (Parameters->IsValidLowLevelFast())
+	if (Parameters && Parameters->IsValidLowLevelFast())
 	{
 		Parameters->Parameters.Empty();
 		Parameters = nullptr;
 	}
-	if (ExportSettings->IsValidLowLevelFast())
+	if (ExportSettings && ExportSettings->IsValidLowLevelFast())
 	{
 		ExportSettings->OutputExpressionsInfos.Empty();
 		ExportSettings = nullptr;
