Texture Graph: Instance to Instance workflow
Fixed Material Expression not working correctly with TG Instance
Added Preview scene settings tab in Exporter/Instance Editor
Fixed slate cleanup issues with Exporter when Editor is shut down
Fixes issues with TG_Expression_TextureGraph not working after latest updates

diff --git a/Engine/Plugins/Experimental/TextureGraph/Source/TextureGraph/Private/Expressions/Input/TG_Expression_Graph.cpp b/Engine/Plugins/Experimental/TextureGraph/Source/TextureGraph/Private/Expressions/Input/TG_Expression_Graph.cpp
--- a/Engine/Plugins/Experimental/TextureGraph/Source/TextureGraph/Private/Expressions/Input/TG_Expression_Graph.cpp
+++ b/Engine/Plugins/Experimental/TextureGraph/Source/TextureGraph/Private/Expressions/Input/TG_Expression_Graph.cpp
@@ -171,13 +171,13 @@ void UTG_Expression_TextureGraph::OnTextureGraphPreSave(UObject* Object, FObject
 	// every editor should check if your texture graph is dependent on the object being saved
 	UTextureGraph* GraphBeingSaved = Cast<UTextureGraph>(Object);
 	UTextureGraph* ParentGraph = GetTypedOuter<UTextureGraph>();
+
 	// if object being saved is our linked TextureGraph, 
 	// and this expression's graph is NOT the same as the graph being saved, we re-create
-	if (GraphBeingSaved && TextureGraph && 
+	if (GraphBeingSaved && TextureGraph && ParentGraph &&
 		GraphBeingSaved != ParentGraph )
 	{
 		
-		check (ParentGraph);
 		// if graph being saved is not ourself, and there is a cyclic dependency
 		bool bSameGraphs = GraphBeingSaved->GetOuter() == TextureGraph->GetOuter();
 		if (bSameGraphs || ParentGraph->IsDependentOn(GraphBeingSaved))
@@ -353,7 +353,7 @@ void UTG_Expression_TextureGraph::Evaluate(FTG_EvaluationContext* InContext)
 bool UTG_Expression_TextureGraph::Validate(MixUpdateCyclePtr Cycle)
 {
 	// if TextureGraph is already in execution list then
-	UTextureGraph* ParentGraph = CastChecked<UTextureGraph>(GetOutermostObject());
+	UTextureGraph* ParentGraph = Cast<UTextureGraph>(GetOutermostObject());
 
 	if((TextureGraph &&  Cycle->ContainsMix(TextureGraph)) ||
 		(ParentGraph && ParentGraph->HasCyclicDependency()))
diff --git a/Engine/Plugins/Experimental/TextureGraph/Source/TextureGraph/Private/Expressions/TG_Expression.cpp b/Engine/Plugins/Experimental/TextureGraph/Source/TextureGraph/Private/Expressions/TG_Expression.cpp
--- a/Engine/Plugins/Experimental/TextureGraph/Source/TextureGraph/Private/Expressions/TG_Expression.cpp
+++ b/Engine/Plugins/Experimental/TextureGraph/Source/TextureGraph/Private/Expressions/TG_Expression.cpp
@@ -279,7 +279,7 @@ void UTG_Expression::SetupAndEvaluate(FTG_EvaluationContext* InContext)
 	// Copy Output var values TO this expression's matching properties (Mainly to copy existing thumbs for Output FTG_Textures)
 	for (auto& OutputVar : InContext->Outputs.VarArguments)
 	{
-		OutputVar.Value.Var.CopyTo(this, OutputVar.Value.Argument);
+		OutputVar.Value.Var->CopyTo(this, OutputVar.Value.Argument);
 	}
 
 	// For debug purpose, log the expression evaluation
@@ -294,8 +294,7 @@ void UTG_Expression::SetupAndEvaluate(FTG_EvaluationContext* InContext)
 	// AFTER evaluation, copy output var values FROM this expression's matching properties
 	for (auto& OutputVar : InContext->Outputs.VarArguments)
 	{
-		FTG_Var* OutputVarInGraph = InContext->Graph->GetVar(OutputVar.Value.Var.GetId());
-		OutputVarInGraph->CopyFrom(this, OutputVar.Value.Argument);
+		OutputVar.Value.Var->CopyFrom(this, OutputVar.Value.Argument);
 	}
 }
 
@@ -322,7 +321,7 @@ void UTG_Expression::LogEvaluation(FTG_EvaluationContext* InContext)
 		InVarIds.Emplace(VarMapEntry.Value.Var.GetId());
 	TArray<FTG_Id> OutVarIds;
 	for (auto& VarMapEntry : InContext->Outputs.VarArguments)
-		OutVarIds.Emplace(VarMapEntry.Value.Var.GetId());
+		OutVarIds.Emplace(VarMapEntry.Value.Var->GetId());
 
 	UE_LOG(LogTextureGraph, VeryVerbose, TEXT("%-*s %s"), UTG_Graph::LogHeaderWidth, *GetDefaultName().ToString(), *UTG_Graph::LogCall(InVarIds, OutVarIds));
 }
diff --git a/Engine/Plugins/Experimental/TextureGraph/Source/TextureGraph/Private/TG_GraphEvaluation.cpp b/Engine/Plugins/Experimental/TextureGraph/Source/TextureGraph/Private/TG_GraphEvaluation.cpp
--- a/Engine/Plugins/Experimental/TextureGraph/Source/TextureGraph/Private/TG_GraphEvaluation.cpp
+++ b/Engine/Plugins/Experimental/TextureGraph/Source/TextureGraph/Private/TG_GraphEvaluation.cpp
@@ -100,7 +100,7 @@ void FTG_Evaluation::TransferVarToPin(UTG_Pin* InPin, FTG_EvaluationContext* Con
 		/// This will get copied to anyway when the expression runs. 
 		if (Var->IsArray())
 			Var->Reset();
-		Context->Outputs.VarArguments.Add(Arg.GetName(), { *Var, Arg });
+		Context->Outputs.VarArguments.Add(Arg.GetName(), { Var, Arg });
 	}
 }
 
@@ -616,7 +616,7 @@ void FTG_Evaluation::EvaluateGraph(UTG_Graph* InGraph, FTG_EvaluationContext* In
 			auto TheVar = InGraph->GetVar(ParamId);
 			if (TheVar)
 			{
-				V.Value.Var.ShareData(*TheVar);
+				V.Value.Var->ShareData(*TheVar);
 			}
 		}
 	}
@@ -696,7 +696,7 @@ void FTG_Evaluation::EvaluateNodeArray(UTG_Node* InNode, const TArray<TObjectPtr
 
 				FTG_Evaluation::VarConverterInfo Info;
 				Info.Context = InContext;
-				Info.InVar = &OutputVar.Value.Var;
+				Info.InVar = OutputVar.Value.Var;
 				Info.OutVar = &VariantVar; // Generate the converted version of the connected input into the CoonvertedVar of the Pin
 				Info.Index = -1;
 
@@ -707,7 +707,7 @@ void FTG_Evaluation::EvaluateNodeArray(UTG_Node* InNode, const TArray<TObjectPtr
 			}
 			else
 			{
-				ArrayOutput->Set(VarIndex, OutputVar.Value.Var.GetAs<FTG_Variant>());
+				ArrayOutput->Set(VarIndex, OutputVar.Value.Var->GetAs<FTG_Variant>());
 			}
 		}
 	}
@@ -717,9 +717,9 @@ void FTG_Evaluation::EvaluateNodeArray(UTG_Node* InNode, const TArray<TObjectPtr
 	for (auto& OutputVar : InContext->Outputs.VarArguments)
 	{
 		FTG_VariantArray* ArrayOutput = ArrayOutputs.Find(OutputVar.Key);
-		OutputVar.Value.Var.ResetAs<FTG_VariantArray>();
-		OutputVar.Value.Var.SetAs<FTG_VariantArray>(*ArrayOutput);
-		OutputVar.Value.Var.SetArray();
+		OutputVar.Value.Var->ResetAs<FTG_VariantArray>();
+		OutputVar.Value.Var->SetAs<FTG_VariantArray>(*ArrayOutput);
+		OutputVar.Value.Var->SetArray();
 	}
 }
 
diff --git a/Engine/Plugins/Experimental/TextureGraph/Source/TextureGraph/Private/TG_Texture.cpp b/Engine/Plugins/Experimental/TextureGraph/Source/TextureGraph/Private/TG_Texture.cpp
--- a/Engine/Plugins/Experimental/TextureGraph/Source/TextureGraph/Private/TG_Texture.cpp
+++ b/Engine/Plugins/Experimental/TextureGraph/Source/TextureGraph/Private/TG_Texture.cpp
@@ -28,7 +28,14 @@ void FTG_Texture::ResetTexturePath()
 
 FArchive& operator<<(FArchive& Ar, FTG_Texture& T)
 {
-	Ar << T.TexturePath;
+	Ar.UsingCustomVersion(FTG_CustomVersion::GUID);
+	int32 Version = Ar.CustomVer(FTG_CustomVersion::GUID);
+		
+	if (Ar.IsSaving() ||
+		Ar.IsLoading() && Version >= FTG_CustomVersion::TGTextureAddedTexturePath)
+	{
+		Ar << T.TexturePath;
+	}
 	Ar << T.Descriptor;
 	return Ar;
 }
\ No newline at end of file
diff --git a/Engine/Plugins/Experimental/TextureGraph/Source/TextureGraph/Private/TextureGraph.cpp b/Engine/Plugins/Experimental/TextureGraph/Source/TextureGraph/Private/TextureGraph.cpp
--- a/Engine/Plugins/Experimental/TextureGraph/Source/TextureGraph/Private/TextureGraph.cpp
+++ b/Engine/Plugins/Experimental/TextureGraph/Source/TextureGraph/Private/TextureGraph.cpp
@@ -70,20 +70,38 @@ void UTextureGraphInstance::Update(MixUpdateCyclePtr InCycle)
 	
 }
 
-void UTextureGraphInstance::SetParent(const TObjectPtr<UTextureGraph>& Parent)
+void UTextureGraphInstance::SetParent(const TObjectPtr<UTextureGraphBase>& Parent)
 {
 	this->ParentTextureGraph = Parent;
 	if (Parent == nullptr)
 	{
-		RuntimeGraph = nullptr;
-		Settings = nullptr;
+		Construct("");
 		InputParams.Empty();
 	}
 	else
 	{
 		RuntimeGraph = TStrongObjectPtr<UTG_Graph>(DuplicateObject(Parent->Graph(), GetTransientPackage()));
 		Settings = DuplicateObject(Parent->GetSettings(), this);
+
+		// copy InputParams from Parent if it is an instance
+		if (Parent.IsA<UTextureGraphInstance>())
+		{
+			InputParams = Cast<UTextureGraphInstance>(Parent)->InputParams;	
+		}
+		else
+		{
+			InputParams = Graph()->GetInputParamsVarMap();
+		}
+	}
+}
+
+UTG_Graph* UTextureGraphInstance::Graph()
+{
+	if (!RuntimeGraph && ParentTextureGraph)
+	{
+		RuntimeGraph = TStrongObjectPtr<UTG_Graph>(DuplicateObject(ParentTextureGraph->Graph(), GetTransientPackage()));
 	}
+	return RuntimeGraph.Get();
 }
 
 void UTextureGraphInstance::Initialize()
diff --git a/Engine/Plugins/Experimental/TextureGraph/Source/TextureGraph/Public/TG_GraphEvaluation.h b/Engine/Plugins/Experimental/TextureGraph/Source/TextureGraph/Public/TG_GraphEvaluation.h
--- a/Engine/Plugins/Experimental/TextureGraph/Source/TextureGraph/Public/TG_GraphEvaluation.h
+++ b/Engine/Plugins/Experimental/TextureGraph/Source/TextureGraph/Public/TG_GraphEvaluation.h
@@ -7,7 +7,6 @@
 #include "TG_GraphEvaluation.generated.h"
 class UTG_Graph;
 class UTG_Node;
-
 	
 USTRUCT()
 struct TEXTUREGRAPH_API FVarArgument
@@ -40,6 +39,47 @@ struct TStructOpsTypeTraits<FVarArgument> : public TStructOpsTypeTraitsBase2<FVa
 		WithSerializer = true,  // Enables the use of a custom Serialize method
 	};
 };
+
+struct TEXTUREGRAPH_API FTG_OutputVarMap
+{
+	struct TEXTUREGRAPH_API FOutputVarArgument
+	{
+		FTG_Var* Var;
+		FTG_Argument Argument;
+	};
+	
+	TMap<FName, FOutputVarArgument> VarArguments;
+
+	FTG_Var* GetVar(FName Name)
+	{
+		auto Element = VarArguments.Find(Name);
+		
+		if (Element)
+		{
+			return Element->Var;
+		}
+
+		return nullptr;
+	}
+
+	const FTG_Argument* GetVarArgument(FName Name)
+	{
+		auto Element = VarArguments.Find(Name);
+		
+		if (Element)
+		{
+			return &Element->Argument;
+		}
+
+		return nullptr;
+	}
+
+	void Empty()
+	{
+		VarArguments.Empty();
+	}
+};
+
 USTRUCT()
 struct TEXTUREGRAPH_API FTG_VarMap
 {
@@ -87,7 +127,7 @@ struct TEXTUREGRAPH_API FTG_EvaluationContext
 	// in and out buckets of vars passed as the arguments to the expression evaluate call
 	// For each node/expression, these are populated in the traverse/evaluate call in TG_Graph
 	FTG_VarMap			Inputs;
-	FTG_VarMap			Outputs;
+	FTG_OutputVarMap	Outputs;
 	
 	// These are the Ids of the Vars connected as Param to an upper Graph
 	// These are persistent in the scope of a graph evaluation
diff --git a/Engine/Plugins/Experimental/TextureGraph/Source/TextureGraph/Public/TextureGraph.h b/Engine/Plugins/Experimental/TextureGraph/Source/TextureGraph/Public/TextureGraph.h
--- a/Engine/Plugins/Experimental/TextureGraph/Source/TextureGraph/Public/TextureGraph.h
+++ b/Engine/Plugins/Experimental/TextureGraph/Source/TextureGraph/Public/TextureGraph.h
@@ -89,15 +89,15 @@ public:
 	
 	/** Parent Texture Graph. */
 	UPROPERTY(EditAnywhere, BlueprintReadOnly, Category=TextureGraphInstance, AssetRegistrySearchable)
-	TObjectPtr<UTextureGraph> ParentTextureGraph;
+	TObjectPtr<UTextureGraphBase> ParentTextureGraph; 
 
 	virtual void Construct(FString Name) override;
 
 	void CopyParamsToRuntimeGraph();
 	void Update(MixUpdateCyclePtr InCycle) override;
-	void SetParent(const TObjectPtr<UTextureGraph>& Parent);
+	void SetParent(const TObjectPtr<UTextureGraphBase>& Parent);
 	
-	virtual const UTG_Graph* Graph() const { return RuntimeGraph.Get(); }
-	virtual UTG_Graph* Graph() { return RuntimeGraph.Get(); }
+	virtual const UTG_Graph* Graph() const override { return RuntimeGraph.Get(); };
+	virtual UTG_Graph* Graph() override;
 	void Initialize();
 };
\ No newline at end of file
diff --git a/Engine/Plugins/Experimental/TextureGraph/Source/TextureGraphEditor/Private/AssetDefinition_TextureGraphInstance.cpp b/Engine/Plugins/Experimental/TextureGraph/Source/TextureGraphEditor/Private/AssetDefinition_TextureGraphInstance.cpp
--- a/Engine/Plugins/Experimental/TextureGraph/Source/TextureGraphEditor/Private/AssetDefinition_TextureGraphInstance.cpp
+++ b/Engine/Plugins/Experimental/TextureGraph/Source/TextureGraphEditor/Private/AssetDefinition_TextureGraphInstance.cpp
@@ -10,6 +10,7 @@
 #include "ToolMenus.h"
 #include "IAssetTools.h"
 #include "ToolMenuSection.h"
+#include "TG_InstanceFactory.h"
 #include "AssetRegistry/AssetRegistryHelpers.h"
 
 #define LOCTEXT_NAMESPACE "UAssetDefinition_TextureGraphInstance"
@@ -50,7 +51,22 @@ namespace MenuExtension_TextureGraphInstance
 			
 		}
 	}
-
+	static void ExecuteNewTextureGraphInstance(const FToolMenuContext& InContext)
+	{
+		if (const UContentBrowserAssetContextMenuContext* CBContext = UContentBrowserAssetContextMenuContext::FindContextWithAssets(InContext))
+		{
+			
+			IAssetTools::Get().CreateAssetsFrom<UTextureGraphInstance>(
+				CBContext->LoadSelectedObjects<UTextureGraphInstance>(), UTextureGraphInstance::StaticClass(), TEXT("_Inst"), [](UTextureGraphInstance* SourceObject)
+				{
+					UTG_InstanceFactory* Factory = NewObject<UTG_InstanceFactory>();
+					Factory->InitialParent = SourceObject;
+					return Factory;
+				}
+			);
+			
+		}
+	}
 	static FDelayedAutoRegisterHelper DelayedAutoRegister(EDelayedRegisterRunPhase::EndOfEngineInit, []{ 
 		UToolMenus::RegisterStartupCallback(FSimpleMulticastDelegate::FDelegate::CreateLambda([]()
 		{
@@ -73,9 +89,21 @@ namespace MenuExtension_TextureGraphInstance
 						}
 					}
 				}));
-			}
 
-			
+				Section.AddDynamicEntry(NAME_None, FNewToolMenuSectionDelegate::CreateLambda([](FToolMenuSection& InSection)
+				{
+					if (const UContentBrowserAssetContextMenuContext* Context = UContentBrowserAssetContextMenuContext::FindContextWithAssets(InSection))
+					{
+						{
+							const TAttribute<FText> Label = LOCTEXT("TextureGraph_NewInstance", "Create Texture Graph Instance");
+							const TAttribute<FText> ToolTip = LOCTEXT("TextureGraph_NewInstanceTooltip", "Creates a parameterized Texture Graph Instance using this Texture Graph as a base");
+							const FSlateIcon Icon = FSlateIcon();
+							const FToolMenuExecuteAction UIAction = FToolMenuExecuteAction::CreateStatic(&ExecuteNewTextureGraphInstance);
+							InSection.AddMenuEntry("TextureGraph_NewInstance", Label, ToolTip, Icon, UIAction);
+						}
+					}
+				}));
+			}
 		}));
 	});
 }
diff --git a/Engine/Plugins/Experimental/TextureGraph/Source/TextureGraphEditor/Private/Customizations/TG_MaterialMappingInfoCustomization.h b/Engine/Plugins/Experimental/TextureGraph/Source/TextureGraphEditor/Private/Customizations/TG_MaterialMappingInfoCustomization.h
--- a/Engine/Plugins/Experimental/TextureGraph/Source/TextureGraphEditor/Private/Customizations/TG_MaterialMappingInfoCustomization.h
+++ b/Engine/Plugins/Experimental/TextureGraph/Source/TextureGraphEditor/Private/Customizations/TG_MaterialMappingInfoCustomization.h
@@ -49,14 +49,17 @@ public:
 		UTextureGraphBase* TextureGraph = Cast<UTextureGraphBase>(MixSettings->Mix());
 
 		OutComboBoxStrings.Add(MakeShared<FString>("None"));		
-		TextureGraph->Graph()->ForEachNodes(
-			[&](const UTG_Node* node, uint32 index)
+		if (TextureGraph->Graph())
 		{
-			if(UTG_Expression_Output* OutputExpression = Cast<UTG_Expression_Output>(node->GetExpression()))
+			TextureGraph->Graph()->ForEachNodes(
+			[&](const UTG_Node* node, uint32 index)
 			{
-				OutComboBoxStrings.Add(MakeShared<FString>(OutputExpression->GetTitleName().ToString()));			
-			}
-		});
+				if(UTG_Expression_Output* OutputExpression = Cast<UTG_Expression_Output>(node->GetExpression()))
+				{
+					OutComboBoxStrings.Add(MakeShared<FString>(OutputExpression->GetTitleName().ToString()));			
+				}
+			});
+		}
 	}
 
 	virtual void CustomizeHeader(TSharedRef<IPropertyHandle> PropertyHandle, FDetailWidgetRow& HeaderRow, IPropertyTypeCustomizationUtils& CustomizationUtils) override
diff --git a/Engine/Plugins/Experimental/TextureGraph/Source/TextureGraphEditor/Private/TG_Editor.cpp b/Engine/Plugins/Experimental/TextureGraph/Source/TextureGraphEditor/Private/TG_Editor.cpp
--- a/Engine/Plugins/Experimental/TextureGraph/Source/TextureGraphEditor/Private/TG_Editor.cpp
+++ b/Engine/Plugins/Experimental/TextureGraph/Source/TextureGraphEditor/Private/TG_Editor.cpp
@@ -108,7 +108,7 @@ void FTG_Editor::RegisterTabSpawners(const TSharedRef<class FTabManager>& InTabM
 		.SetGroup(WorkspaceMenuCategoryRef)
 		.SetIcon(FSlateIcon(FAppStyle::GetAppStyleSetName(), "Kismet.Tabs.FindResults"));*/
 
-	InTabManager->RegisterTabSpawner(FTG_EditorTabs::PreviewSceneSettingsTabId, FOnSpawnTab::CreateSP(this, &FTG_Editor::SpawnTab_PreviewSettings))
+	InTabManager->RegisterTabSpawner(FTG_EditorTabs::PreviewSceneSettingsTabId, FOnSpawnTab::CreateSP(this, &FTG_Editor::SpawnTab_PreviewSceneSettings))
 		.SetDisplayName(LOCTEXT("PreviewSceneSettingsTab", "Preview Scene Settings"))
 		.SetGroup(WorkspaceMenuCategoryRef)
 		.SetIcon(FSlateIcon(FAppStyle::GetAppStyleSetName(), "LevelEditor.Tabs.Details"));
@@ -128,7 +128,7 @@ void FTG_Editor::RegisterTabSpawners(const TSharedRef<class FTabManager>& InTabM
 		.SetGroup(WorkspaceMenuCategoryRef)
 		.SetIcon(FSlateIcon(FAppStyle::GetAppStyleSetName(), "LevelEditor.Tabs.Viewports"));*/
 
-	InTabManager->RegisterTabSpawner(FTG_EditorTabs::PreviewSettingsTabId, FOnSpawnTab::CreateSP(this, &FTG_Editor::SpawnTab_Settings))
+	InTabManager->RegisterTabSpawner(FTG_EditorTabs::PreviewSettingsTabId, FOnSpawnTab::CreateSP(this, &FTG_Editor::SpawnTab_PreviewSettings))
 		.SetDisplayName(LOCTEXT("PreviewSettingsTab", "3D Preview Settings"))
 		.SetGroup(WorkspaceMenuCategoryRef)
 		.SetIcon(FSlateIcon(FAppStyle::GetAppStyleSetName(), "LevelEditor.Tabs.Details"));
@@ -232,9 +232,9 @@ void FTG_Editor::InitEditor(const EToolkitMode::Type Mode, const TSharedPtr< cla
 	SettingsViewArgs.NameAreaSettings = FDetailsViewArgs::HideNameArea;
 	SettingsViewArgs.bHideSelectionTip = true;
 	SettingsViewArgs.NotifyHook = this;
-	SettingsView = PropertyEditorModule.CreateDetailView(SettingsViewArgs);
+	PreviewSettingsView = PropertyEditorModule.CreateDetailView(SettingsViewArgs);
 
-	SettingsView->OnFinishedChangingProperties().AddSP(this, &FTG_Editor::OnFinishedChangingSettings);
+	PreviewSettingsView->OnFinishedChangingProperties().AddSP(this, &FTG_Editor::OnFinishedChangingSettings);
 
 	// Output details view
 	FDetailsViewArgs OutputViewArgs;
@@ -1029,10 +1029,8 @@ TSharedRef<SDockTab> FTG_Editor::SpawnTab_Find(const FSpawnTabArgs& Args)
 	return SNew(SDockTab);
 }
 
-TSharedRef<SDockTab> FTG_Editor::SpawnTab_PreviewSettings(const FSpawnTabArgs& Args)
+TSharedRef<SDockTab> FTG_Editor::SpawnTab_PreviewSceneSettings(const FSpawnTabArgs& Args)
 {
-	check(Args.GetTabId() == FTG_EditorTabs::PreviewSceneSettingsTabId);
-
 	check( Args.GetTabId() == FTG_EditorTabs::PreviewSceneSettingsTabId );
 	return SAssignNew(PreviewSceneSettingsDockTab, SDockTab)
 		.Label( LOCTEXT("TG_EditorPreviewSceneSettings_TabTitle", "Preview Scene Settings") )
@@ -1140,18 +1138,18 @@ FReply FTG_Editor::OnExportClick()
 	return FReply::Handled();
 }
 
-TSharedRef<SDockTab> FTG_Editor::SpawnTab_Settings(const FSpawnTabArgs& Args)
+TSharedRef<SDockTab> FTG_Editor::SpawnTab_PreviewSettings(const FSpawnTabArgs& Args)
 {
 	check(Args.GetTabId() == FTG_EditorTabs::PreviewSettingsTabId);
 
-	TSharedPtr<SDockTab> SettingsTab = SNew(SDockTab)
+	TSharedPtr<SDockTab> PreviewSettingsTab = SNew(SDockTab)
 		[
-			SettingsView.ToSharedRef()
+			PreviewSettingsView.ToSharedRef()
 		];
 
-	GetSettingsView()->SetObject(EditedTextureGraph->GetSettings(), true);
+	GetPreviewSettingsView()->SetObject(EditedTextureGraph->GetSettings(), true);
 
-	return SettingsTab.ToSharedRef();
+	return PreviewSettingsTab.ToSharedRef();
 }
 
 TSharedRef<SDockTab> FTG_Editor::SpawnTab_Errors(const FSpawnTabArgs& Args)
diff --git a/Engine/Plugins/Experimental/TextureGraph/Source/TextureGraphEditor/Private/TG_Editor.h b/Engine/Plugins/Experimental/TextureGraph/Source/TextureGraphEditor/Private/TG_Editor.h
--- a/Engine/Plugins/Experimental/TextureGraph/Source/TextureGraphEditor/Private/TG_Editor.h
+++ b/Engine/Plugins/Experimental/TextureGraph/Source/TextureGraphEditor/Private/TG_Editor.h
@@ -86,7 +86,7 @@ public:
 	TSharedRef<class IDetailsView>					GetDetailView() const { return DetailsView.ToSharedRef(); }
 	
 	// Widget Accessors
-	TSharedRef<class IDetailsView>					GetSettingsView() const { return SettingsView.ToSharedRef(); }
+	TSharedRef<class IDetailsView>					GetPreviewSettingsView() const { return PreviewSettingsView.ToSharedRef(); }
 
 	// Widget Accessors
 	TSharedRef<class IDetailsView>					GetOutputView() const { return OutputView.ToSharedRef(); }
@@ -259,7 +259,7 @@ private:
 	TSharedRef<SDockTab>							SpawnTab_ParameterDefaults(const FSpawnTabArgs& Args);
 	TSharedRef<SDockTab>							SpawnTab_NodePreview(const FSpawnTabArgs& Args);
 	TSharedRef<SDockTab>							SpawnTab_Output(const FSpawnTabArgs& Args);
-	TSharedRef<SDockTab>							SpawnTab_Settings(const FSpawnTabArgs& Args);
+	TSharedRef<SDockTab>							SpawnTab_PreviewSceneSettings(const FSpawnTabArgs& Args);
 	TSharedRef<SDockTab>							SpawnTab_Errors(const FSpawnTabArgs& Args);
 	TSharedRef<SDockTab>							SpawnTab_TextureDetails(const FSpawnTabArgs& Args);
 
@@ -295,7 +295,7 @@ private:
 	TSharedPtr<class IDetailsView>					DetailsView;
 	
 	/** Settings View */
-	TSharedPtr<class IDetailsView>					SettingsView;
+	TSharedPtr<class IDetailsView>					PreviewSettingsView;
 
 	/** Output View */
 	TSharedPtr<class IDetailsView>					OutputView;
diff --git a/Engine/Plugins/Experimental/TextureGraph/Source/TextureGraphEditor/Private/TG_Exporter.cpp b/Engine/Plugins/Experimental/TextureGraph/Source/TextureGraphEditor/Private/TG_Exporter.cpp
--- a/Engine/Plugins/Experimental/TextureGraph/Source/TextureGraphEditor/Private/TG_Exporter.cpp
+++ b/Engine/Plugins/Experimental/TextureGraph/Source/TextureGraphEditor/Private/TG_Exporter.cpp
@@ -69,6 +69,7 @@ void FTG_ExporterCommands::RegisterCommands()
 	UI_COMMAND(Show3DPreviewSettings, "3D Preview Settings", "Toggles visibility of the 3D Preview Settings window", EUserInterfaceActionType::Check, FInputChord());
 	UI_COMMAND(ShowParameters, "Parameters", "Toggles visibility of the Parameters window", EUserInterfaceActionType::Check, FInputChord());
 	UI_COMMAND(ShowExportSettings, "Export Settings", "Toggles visibility of the Export Settings window", EUserInterfaceActionType::Check, FInputChord());
+	UI_COMMAND(ShowDetails, "Property Settings", "Toggles visibility of the Property Detail Settings window", EUserInterfaceActionType::Check, FInputChord());
 }
 FTG_InstanceImpl::FTG_InstanceImpl() :
 	TargetExportSettings(MakeShared<FExportSettings>())
@@ -81,31 +82,40 @@ FTG_InstanceImpl::~FTG_InstanceImpl()
 	DetailsView.Reset();
 	ExportSettingsView.Reset();
 	PreviewSettingsView.Reset();
+	PreviewSceneSettingsView.Reset();
 	ParametersView.Reset();
 
 	// cleanup UI
-	if (Parameters && Parameters->IsValidLowLevelFast())
+	if (Parameters && Parameters->IsValidLowLevel())
 	{
 		Parameters->Parameters.Empty();
 		Parameters = nullptr;
 	}
-	if (ExportSettings && ExportSettings->IsValidLowLevelFast())
+	if (ExportSettings  && ExportSettings->IsValidLowLevel())
 	{
 		ExportSettings->OutputExpressionsInfos.Empty();
 		ExportSettings = nullptr;
 	}
 	OutputNodesComboBoxWidget.Reset();
+
 }
 
 void FTG_InstanceImpl::Cleanup()
 {
-	if (TextureGraphInstance && TextureGraphInstance->IsValidLowLevelFast())
+	if (TextureGraphInstance && TextureGraphInstance->IsValidLowLevel())
 	{
 		// cleanup events
-		TextureGraphInstance->GetSettings()->GetViewportSettings().OnViewportMaterialChangedEvent.RemoveAll(this);
-		TextureGraphInstance->GetSettings()->GetViewportSettings().OnMaterialMappingChangedEvent.RemoveAll(this);
-		TextureGraphInstance->GetSettings()->OnPreviewMeshChangedEvent.RemoveAll(this);
-		TextureGraphInstance->Graph()->OnGraphChangedDelegate.RemoveAll(this);
+		if (UMixSettings* Settings = TextureGraphInstance->GetSettings())
+		{
+			Settings->GetViewportSettings().OnViewportMaterialChangedEvent.RemoveAll(this);
+			Settings->GetViewportSettings().OnMaterialMappingChangedEvent.RemoveAll(this);
+			Settings->OnPreviewMeshChangedEvent.RemoveAll(this);
+		}
+		if (UTG_Graph* Graph = TextureGraphInstance->Graph())
+		{
+			TextureGraphInstance->FlushInvalidations();
+			Graph->OnGraphChangedDelegate.RemoveAll(this);
+		}
 		TextureGraphInstance->OnRenderDone.Unbind();
 		TextureGraphInstance = nullptr;
 	}
@@ -114,6 +124,15 @@ void FTG_InstanceImpl::Cleanup()
 	SelectedNode = nullptr;
 }
 
+void FTG_InstanceImpl::NotifyPostChange(const FPropertyChangedEvent& PropertyChangedEvent, FProperty* PropertyThatChanged)
+{
+	if(PropertyThatChanged->GetName() == GET_MEMBER_NAME_CHECKED(UTextureGraphInstance, ParentTextureGraph))
+	{
+		TextureGraphInstance->SetParent(TextureGraphInstance->ParentTextureGraph);
+		SetTextureGraphToExport(TextureGraphInstance);
+	}
+}
+
 void FTG_InstanceImpl::RegisterTabSpawners(const TSharedPtr<FTabManager> InTabManager)
 {
 	InTabManager->RegisterTabSpawner(FTG_EditorTabs::ViewportTabId, FOnSpawnTab::CreateRaw(this, &FTG_InstanceImpl::SpawnTab_Viewport))
@@ -132,13 +151,18 @@ void FTG_InstanceImpl::RegisterTabSpawners(const TSharedPtr<FTabManager> InTabMa
 		.SetDisplayName(LOCTEXT("PreviewSettingsTab", "3D Preview Settings"))
 		.SetIcon(FSlateIcon(FAppStyle::GetAppStyleSetName(), "LevelEditor.Tabs.Details"));
 
-	InTabManager->RegisterTabSpawner(FTG_EditorTabs::PreviewSceneSettingsTabId, FOnSpawnTab::CreateRaw(this, &FTG_InstanceImpl::SpawnTab_PreviewSettings))
+	InTabManager->RegisterTabSpawner(FTG_EditorTabs::PreviewSceneSettingsTabId, FOnSpawnTab::CreateRaw(this, &FTG_InstanceImpl::SpawnTab_PreviewSceneSettings))
 			.SetDisplayName(LOCTEXT("PreviewSceneSettingsTab", "Preview Scene Settings"))
 			.SetIcon(FSlateIcon(FAppStyle::GetAppStyleSetName(), "LevelEditor.Tabs.Details"));
 	
 	InTabManager->RegisterTabSpawner(FTG_EditorTabs::OutputTabId, FOnSpawnTab::CreateRaw(this, &FTG_InstanceImpl::SpawnTab_ExportSettings))
 		.SetDisplayName(LOCTEXT("ExportSettingsTab", "Export Settings"))
 		.SetIcon(FSlateIcon(FAppStyle::GetAppStyleSetName(), "LevelEditor.Tabs.Details"));
+	
+	InTabManager->RegisterTabSpawner(FTG_EditorTabs::PropertiesTabId, FOnSpawnTab::CreateRaw(this, &FTG_InstanceImpl::SpawnTab_TG_Properties))
+		.SetDisplayName(LOCTEXT("DetailsTab", "Details"))
+		.SetIcon(FSlateIcon(FAppStyle::GetAppStyleSetName(), "LevelEditor.Tabs.Details"));
+	
 }
 
 void FTG_InstanceImpl::UnregisterTabSpawners(const TSharedPtr<FTabManager> InTabManager)
@@ -149,6 +173,7 @@ void FTG_InstanceImpl::UnregisterTabSpawners(const TSharedPtr<FTabManager> InTab
 	InTabManager->UnregisterTabSpawner(FTG_EditorTabs::OutputTabId);
 	InTabManager->UnregisterTabSpawner(FTG_EditorTabs::PreviewSettingsTabId);
 	InTabManager->UnregisterTabSpawner(FTG_EditorTabs::PreviewSceneSettingsTabId);
+	InTabManager->UnregisterTabSpawner(FTG_EditorTabs::PropertiesTabId);
 }
 
 void FTG_InstanceImpl::SetMesh(class UMeshComponent* InPreviewMesh, class UWorld* InWorld)
@@ -177,12 +202,11 @@ TSharedPtr<STG_EditorViewport> FTG_InstanceImpl::GetEditorViewport() const
 
 	return nullptr;
 }
-
 TSharedRef<SDockTab> FTG_InstanceImpl::SpawnTab_PreviewSettings(const FSpawnTabArgs& Args)
 {
 	check(Args.GetTabId() == FTG_EditorTabs::PreviewSettingsTabId);
 
-	TSharedPtr<SDockTab> SettingsTab = SNew(SDockTab)
+	TSharedPtr<SDockTab> PreviewSettingsTab = SNew(SDockTab)
 		[
 			PreviewSettingsView.ToSharedRef()
 		];
@@ -192,8 +216,30 @@ TSharedRef<SDockTab> FTG_InstanceImpl::SpawnTab_PreviewSettings(const FSpawnTabA
 		GetPreviewSettingsView()->SetObject(TextureGraphInstance->GetSettings(), true);
 	}
 
-	return SettingsTab.ToSharedRef();
+	return PreviewSettingsTab.ToSharedRef();
+}
+
+TSharedRef<SDockTab> FTG_InstanceImpl::SpawnTab_PreviewSceneSettings(const FSpawnTabArgs& Args)
+{
+	check( Args.GetTabId() == FTG_EditorTabs::PreviewSceneSettingsTabId );
+	return SAssignNew(PreviewSceneSettingsDockTab, SDockTab)
+		.Label( LOCTEXT("TG_EditorPreviewSceneSettings_TabTitle", "Preview Scene Settings") )
+		[
+			AdvancedPreviewSettingsWidget.IsValid() ? AdvancedPreviewSettingsWidget.ToSharedRef() : SNullWidget::NullWidget
+		];
 }
+
+TSharedRef<SDockTab> FTG_InstanceImpl::SpawnTab_TG_Properties(const FSpawnTabArgs& Args)
+{
+	check(Args.GetTabId() == FTG_EditorTabs::PropertiesTabId);
+
+	TSharedPtr<SDockTab> DetailsTab = SNew(SDockTab)
+		[
+			DetailsView.ToSharedRef()
+		];
+	return DetailsTab.ToSharedRef();
+}
+
 TSharedRef<SDockTab> FTG_InstanceImpl::SpawnTab_ExportSettings(const FSpawnTabArgs& Args)
 {
 	check(Args.GetTabId() == FTG_EditorTabs::OutputTabId);
@@ -250,6 +296,7 @@ TSharedRef<SDockTab> FTG_InstanceImpl::SpawnTab_Viewport(const FSpawnTabArgs& Ar
 
 	// Create a new tab
 	ViewportTabContent = MakeShared<FEditorViewportTabContent>();
+	ViewportTabContent->OnViewportTabContentLayoutChanged().AddRaw(this, &FTG_InstanceImpl::OnEditorLayoutChanged);
 
 	const FString LayoutId = FString("TG_EditorViewport");
 	ViewportTabContent->Initialize(MakeViewportFunc, DockableTab, LayoutId);
@@ -259,7 +306,26 @@ TSharedRef<SDockTab> FTG_InstanceImpl::SpawnTab_Viewport(const FSpawnTabArgs& Ar
 	
 	return DockableTab;
 }
+void FTG_InstanceImpl::BuildSubTools()
+{
+	FAdvancedPreviewSceneModule& AdvancedPreviewSceneModule = FModuleManager::LoadModuleChecked<FAdvancedPreviewSceneModule>("AdvancedPreviewScene");
+	
+	TArray<FAdvancedPreviewSceneModule::FDetailDelegates> Delegates;
+	Delegates.Add({ OnPreviewSceneChangedDelegate });
+	AdvancedPreviewSettingsWidget = AdvancedPreviewSceneModule.CreateAdvancedPreviewSceneSettingsWidget(GetEditorViewport()->GetPreviewScene(), nullptr, TArray<FAdvancedPreviewSceneModule::FDetailCustomizationInfo>(),  TArray<FAdvancedPreviewSceneModule::FPropertyTypeCustomizationInfo>(), Delegates);
+	
+	if (PreviewSceneSettingsDockTab.IsValid())
+	{
+		PreviewSceneSettingsDockTab.Pin()->SetContent(AdvancedPreviewSettingsWidget.ToSharedRef());
+	}
+}
 
+void FTG_InstanceImpl::OnEditorLayoutChanged()
+{
+	BuildSubTools();
+	
+	OnPreviewSceneChangedDelegate.Broadcast(GetEditorViewport()->GetPreviewScene());
+}
 TSharedRef<SDockTab> FTG_InstanceImpl::SpawnTab_ParameterDefaults(const FSpawnTabArgs& Args)
 {
 	check(Args.GetTabId() == FTG_EditorTabs::ParameterDefaultsTabId);
@@ -341,6 +407,12 @@ void FTG_InstanceImpl::UpdatePreviewMesh()
 	}
 }
 
+void FTG_InstanceImpl::CleanupSlateReferences()
+{
+	ViewportTabContent.Reset();
+	AdvancedPreviewSettingsWidget.Reset();
+}
+
 // Function to generate combo box items
 TSharedRef<SWidget> FTG_InstanceImpl::GenerateOutputComboItem(TSharedPtr<FName> InItem)
 {
@@ -399,13 +471,17 @@ void FTG_InstanceImpl::SetTextureGraphToExport(UTextureGraphInstance* InTextureG
 	//Exporter gets notified when rendering is done
 	TextureGraphInstance->OnRenderDone.BindRaw(this, &FTG_InstanceImpl::OnRenderingDone);
 	
+	DetailsView->SetObject(TextureGraphInstance, true);
+	
 	UpdateParametersUI();
 
 	UpdateExportSettingsUI();
 
 	// Update list of output nodes in 2d View
 	OutputNodesList.Empty();
-	TextureGraphInstance->Graph()->ForEachNodes(
+	if (UTG_Graph* Graph = TextureGraphInstance->Graph())
+	{
+		Graph->ForEachNodes(
 		[&](const UTG_Node* node, uint32 index)
 		{
 			if(UTG_Expression_Output* OutputExpression = Cast<UTG_Expression_Output>(node->GetExpression()))
@@ -413,7 +489,9 @@ void FTG_InstanceImpl::SetTextureGraphToExport(UTextureGraphInstance* InTextureG
 				OutputNodesList.Add(MakeShared<FName>(OutputExpression->GetTitleName()));			
 			}
 		});
-	
+		
+		Graph->OnGraphChangedDelegate.AddRaw(this, &FTG_InstanceImpl::OnGraphChanged);
+	}
 	if (OutputNodesComboBoxWidget.IsValid())
 	{
 		if (!OutputNodesList.IsEmpty())
@@ -428,7 +506,7 @@ void FTG_InstanceImpl::SetTextureGraphToExport(UTextureGraphInstance* InTextureG
 	ViewportSettings.OnViewportMaterialChangedEvent.AddRaw(this, &FTG_InstanceImpl::OnViewportMaterialChanged);
 	ViewportSettings.OnMaterialMappingChangedEvent.AddRaw(this, &FTG_InstanceImpl::OnMaterialMappingChanged);
 	TextureGraphInstance->GetSettings()->OnPreviewMeshChangedEvent.AddRaw(this, &FTG_InstanceImpl::SetViewportPreviewMesh);
-	TextureGraphInstance->Graph()->OnGraphChangedDelegate.AddRaw(this, &FTG_InstanceImpl::OnGraphChanged);
+
 	GetPreviewSettingsView()->SetObject(TextureGraphInstance->GetSettings(), true);
 
 	UpdatePreviewMesh();
@@ -451,7 +529,7 @@ void FTG_InstanceImpl::OnGraphChanged(UTG_Graph* InGraph, UTG_Node* InNode, bool
 
 TSharedRef<FTabManager::FLayout> FTG_InstanceImpl::GetDefaultLayout()
 {
-	return FTabManager::NewLayout("Standalone_TextureGraphExporter_Layout_v1")
+	return FTabManager::NewLayout("Standalone_TextureGraphExporter_Layout_v2")
 	->AddArea
 	(
 		FTabManager::NewPrimaryArea()->SetOrientation(Orient_Vertical)
@@ -469,7 +547,9 @@ TSharedRef<FTabManager::FLayout> FTG_InstanceImpl::GetDefaultLayout()
 					->Split
 					(
 						FTabManager::NewStack()
+						->AddTab(FTG_EditorTabs::PropertiesTabId, ETabState::OpenedTab)
 						->AddTab(FTG_EditorTabs::OutputTabId, ETabState::OpenedTab)
+						->SetForegroundTab(FTG_EditorTabs::OutputTabId)
 					)
 				)
 				->Split
@@ -495,7 +575,6 @@ TSharedRef<FTabManager::FLayout> FTG_InstanceImpl::GetDefaultLayout()
 
 void FTG_InstanceImpl::Initialize()
 {
-	
 	FPropertyEditorModule& PropertyEditorModule = FModuleManager::GetModuleChecked<FPropertyEditorModule>("PropertyEditor");
 	FDetailsViewArgs ParameterViewArgs;
 	ParameterViewArgs.NameAreaSettings = FDetailsViewArgs::HideNameArea;
@@ -514,7 +593,14 @@ void FTG_InstanceImpl::Initialize()
 	SettingsViewArgs.NameAreaSettings = FDetailsViewArgs::HideNameArea;
 	SettingsViewArgs.bHideSelectionTip = true;
 	PreviewSettingsView = PropertyEditorModule.CreateDetailView(SettingsViewArgs);
-		
+	
+	FDetailsViewArgs DetailsViewArgs;
+	DetailsViewArgs.NameAreaSettings = FDetailsViewArgs::HideNameArea;
+	DetailsViewArgs.bShowObjectLabel = false;
+	DetailsViewArgs.bHideSelectionTip = true;
+	DetailsViewArgs.NotifyHook = this;
+	DetailsViewArgs.ColumnWidth = 0.70;
+	DetailsView = PropertyEditorModule.CreateDetailView(DetailsViewArgs);
 	
 }
 
@@ -562,38 +648,41 @@ void FTG_InstanceImpl::UpdateExportSettingsUI()
 	// recreate export settings UI
 	ExportSettings = NewObject<UTG_ExportSettings>(TextureGraphInstance);
 	
-	UTG_Graph* Graph = TextureGraphInstance->Graph();
-
-	Graph->ForEachNodes([&](const UTG_Node* node, uint32 index)
+	if (UTG_Graph* Graph = TextureGraphInstance->Graph())
 	{
-		if (UTG_Expression_Output* OutputExpression = dynamic_cast<UTG_Expression_Output*>(node->GetExpression()))
+		Graph->ForEachNodes([&](const UTG_Node* node, uint32 index)
 		{
-			ExportSettings->OutputExpressionsInfos.Add(FOutputExpressionInfo{OutputExpression->GetTitleName(), node->GetId()});
-		}
-	});
-
+			if (UTG_Expression_Output* OutputExpression = dynamic_cast<UTG_Expression_Output*>(node->GetExpression()))
+			{
+				ExportSettings->OutputExpressionsInfos.Add(FOutputExpressionInfo{OutputExpression->GetTitleName(), node->GetId()});
+			}
+		});
+	}
 	GetExportSettingsView()->SetObject(ExportSettings);
 }
 void FTG_InstanceImpl::UpdateParametersUI()
 {
-	UTG_Graph* Graph = TextureGraphInstance->Graph();
-	FTG_Ids Ids = Graph->GetParamIds();
-
+	
 	// Create a new object to set for the view.
 	Parameters = NewObject<UTG_Parameters>();
 
-	for (const FTG_Id& id : Ids)
+	if (UTG_Graph* Graph = TextureGraphInstance->Graph())
 	{
-		UTG_Pin* Pin = Graph->GetPin(id);
+		FTG_Ids Ids = Graph->GetParamIds();
 
-		if (Pin && (Pin->IsInput() || Pin->IsSetting()))
+		for (const FTG_Id& id : Ids)
 		{
-			FTG_ParameterInfo Info{ id, Pin->GetAliasName() };
-			Parameters->Parameters.Add(Info);
+			UTG_Pin* Pin = Graph->GetPin(id);
+
+			if (Pin && (Pin->IsInput() || Pin->IsSetting()))
+			{
+				FTG_ParameterInfo Info{ id, Pin->GetAliasName() };
+				Parameters->Parameters.Add(Info);
+			}
 		}
+		Parameters->TextureGraph = Graph;
 	}
-	Parameters->TextureGraph = Graph;
-
+	
 	GetParametersView()->SetObject(Parameters);
 }
 void FTG_InstanceImpl::Tick(float DeltaTime)
@@ -626,26 +715,34 @@ FTG_ExporterUtility::FTG_ExporterUtility()
 	: Impl(MakeUnique<FTG_InstanceImpl>())
 {
 	const IWorkspaceMenuStructure& MenuStructure = WorkspaceMenu::GetMenuStructure();
-
+	
 	FGlobalTabmanager::Get()->RegisterNomadTabSpawner(FTG_EditorTabs::TextureExporterTabId, FOnSpawnTab::CreateRaw(this, &FTG_ExporterUtility::CreateTGExporterTab))
 		.SetDisplayName(NSLOCTEXT("TextureGraphExporter", "TabTitle", "Texture Graph Exporter"))
 		.SetTooltipText(NSLOCTEXT("TextureGraphExporter", "TooltipText", "Open the Texture Graph Exporter tab."))
 		.SetGroup(MenuStructure.GetDeveloperToolsMiscCategory())
 		.SetIcon(FSlateIcon(FAppStyle::GetAppStyleSetName(), "ClassIcon.Texture2D"));
+
+	FSlateApplication::Get().OnPreShutdown().AddRaw(this, &FTG_ExporterUtility::Cleanup);
+	
 }
 
-FTG_ExporterUtility::~FTG_ExporterUtility()
+void FTG_ExporterUtility::Cleanup()
 {
 	if (FSlateApplication::IsInitialized())
 	{
+		if (TGExporterTabManager.IsValid())
+		{
+			Impl->CleanupSlateReferences();
+			Impl->UnregisterTabSpawners(TGExporterTabManager);
+		}
 		FGlobalTabmanager::Get()->UnregisterNomadTabSpawner(FTG_EditorTabs::TextureExporterTabId);
 	}
+	
+	FSlateApplication::Get().OnPreShutdown().RemoveAll(this);
+}
 
-	if (TGExporterTabManager.IsValid())
-	{
-		Impl->UnregisterTabSpawners(TGExporterTabManager);
-		FGlobalTabmanager::Get()->UnregisterTabSpawner(FTG_EditorTabs::TextureExporterTabId);
-	}
+FTG_ExporterUtility::~FTG_ExporterUtility()
+{
 }
 
 TSharedRef<SDockTab> FTG_ExporterUtility::CreateTGExporterTab(const FSpawnTabArgs& Args)
@@ -683,7 +780,6 @@ TSharedRef<SDockTab> FTG_ExporterUtility::CreateTGExporterTab(const FSpawnTabArg
 				OwningTabManager->CloseAllAreas();
 			}
 			Impl->Cleanup();
-			
 		}
 		, TGExporterTabManagerWeak
 	));
@@ -820,6 +916,23 @@ TSharedRef<SDockTab> FTG_ExporterUtility::CreateTGExporterTab(const FSpawnTabArg
 				FTG_EditorTabs::OutputTabId
 			)
 		);
+
+	CommandList->MapAction(
+			FTG_ExporterCommands::Get().ShowDetails,
+			FExecuteAction::CreateStatic(
+				ToggleTabVisibility,
+				TGExportManagerWeak,
+				FTG_EditorTabs::PropertiesTabId
+			),
+			FCanExecuteAction::CreateStatic(
+				[]() { return true; }
+			),
+			FIsActionChecked::CreateStatic(
+				IsTabVisible,
+				TGExportManagerWeak,
+				FTG_EditorTabs::PropertiesTabId
+			)
+		);
 	
 	FMenuBarBuilder MenuBarBuilder(CommandList);
 	MenuBarBuilder.AddPullDownMenu(
@@ -831,6 +944,7 @@ TSharedRef<SDockTab> FTG_ExporterUtility::CreateTGExporterTab(const FSpawnTabArg
 			Builder.AddMenuEntry(FTG_ExporterCommands::Get().Show3DPreviewSettings);
 			Builder.AddMenuEntry(FTG_ExporterCommands::Get().Show3DPreview);
 			Builder.AddMenuEntry(FTG_ExporterCommands::Get().ShowExportSettings);
+			Builder.AddMenuEntry(FTG_ExporterCommands::Get().ShowDetails);
 			})
 	);
 	
@@ -859,7 +973,10 @@ TSharedRef<SDockTab> FTG_ExporterUtility::CreateTGExporterTab(const FSpawnTabArg
 	
 	// Tell tab-manager about the multi-box for platforms with a global menu bar
 	TGExporterTabManager->SetMenuMultiBox(MenuBarBuilder.GetMultiBox(), MenuBarWidget);
-	
+
+	TextureGraphInstance = NewObject<UTextureGraphInstance>();
+	TextureGraphInstance->Construct(FString());
+	Impl->SetTextureGraphToExport(TextureGraphInstance);
 	return NomadTab;
 }
 
@@ -873,7 +990,11 @@ void FTG_ExporterUtility::SetTextureGraphToExport(UTextureGraphBase* InTextureGr
 	}
 	else if (InTextureGraph->IsA<UTextureGraphInstance>())
 	{
-		TextureGraphInstance = Cast<UTextureGraphInstance>(InTextureGraph);
+		TextureGraphInstance = DuplicateObject(Cast<UTextureGraphInstance>(InTextureGraph), GetTransientPackage());
+		if (!InTextureGraph->Graph())
+		{
+			TextureGraphInstance->Construct(FString());
+		}
 	}
 	
 	// force open export window
diff --git a/Engine/Plugins/Experimental/TextureGraph/Source/TextureGraphEditor/Private/TG_InstanceEditor.cpp b/Engine/Plugins/Experimental/TextureGraph/Source/TextureGraphEditor/Private/TG_InstanceEditor.cpp
--- a/Engine/Plugins/Experimental/TextureGraph/Source/TextureGraphEditor/Private/TG_InstanceEditor.cpp
+++ b/Engine/Plugins/Experimental/TextureGraph/Source/TextureGraphEditor/Private/TG_InstanceEditor.cpp
@@ -4,11 +4,8 @@
 
 #include "TextureGraphEditorModule.h"
 #include "TextureGraph.h"
-#include "TG_EditorTabs.h"
-#include "STG_EditorViewport.h"
 #include "Components/MeshComponent.h"
 #include "Components/StaticMeshComponent.h"
-#include "DetailsViewArgs.h"
 #include "Editor.h"
 
 #include "Framework/Application/SlateApplication.h"
@@ -20,7 +17,6 @@
 #include "Widgets/Text/STextBlock.h"
 #include <Toolkits/AssetEditorToolkit.h>
 
-#include "EdGraph/TG_EditorErrorReporter.h"
 #include "Model/Mix/MixInterface.h"
 #include "UObject/MetaData.h"
 #include "UObject/ObjectSaveContext.h"
@@ -47,35 +43,15 @@ void FTG_InstanceEditor::SetMesh(class UMeshComponent* InPreviewMesh, class UWor
 
 void FTG_InstanceEditor::RegisterTabSpawners(const TSharedRef<class FTabManager>& InTabManager)
 {
-	WorkspaceMenuCategory = InTabManager->AddLocalWorkspaceMenuCategory(LOCTEXT("WorkspaceMenu_TG_Editor", "Texture Graph Instance Editor"));
-	auto WorkspaceMenuCategoryRef = WorkspaceMenuCategory.ToSharedRef();
-
 	FAssetEditorToolkit::RegisterTabSpawners(InTabManager);
-
-	InTabManager->RegisterTabSpawner(FTG_EditorTabs::PropertiesTabId, FOnSpawnTab::CreateSP(this, &FTG_InstanceEditor::SpawnTab_TG_Properties))
-		.SetDisplayName(LOCTEXT("DetailsTab", "Details"))
-		.SetGroup( WorkspaceMenuCategoryRef )
-		.SetIcon(FSlateIcon(FAppStyle::GetAppStyleSetName(), "LevelEditor.Tabs.Details"));
-	
 	InstanceImpl->RegisterTabSpawners(InTabManager);
 }
 
 void FTG_InstanceEditor::UnregisterTabSpawners(const TSharedRef<class FTabManager>& InTabManager)
 {
 	InstanceImpl->UnregisterTabSpawners(InTabManager);
-	InTabManager->UnregisterTabSpawner(FTG_EditorTabs::PropertiesTabId);
 	FAssetEditorToolkit::UnregisterTabSpawners(InTabManager);
 }
-TSharedRef<SDockTab> FTG_InstanceEditor::SpawnTab_TG_Properties(const FSpawnTabArgs& Args)
-{
-	check(Args.GetTabId() == FTG_EditorTabs::PropertiesTabId);
-
-	TSharedPtr<SDockTab> DetailsTab = SNew(SDockTab)
-		[
-			DetailsView.ToSharedRef()
-		];
-	return DetailsTab.ToSharedRef();
-}
 
 void FTG_InstanceEditor::OnTextureGraphPreSave(UObject* Object, FObjectPreSaveContext SaveContext)
 {
@@ -93,14 +69,6 @@ void FTG_InstanceEditor::InitEditor(const EToolkitMode::Type Mode, const TShared
 		InTextureGraph->Initialize();
 	}
 	
-	FDetailsViewArgs DetailsViewArgs;
-	DetailsViewArgs.NameAreaSettings = FDetailsViewArgs::HideNameArea;
-	DetailsViewArgs.bShowObjectLabel = false;
-	DetailsViewArgs.bHideSelectionTip = true;
-	DetailsViewArgs.NotifyHook = this;
-	DetailsViewArgs.ColumnWidth = 0.70;
-	DetailsView = PropertyEditorModule.CreateDetailView(DetailsViewArgs);
-	
 	InstanceImpl = MakeUnique<FTG_InstanceImpl>();
 	InstanceImpl->Initialize();
 
@@ -109,7 +77,6 @@ void FTG_InstanceEditor::InitEditor(const EToolkitMode::Type Mode, const TShared
     
 	InstanceImpl->SetTextureGraphToExport(InTextureGraph);
 	
-	DetailsView->SetObject(TextureGraphInstance, true);
 
 // 	// Add analytics tag
 // 	if (FEngineAnalytics::IsAvailable())
@@ -125,20 +92,10 @@ FTG_InstanceEditor::FTG_InstanceEditor()
 
 FTG_InstanceEditor::~FTG_InstanceEditor()
 {
-	DetailsView.Reset();
 	InstanceImpl.Reset();
 	TextureGraphInstance = nullptr;
 }
 
-void FTG_InstanceEditor::NotifyPostChange(const FPropertyChangedEvent& PropertyChangedEvent, FProperty* PropertyThatChanged)
-{
-	if(PropertyThatChanged->GetName() == GET_MEMBER_NAME_CHECKED(UTextureGraphInstance, ParentTextureGraph))
-	{
-		TextureGraphInstance->SetParent(TextureGraphInstance->ParentTextureGraph);
-		InstanceImpl->SetTextureGraphToExport(TextureGraphInstance);
-	}
-}
-
 FLinearColor FTG_InstanceEditor::GetWorldCentricTabColorScale() const
 {
 	return FLinearColor(0.3f, 0.2f, 0.5f, 0.5f);
@@ -164,7 +121,7 @@ void FTG_InstanceEditor::AddReferencedObjects(FReferenceCollector& Collector)
 	Collector.AddReferencedObject(TextureGraphInstance);
 }
 
-class UMixInterface* FTG_InstanceEditor::GetTextureGraphInterface() const
+UMixInterface* FTG_InstanceEditor::GetTextureGraphInterface() const
 {
 	return TextureGraphInstance;
 }
diff --git a/Engine/Plugins/Experimental/TextureGraph/Source/TextureGraphEditor/Private/TG_InstanceEditor.h b/Engine/Plugins/Experimental/TextureGraph/Source/TextureGraphEditor/Private/TG_InstanceEditor.h
--- a/Engine/Plugins/Experimental/TextureGraph/Source/TextureGraphEditor/Private/TG_InstanceEditor.h
+++ b/Engine/Plugins/Experimental/TextureGraph/Source/TextureGraphEditor/Private/TG_InstanceEditor.h
@@ -3,19 +3,10 @@
 #pragma once
 
 #include "CoreMinimal.h"
-#include "AdvancedPreviewSceneModule.h"
 #include "ITG_Editor.h"
 #include "EdGraph/EdGraphPin.h"
 #include "EditorUndoClient.h"
 #include "Misc/NotifyHook.h"
-#include "AssetEditorViewportLayout.h"
-#include "EditorViewportTabContent.h"
-#include "IDetailsView.h"
-#include "Export/TextureExporter.h"
-
-#include "TG_Parameter.h"
-#include "EdGraph/TG_EdGraphNode.h"
-#include "STG_Palette.h"
 #include "GraphEditor.h"
 
 class UTextureGraphInstance;
@@ -34,7 +25,7 @@ struct FGraphAppearanceInfo;
 
 DECLARE_LOG_CATEGORY_EXTERN(LogTextureGraphInstanceEditor, Log, All);
 
-class FTG_InstanceEditor : public ITG_Editor, public FGCObject, public FTickableGameObject, public FEditorUndoClient, public FNotifyHook
+class FTG_InstanceEditor : public ITG_Editor, public FGCObject, public FTickableGameObject, public FEditorUndoClient
 {
 protected:
 	virtual void RefreshViewport() override;
@@ -62,9 +53,6 @@ public:
 
 	virtual											~FTG_InstanceEditor() override;
 
-	/** Post edit change notify for properties. */
-	virtual void NotifyPostChange( const FPropertyChangedEvent& PropertyChangedEvent, FProperty* PropertyThatChanged ) override;
-	
 	// Inherited via ITG_Editor
 	virtual FLinearColor							GetWorldCentricTabColorScale() const override;
 	virtual FName									GetToolkitFName() const override;
@@ -76,7 +64,7 @@ public:
 
 	virtual FString									GetReferencerName() const override { return TEXT("FTextureScriptEditor");}
 
-	class UMixInterface*							GetTextureGraphInterface() const override;
+	virtual class UMixInterface*					GetTextureGraphInterface() const override;
 
 	// Inherited via FTickableGameObject
 	virtual void									Tick(float DeltaTime) override;
@@ -96,14 +84,10 @@ public:
 	virtual void									OnClose() override;
 	// ~End FAssetEditorToolkit interface
 
-	// Widget Accessors
-	TSharedRef<class IDetailsView>					GetDetailView() const { return DetailsView.ToSharedRef(); }
 private:
 	//	FDateTime										SessionStartTime;
 
 private:
-	/** Property View */
-	TSharedPtr<class IDetailsView>					DetailsView;
 	
 	TUniquePtr<struct FTG_InstanceImpl>				InstanceImpl;
 	
diff --git a/Engine/Plugins/Experimental/TextureGraph/Source/TextureGraphEditor/Private/TG_InstanceFactory.h b/Engine/Plugins/Experimental/TextureGraph/Source/TextureGraphEditor/Private/TG_InstanceFactory.h
--- a/Engine/Plugins/Experimental/TextureGraph/Source/TextureGraphEditor/Private/TG_InstanceFactory.h
+++ b/Engine/Plugins/Experimental/TextureGraph/Source/TextureGraphEditor/Private/TG_InstanceFactory.h
@@ -12,7 +12,7 @@ class UTG_InstanceFactory : public UFactory
 	GENERATED_BODY()
 public:
 	UPROPERTY()
-	TObjectPtr<class UTextureGraph> InitialParent;
+	TObjectPtr<class UTextureGraphBase> InitialParent;
 	
 	UTG_InstanceFactory();
 	UObject* FactoryCreateNew(UClass* Class, UObject* InParent, FName Name, EObjectFlags Flags, UObject* Context, FFeedbackContext* Warn);
diff --git a/Engine/Plugins/Experimental/TextureGraph/Source/TextureGraphEditor/Private/TextureGraphEditorModule.cpp b/Engine/Plugins/Experimental/TextureGraph/Source/TextureGraphEditor/Private/TextureGraphEditorModule.cpp
--- a/Engine/Plugins/Experimental/TextureGraph/Source/TextureGraphEditor/Private/TextureGraphEditorModule.cpp
+++ b/Engine/Plugins/Experimental/TextureGraph/Source/TextureGraphEditor/Private/TextureGraphEditorModule.cpp
@@ -64,6 +64,8 @@ void FTextureGraphEditorModule::ShutdownModule()
 {
 	// This function may be called during shutdown to clean up your module.  For modules that support dynamic reloading,
 	// we call this function before unloading the module.
+
+	
 	UnRegisterAllAssetTypeActions();
 	FTG_EditorCommands::Unregister();
 	FTG_ExporterCommands::Unregister();
@@ -91,11 +93,6 @@ void FTextureGraphEditorModule::ShutdownModule()
 	ShutdownTextureGraphEngine();
 }
 
-void FTextureGraphEditorModule::PreUnloadCallback()
-{
-	TG_Exporter.Reset();
-}
-
 void FTextureGraphEditorModule::StartTextureGraphEngine()
 {
 	if (!TextureGraphEngine::GetInstance())
diff --git a/Engine/Plugins/Experimental/TextureGraph/Source/TextureGraphEditor/Public/TG_Exporter.h b/Engine/Plugins/Experimental/TextureGraph/Source/TextureGraphEditor/Public/TG_Exporter.h
--- a/Engine/Plugins/Experimental/TextureGraph/Source/TextureGraphEditor/Public/TG_Exporter.h
+++ b/Engine/Plugins/Experimental/TextureGraph/Source/TextureGraphEditor/Public/TG_Exporter.h
@@ -2,11 +2,13 @@
 
 #pragma once
 
+#include "AdvancedPreviewSceneModule.h"
 #include "Templates/UniquePtr.h"
 #include "Tickable.h"
 #include "Framework/Commands/Commands.h"
 #include "Framework/Docking/TabManager.h"
 #include "Widgets/Input/SComboBox.h"
+#include "Misc/NotifyHook.h"
 #define LOCTEXT_NAMESPACE "TextureGraphExporter"
 class UTextureGraphInstance;
 class UTextureGraphBase;
@@ -15,6 +17,7 @@ struct FTG_ExporterUtility
 {
 	// Initializes the global state of the exporter (commands, tab spawners, etc):
 	FTG_ExporterUtility();
+	void Cleanup();
 
 	// Destructor declaration purely so that we can pimpl:
 	~FTG_ExporterUtility();
@@ -61,12 +64,13 @@ struct FTG_ExporterCommands : public TCommands<FTG_ExporterCommands>
 	TSharedPtr<FUICommandInfo> Show3DPreviewSettings;
 	TSharedPtr<FUICommandInfo> ShowParameters;
 	TSharedPtr<FUICommandInfo> ShowExportSettings;
+	TSharedPtr<FUICommandInfo> ShowDetails;
 
 };
 class UTG_Node;
 class UTG_Graph;
 
-struct FTG_InstanceImpl: public FTickableGameObject, public FGCObject
+struct FTG_InstanceImpl: public FTickableGameObject, public FGCObject, public FNotifyHook
 {
 	friend class FTG_Exporter;
 	friend class FTG_InstanceEditor;
@@ -86,6 +90,7 @@ public:
 	void 											SetMesh(class UMeshComponent* InPreviewMesh, class UWorld* InWorld);
 	void 											SetViewportPreviewMesh();
 	void											UpdatePreviewMesh();
+	void											CleanupSlateReferences();
 
 private:
 	TObjectPtr<UTextureGraphInstance>				TextureGraphInstance;
@@ -95,6 +100,15 @@ private:
 	TSharedPtr<class IDetailsView>					ParametersView;
 	TSharedPtr<class IDetailsView>					ExportSettingsView;
 	TSharedPtr<class IDetailsView>					PreviewSettingsView;
+	TSharedPtr<class IDetailsView>					PreviewSceneSettingsView;
+	
+	/** Scene preview settings widget */
+	TSharedPtr<SWidget>								AdvancedPreviewSettingsWidget;
+	
+	TWeakPtr<SDockTab>								PreviewSceneSettingsDockTab;
+	
+	FAdvancedPreviewSceneModule::FOnPreviewSceneChanged OnPreviewSceneChangedDelegate;
+	
 // Tracking the active viewports in this editor.
 	TSharedPtr<class FEditorViewportTabContent>		ViewportTabContent;
 
@@ -116,21 +130,34 @@ private:
 
 	virtual TStatId									GetStatId() const override;
 
+	/** Post edit change notify for properties. */
+	virtual void									NotifyPostChange( const FPropertyChangedEvent& PropertyChangedEvent, FProperty* PropertyThatChanged ) override;
+	
 	// Inherited via FGCObject
 	virtual void									AddReferencedObjects(FReferenceCollector& Collector) override;
 	virtual FString									GetReferencerName() const override { return TEXT("FTextureGraphExporter");}
 
 private:
+	
+	/** Builds the sub-tools that are a part of the texture graph editor. */
+	void											BuildSubTools();
+	
+	/** Called when the Viewport Layout has changed. */
+	void											OnEditorLayoutChanged();
+	
 	TSharedRef<SDockTab>							SpawnTab_Viewport(const FSpawnTabArgs& Args);
 	TSharedRef<SDockTab>							SpawnTab_ParameterDefaults(const FSpawnTabArgs& Args);
 	TSharedRef<SDockTab>							SpawnTab_NodePreview(const FSpawnTabArgs& Args);
 	TSharedRef<SDockTab>							SpawnTab_ExportSettings(const FSpawnTabArgs& Args);
+	TSharedRef<SDockTab>							SpawnTab_TG_Properties(const FSpawnTabArgs& Args);
 	TSharedRef<SDockTab>							SpawnTab_PreviewSettings(const FSpawnTabArgs& Args);
+	TSharedRef<SDockTab>							SpawnTab_PreviewSceneSettings(const FSpawnTabArgs& Args);
 	TSharedPtr<class STG_EditorViewport>			GetEditorViewport() const;
 	
 	bool 											SetPreviewAsset(UObject* InAsset);
 	TSharedRef<class IDetailsView>					GetDetailsView() const { return DetailsView.ToSharedRef(); }
 	TSharedRef<class IDetailsView>					GetPreviewSettingsView() const { return PreviewSettingsView.ToSharedRef(); }
+	TSharedRef<class IDetailsView>					GetPreviewSceneSettingsView() const { return PreviewSceneSettingsView.ToSharedRef(); }
 	TSharedRef<class IDetailsView>					GetExportSettingsView() const { return ExportSettingsView.ToSharedRef(); }
 	TSharedRef<class IDetailsView>					GetParametersView() const { return ParametersView.ToSharedRef(); }
 	void											OnViewportMaterialChanged();
diff --git a/Engine/Plugins/Experimental/TextureGraph/Source/TextureGraphEditor/Public/TextureGraphEditorModule.h b/Engine/Plugins/Experimental/TextureGraph/Source/TextureGraphEditor/Public/TextureGraphEditorModule.h
--- a/Engine/Plugins/Experimental/TextureGraph/Source/TextureGraphEditor/Public/TextureGraphEditorModule.h
+++ b/Engine/Plugins/Experimental/TextureGraph/Source/TextureGraphEditor/Public/TextureGraphEditorModule.h
@@ -33,7 +33,6 @@ public:
 	/** IModuleInterface implementation */
 	virtual void					StartupModule() override;
 	virtual void					ShutdownModule() override;
-	virtual void					PreUnloadCallback() override;
 
 	
 	virtual void					StartTextureGraphEngine();
diff --git a/Engine/Plugins/Experimental/TextureGraph/Source/TextureGraphEngine/FxMat/RenderMaterial_BP.cpp b/Engine/Plugins/Experimental/TextureGraph/Source/TextureGraphEngine/FxMat/RenderMaterial_BP.cpp
--- a/Engine/Plugins/Experimental/TextureGraph/Source/TextureGraphEngine/FxMat/RenderMaterial_BP.cpp
+++ b/Engine/Plugins/Experimental/TextureGraph/Source/TextureGraphEngine/FxMat/RenderMaterial_BP.cpp
@@ -25,7 +25,7 @@ TSharedPtr<FMaterialInstanceStaticSwitchPermutationMap> FMaterialInstanceStaticS
 	{
 		MISSPMap = MakeShared<FMaterialInstanceStaticSwitchPermutationMap>();
 
-		UMaterialInstanceConstant* MIC = NewObject<UMaterialInstanceConstant>(InMaterial);
+		TStrongObjectPtr<UMaterialInstanceConstant> MIC = TStrongObjectPtr<UMaterialInstanceConstant>(NewObject<UMaterialInstanceConstant>(InMaterial));
 #if WITH_EDITOR
 		InMaterial->GetStaticParameterValues(MISSPMap->DefaultStaticParameterSet);
 		MISSPMap->DefaultKey = MISSPMap->KeyFromStaticSwitchParameters(MISSPMap->DefaultStaticParameterSet.StaticSwitchParameters);
@@ -44,12 +44,12 @@ TSharedPtr<FMaterialInstanceStaticSwitchPermutationMap> FMaterialInstanceStaticS
 
 UMaterialInterface* FMaterialInstanceStaticSwitchPermutationMap::GetRootMaterial()
 {
-	TObjectPtr<UMaterialInstanceConstant>* MICP = PermutationsMap.Find(DefaultKey);
+	TStrongObjectPtr<UMaterialInstanceConstant>* MICP = PermutationsMap.Find(DefaultKey);
 	if (MICP)
 	{
 		// Use the FIRST MIC that initiated the MISWPermutationMap as the "root Material" 
 		// in order to preserve the default values
-		return Cast<UMaterialInterface>((*MICP));
+		return Cast<UMaterialInterface>(MICP->Get());
 	}
 	return nullptr;
 }
@@ -72,10 +72,10 @@ UMaterialInstanceConstant* FMaterialInstanceStaticSwitchPermutationMap::GetMater
 
 	int32 Key = KeyFromStaticSwitchParameters(Parameters);
 
-	TObjectPtr<UMaterialInstanceConstant>* MICP = PermutationsMap.Find(Key);
+	TStrongObjectPtr<UMaterialInstanceConstant>* MICP = PermutationsMap.Find(Key);
 	if (!MICP)
 	{
-		UMaterialInstanceConstant* MIC = NewObject<UMaterialInstanceConstant>(GetRootMaterial());
+		TStrongObjectPtr<UMaterialInstanceConstant> MIC = TStrongObjectPtr<UMaterialInstanceConstant>(NewObject<UMaterialInstanceConstant>(GetRootMaterial()));
 #if WITH_EDITOR
 		MIC->SetParentEditorOnly(GetRootMaterial());
 		MIC->ClearParameterValuesEditorOnly();
@@ -90,7 +90,7 @@ UMaterialInstanceConstant* FMaterialInstanceStaticSwitchPermutationMap::GetMater
 #endif
 		PermutationsMap.Add(Key, MIC);
 
-		return MIC;
+		return MIC.Get();
 	}
 	else
 	{
diff --git a/Engine/Plugins/Experimental/TextureGraph/Source/TextureGraphEngine/FxMat/RenderMaterial_BP.h b/Engine/Plugins/Experimental/TextureGraph/Source/TextureGraphEngine/FxMat/RenderMaterial_BP.h
--- a/Engine/Plugins/Experimental/TextureGraph/Source/TextureGraphEngine/FxMat/RenderMaterial_BP.h
+++ b/Engine/Plugins/Experimental/TextureGraph/Source/TextureGraphEngine/FxMat/RenderMaterial_BP.h
@@ -17,7 +17,7 @@ struct TEXTUREGRAPHENGINE_API FMaterialInstanceStaticSwitchPermutationMap
 {
 	int32												DefaultKey = 0;
 	FStaticParameterSet									DefaultStaticParameterSet;
-	TMap<int32, TObjectPtr<UMaterialInstanceConstant>>	PermutationsMap;
+	TMap<int32, TStrongObjectPtr<UMaterialInstanceConstant>>	PermutationsMap;
 
 	static TSharedPtr<FMaterialInstanceStaticSwitchPermutationMap> Create(UMaterialInterface* InMaterial);
 	UMaterialInterface*				GetRootMaterial();
