Horde test failure fix

diff --git a/Engine/Plugins/Experimental/TextureGraph/Source/TextureGraphEditor/Private/TG_Exporter.cpp b/Engine/Plugins/Experimental/TextureGraph/Source/TextureGraphEditor/Private/TG_Exporter.cpp
--- a/Engine/Plugins/Experimental/TextureGraph/Source/TextureGraphEditor/Private/TG_Exporter.cpp
+++ b/Engine/Plugins/Experimental/TextureGraph/Source/TextureGraphEditor/Private/TG_Exporter.cpp
@@ -722,8 +722,10 @@ FTG_ExporterUtility::FTG_ExporterUtility()
 		.SetGroup(MenuStructure.GetDeveloperToolsMiscCategory())
 		.SetIcon(FSlateIcon(FAppStyle::GetAppStyleSetName(), "ClassIcon.Texture2D"));
 
-	FSlateApplication::Get().OnPreShutdown().AddRaw(this, &FTG_ExporterUtility::Cleanup);
-	
+	if (FSlateApplication::IsInitialized())
+	{
+		FSlateApplication::Get().OnPreShutdown().AddRaw(this, &FTG_ExporterUtility::Cleanup);
+	}
 }
 
 void FTG_ExporterUtility::Cleanup()
@@ -736,9 +738,10 @@ void FTG_ExporterUtility::Cleanup()
 			Impl->UnregisterTabSpawners(TGExporterTabManager);
 		}
 		FGlobalTabmanager::Get()->UnregisterNomadTabSpawner(FTG_EditorTabs::TextureExporterTabId);
+		
+		FSlateApplication::Get().OnPreShutdown().RemoveAll(this);
 	}
 	
-	FSlateApplication::Get().OnPreShutdown().RemoveAll(this);
 }
 
 FTG_ExporterUtility::~FTG_ExporterUtility()
