Texture Graph: Instance fixes for blueprint and export/output settings
Also updated/re-saved some other TG test assets which had their output nodes showing warnings

#virtualized

diff --git a/Engine/Plugins/Experimental/TextureGraph/Source/TextureGraph/Private/Blueprint/TG_AsyncExportTask.cpp b/Engine/Plugins/Experimental/TextureGraph/Source/TextureGraph/Private/Blueprint/TG_AsyncExportTask.cpp
--- a/Engine/Plugins/Experimental/TextureGraph/Source/TextureGraph/Private/Blueprint/TG_AsyncExportTask.cpp
+++ b/Engine/Plugins/Experimental/TextureGraph/Source/TextureGraph/Private/Blueprint/TG_AsyncExportTask.cpp
@@ -14,7 +14,7 @@ UTG_AsyncExportTask::UTG_AsyncExportTask(const FObjectInitializer& ObjectInitial
 
 }
 
-UTG_AsyncExportTask* UTG_AsyncExportTask::TG_AsyncExportTask( UTextureGraph* InTextureGraph, const bool OverwriteTextures, const bool bSave)
+UTG_AsyncExportTask* UTG_AsyncExportTask::TG_AsyncExportTask( UTextureGraphBase* InTextureGraph, const bool OverwriteTextures, const bool bSave)
 {
 	UTG_AsyncExportTask* Task = NewObject<UTG_AsyncExportTask>();
 	Task->SetFlags(RF_Standalone);
diff --git a/Engine/Plugins/Experimental/TextureGraph/Source/TextureGraph/Private/Blueprint/TG_AsyncRenderTask.cpp b/Engine/Plugins/Experimental/TextureGraph/Source/TextureGraph/Private/Blueprint/TG_AsyncRenderTask.cpp
--- a/Engine/Plugins/Experimental/TextureGraph/Source/TextureGraph/Private/Blueprint/TG_AsyncRenderTask.cpp
+++ b/Engine/Plugins/Experimental/TextureGraph/Source/TextureGraph/Private/Blueprint/TG_AsyncRenderTask.cpp
@@ -19,7 +19,7 @@ UTG_AsyncRenderTask::UTG_AsyncRenderTask(const FObjectInitializer& ObjectInitial
 	
 }
 
-UTG_AsyncRenderTask* UTG_AsyncRenderTask::TG_AsyncRenderTask( UTextureGraph* InTextureGraph)
+UTG_AsyncRenderTask* UTG_AsyncRenderTask::TG_AsyncRenderTask( UTextureGraphBase* InTextureGraph)
 {
 	UTG_AsyncRenderTask* Task = NewObject<UTG_AsyncRenderTask>();
 	Task->SetFlags(RF_Standalone);
diff --git a/Engine/Plugins/Experimental/TextureGraph/Source/TextureGraph/Private/TG_Graph.cpp b/Engine/Plugins/Experimental/TextureGraph/Source/TextureGraph/Private/TG_Graph.cpp
--- a/Engine/Plugins/Experimental/TextureGraph/Source/TextureGraph/Private/TG_Graph.cpp
+++ b/Engine/Plugins/Experimental/TextureGraph/Source/TextureGraph/Private/TG_Graph.cpp
@@ -4,9 +4,11 @@
 #include "TG_Graph.h"
 #include "Expressions/TG_Expression.h"
 #include "Expressions/Output/TG_Expression_Output.h"
+#include "TG_OutputSettings.h"
 #include "TG_CustomVersion.h"
 
 #include "Algo/Reverse.h"
+#include "Expressions/Input/TG_Expression_OutputSettings.h"
 #include "UObject/ObjectSaveContext.h"
 
 #if WITH_EDITOR
@@ -462,6 +464,41 @@ bool UTG_Graph::Connect(UTG_Node& NodeFrom, FTG_Name& PinFromName, UTG_Node& Nod
 	return true;
 }
 
+void UTG_Graph::CollectOutputSettings(TMap<FTG_Id, FTG_OutputSettings>& OutExportSettingsMap) const
+{
+	ForEachNodes([&](const UTG_Node* Node, uint32 Index)
+	{
+		if (UTG_Expression_Output* OutputExpression = dynamic_cast<UTG_Expression_Output*>(Node->GetExpression()))
+		{
+			UTG_Pin* SettingsPin = Node->GetPin(GET_MEMBER_NAME_CHECKED(UTG_Expression_Output, OutputSettings));
+			
+			// skip connected pins as they will be exposed as Input Parameters
+			if (!SettingsPin->IsConnected())
+			{
+				OutExportSettingsMap.Add(Node->Id, OutputExpression->OutputSettings);
+			}
+		}
+	});
+}
+
+void UTG_Graph::SetOutputSettings(const TMap<FTG_Id, FTG_OutputSettings>& OutputSettingsMap)
+{
+	ForEachNodes([&](const UTG_Node* Node, uint32 Index)
+	{
+		if (UTG_Expression_Output* OutputExpression = dynamic_cast<UTG_Expression_Output*>(Node->GetExpression()))
+		{
+			const FTG_OutputSettings* OutputSettingsToSet = OutputSettingsMap.Find(Node->Id);
+			if (OutputSettingsToSet)
+			{
+				OutputExpression->OutputSettings = *OutputSettingsToSet;
+				UTG_Pin* SettingsPin = Node->GetPin(GET_MEMBER_NAME_CHECKED(UTG_Expression_Output, OutputSettings));
+				SettingsPin->FromString(OutputSettingsToSet->ToString());
+			}
+		}
+	});
+}
+
+
 bool UTG_Graph::Validate(MixUpdateCyclePtr Cycle)
 {
 	// TODO: Add any graph level validation steps here
@@ -871,7 +908,9 @@ FTG_VarMap UTG_Graph::GetInputParamsVarMap()
 		{
 			if (ParamPin->IsInput())
 			{
-				Map.VarArguments.Add(ParamPin->GetAliasName(), { *ParamPin->EditSelfVar(), ParamPin->GetArgument() });
+				const auto VarId = ParamPin->GetVarId();
+				FTG_Var* Var = GetVar(VarId);
+				Map.VarArguments.Add(ParamPin->GetAliasName(), { *Var, ParamPin->GetArgument() });
 			}
 		}
 	};
@@ -886,13 +925,16 @@ void UTG_Graph::SetInputParamsFromVarMap(FTG_VarMap InMap)
 		UTG_Pin* Pin = Cast<UTG_Pin>(GetPin(InputId));
 		FTG_Var* VarToCopy = InMap.GetVar(Pin->GetAliasName());
 		const FTG_Argument* VarArgument = InMap.GetVarArgument(Pin->GetAliasName());
+
 		if (VarToCopy)
 		{
-			Pin->EditSelfVar()->CopyFrom(VarToCopy);
+			const auto VarId = Pin->GetVarId();
+			FTG_Var* Var = GetVar(VarId);
+			Var->CopyFrom(VarToCopy);
 
 			UTG_Expression* OwnerExpression = Pin->GetNodePtr()->GetExpression();
 			// also copy to expression
-			OwnerExpression->CopyVarToExpressionArgument(*VarArgument, Pin->EditSelfVar());
+			OwnerExpression->CopyVarToExpressionArgument(*VarArgument, Var);
 		}
 	};
 	
diff --git a/Engine/Plugins/Experimental/TextureGraph/Source/TextureGraph/Private/TG_OutputSettings.cpp b/Engine/Plugins/Experimental/TextureGraph/Source/TextureGraph/Private/TG_OutputSettings.cpp
--- a/Engine/Plugins/Experimental/TextureGraph/Source/TextureGraph/Private/TG_OutputSettings.cpp
+++ b/Engine/Plugins/Experimental/TextureGraph/Source/TextureGraph/Private/TG_OutputSettings.cpp
@@ -82,4 +82,19 @@ void FTG_OutputSettings::Initialize(FString PathName,FName InName /*= "Output"*/
 
 	FolderPath = FName(DefaultDirectory);
 	TextureFormat = ETG_TextureFormat::BGRA8;
-}
\ No newline at end of file
+}
+
+FArchive& operator<<(FArchive& Ar, FTG_OutputSettings& T)
+{
+	Ar << T.BaseName;
+	Ar << T.OutputName;
+	Ar << T.FolderPath;
+	Ar << T.Width;
+	Ar << T.Height;
+	Ar << T.TextureFormat;
+	Ar << T.TexturePresetType;
+	Ar << T.LODGroup;
+	Ar << T.Compression;
+	Ar << T.bSRGB;
+	return Ar;
+}
diff --git a/Engine/Plugins/Experimental/TextureGraph/Source/TextureGraph/Private/TG_Var.cpp b/Engine/Plugins/Experimental/TextureGraph/Source/TextureGraph/Private/TG_Var.cpp
--- a/Engine/Plugins/Experimental/TextureGraph/Source/TextureGraph/Private/TG_Var.cpp
+++ b/Engine/Plugins/Experimental/TextureGraph/Source/TextureGraph/Private/TG_Var.cpp
@@ -694,7 +694,6 @@ void VarArchiveSerializer_float(FTG_Var::VarArchiveSerialInfo& Info)
 {
 	Generic_Simple_ArSerializer<float>(Info);
 }
-
 void VarArchiveSerializer_bool(FTG_Var::VarArchiveSerialInfo& Info)
 {
 	Generic_Simple_ArSerializer<bool>(Info);
@@ -739,6 +738,10 @@ void VarArchiveSerializer_UTexture(FTG_Var::VarArchiveSerialInfo& Info)
 	}
 
 }
+void VarArchiveSerializer_FTG_OutputSettings(FTG_Var::VarArchiveSerialInfo& Info)
+{
+	Generic_Simple_ArSerializer<TObjectPtr<UTexture>>(Info);
+}
 //////////////////////////////////////////////////////////////////////////////////////////////////////////////////////////////////////
 FTG_Var::VarArchiveSerializerMap FTG_Var::DefaultArchiveSerializers
 (
@@ -752,7 +755,8 @@ FTG_Var::VarArchiveSerializerMap FTG_Var::DefaultArchiveSerializers
 		VAR_ARCHIVE_SERIALIZER_DEF(FVector4f),
 		VAR_ARCHIVE_SERIALIZER_DEF(FVector2f),
 		VAR_ARCHIVE_SERIALIZER_DEF(FLinearColor),
-		VAR_ARCHIVE_SERIALIZER_DEF(FTG_VariantArray)
+		VAR_ARCHIVE_SERIALIZER_DEF(FTG_VariantArray),
+		VAR_ARCHIVE_SERIALIZER_DEF(FTG_OutputSettings)
 	}
 ); 
 
diff --git a/Engine/Plugins/Experimental/TextureGraph/Source/TextureGraph/Private/TextureGraph.cpp b/Engine/Plugins/Experimental/TextureGraph/Source/TextureGraph/Private/TextureGraph.cpp
--- a/Engine/Plugins/Experimental/TextureGraph/Source/TextureGraph/Private/TextureGraph.cpp
+++ b/Engine/Plugins/Experimental/TextureGraph/Source/TextureGraph/Private/TextureGraph.cpp
@@ -40,8 +40,16 @@ void UTextureGraphInstance::Construct(FString InName)
 
 void UTextureGraphInstance::CopyParamsToRuntimeGraph()
 {
-	check (RuntimeGraph);
 	RuntimeGraph->SetInputParamsFromVarMap(InputParams);
+
+	if (!OutputSettingsMap.IsEmpty())
+	{
+		RuntimeGraph->SetOutputSettings(OutputSettingsMap);
+	}
+	else
+	{
+		UpdateOutputSettingsFromGraph();
+	}
 }
 
 void UTextureGraphInstance::Update(MixUpdateCyclePtr InCycle)
@@ -54,15 +62,12 @@ void UTextureGraphInstance::Update(MixUpdateCyclePtr InCycle)
 
 	// Now Evaluate the Graph!
 	FTG_EvaluationContext EvaluationContext;
-	EvaluationContext.Inputs = InputParams;
 	EvaluationContext.Cycle = InCycle;
 	EvaluationContext.Cycle->PushMix(this);
 	
 	if (Graph()->Validate(InCycle))
 		Graph()->Evaluate(&EvaluationContext);
 	
-	InputParams = Graph()->GetInputParamsVarMap();
-	
 	EvaluationContext.Cycle->PopMix();
 	
 	/// This will be the final result of the rendering
@@ -83,18 +88,54 @@ void UTextureGraphInstance::SetParent(const TObjectPtr<UTextureGraphBase>& Paren
 		RuntimeGraph = TStrongObjectPtr<UTG_Graph>(DuplicateObject(Parent->Graph(), GetTransientPackage()));
 		Settings = DuplicateObject(Parent->GetSettings(), this);
 
-		// copy InputParams from Parent if it is an instance
+		// copy InputParams and ExportSettings from Parent if it is an instance
 		if (Parent.IsA<UTextureGraphInstance>())
 		{
-			InputParams = Cast<UTextureGraphInstance>(Parent)->InputParams;	
+			InputParams = Cast<UTextureGraphInstance>(Parent)->InputParams;
+			OutputSettingsMap = Cast<UTextureGraphInstance>(Parent)->OutputSettingsMap;
 		}
 		else
 		{
 			InputParams = Graph()->GetInputParamsVarMap();
+			
+			// recreate fresh output settings from graph
+			UpdateOutputSettingsFromGraph();
 		}
 	}
 }
 
+void UTextureGraphInstance::PreSave(FObjectPreSaveContext SaveContext)
+{
+	Super::PreSave(SaveContext);
+
+	// save local param changes from graph updates before saving
+	InputParams = Graph()->GetInputParamsVarMap();
+	UpdateOutputSettingsFromGraph();
+}
+
+bool UTextureGraphInstance::CheckOutputSettingsMatch(const TObjectPtr<UTextureGraphBase>& Parent) const
+{
+	if (OutputSettingsMap.IsEmpty())
+		return false;
+
+	FTG_Ids OutputIds = Parent->Graph()->GetOutputParamIds();
+
+	if (OutputIds.Num() != OutputSettingsMap.Num())
+		return false;
+
+	bool returnValue = true;
+	for(FTG_Id OutputId : OutputIds)
+	{
+		returnValue &= OutputSettingsMap.Contains(OutputId);
+	}
+	return returnValue;
+}
+void UTextureGraphInstance::UpdateOutputSettingsFromGraph()
+{
+	Graph()->CollectOutputSettings(OutputSettingsMap);
+	//TODO: Keep values if names match
+}
+
 UTG_Graph* UTextureGraphInstance::Graph()
 {
 	if (!RuntimeGraph && ParentTextureGraph)
diff --git a/Engine/Plugins/Experimental/TextureGraph/Source/TextureGraph/Public/Blueprint/TG_AsyncExportTask.h b/Engine/Plugins/Experimental/TextureGraph/Source/TextureGraph/Public/Blueprint/TG_AsyncExportTask.h
--- a/Engine/Plugins/Experimental/TextureGraph/Source/TextureGraph/Public/Blueprint/TG_AsyncExportTask.h
+++ b/Engine/Plugins/Experimental/TextureGraph/Source/TextureGraph/Public/Blueprint/TG_AsyncExportTask.h
@@ -21,7 +21,7 @@ class TEXTUREGRAPH_API UTG_AsyncExportTask : public UTG_AsyncTask
 
 public:
 	UFUNCTION(BlueprintCallable, Category = "TextureGraph" , meta=(DisplayName="Texture Graph Export (Async)", BlueprintInternalUseOnly = "true"))
-	static UTG_AsyncExportTask* TG_AsyncExportTask(UTextureGraph* TextureGraph, const bool OverwriteTextures, const bool bSave);
+	static UTG_AsyncExportTask* TG_AsyncExportTask(UTextureGraphBase* TextureGraph, const bool OverwriteTextures, const bool bSave);
 
 	virtual void Activate() override;
 
@@ -39,9 +39,9 @@ private:
 
 	bool bSave = false;
 
-	UTextureGraph* OrignalTextureGraphPtr;
+	UTextureGraphBase* OrignalTextureGraphPtr;
 
-	UTextureGraph* TextureGraphPtr;
+	UTextureGraphBase* TextureGraphPtr;
 
 	FExportSettings TargetExportSettings;
 };
diff --git a/Engine/Plugins/Experimental/TextureGraph/Source/TextureGraph/Public/Blueprint/TG_AsyncRenderTask.h b/Engine/Plugins/Experimental/TextureGraph/Source/TextureGraph/Public/Blueprint/TG_AsyncRenderTask.h
--- a/Engine/Plugins/Experimental/TextureGraph/Source/TextureGraph/Public/Blueprint/TG_AsyncRenderTask.h
+++ b/Engine/Plugins/Experimental/TextureGraph/Source/TextureGraph/Public/Blueprint/TG_AsyncRenderTask.h
@@ -23,7 +23,7 @@ class TEXTUREGRAPH_API UTG_AsyncRenderTask : public UTG_AsyncTask
 
 public:
 	UFUNCTION(BlueprintCallable, Category = "TextureGraph" , meta=(DisplayName="Texture Graph Render (Async)" , BlueprintInternalUseOnly = "true"))
-	static UTG_AsyncRenderTask* TG_AsyncRenderTask(UTextureGraph* InTextureGraph);
+	static UTG_AsyncRenderTask* TG_AsyncRenderTask(UTextureGraphBase* InTextureGraph);
 
 	virtual void Activate() override;
 
@@ -45,9 +45,9 @@ private:
 
 	TArray<BlobPtr> OutputBlobs;
 
-	UTextureGraph* OriginalTextureGraphPtr;
+	UTextureGraphBase* OriginalTextureGraphPtr;
 
-	UTextureGraph* TextureGraphPtr;
+	UTextureGraphBase* TextureGraphPtr;
 
 	bool bShouldDestroyOnRenderComplete = false;
 
diff --git a/Engine/Plugins/Experimental/TextureGraph/Source/TextureGraph/Public/TG_Graph.h b/Engine/Plugins/Experimental/TextureGraph/Source/TextureGraph/Public/TG_Graph.h
--- a/Engine/Plugins/Experimental/TextureGraph/Source/TextureGraph/Public/TG_Graph.h
+++ b/Engine/Plugins/Experimental/TextureGraph/Source/TextureGraph/Public/TG_Graph.h
@@ -9,6 +9,8 @@
 #include "TG_Node.h"
 
 #include <functional>
+
+#include "TG_OutputSettings.h"
 #include "TG_Graph.generated.h"
 
 // Graph Traversal
@@ -159,7 +161,7 @@ public:
     // Connect a output pin from a node to the specified input pin of the to node
     // Return true if the edge is created or false if couldn't create the edge
 	bool Connect(UTG_Node& NodeFrom, FTG_Name& PinFrom, UTG_Node& NodeTo, FTG_Name& PinTo);
-
+	
     // Check that the connection from PinFrom to PInTO doesn't create a loop
 	static bool ConnectionCausesLoop(const UTG_Pin* PinFrom, const UTG_Pin* PinTo);
 
@@ -240,8 +242,18 @@ public:
 	// along with their names optionally
 	int GetAllOutputParamValues(TArray<FTG_Variant>& OutVariants, TArray<FName>* OutNames = nullptr) const;
 
+	// Get a map of input parameter vars from the graph
 	FTG_VarMap GetInputParamsVarMap();
+
+	// Apply an override map of input parameter vars to the graph
     void SetInputParamsFromVarMap(FTG_VarMap InMap);
+
+	// get a map of output settings from the graph
+	void CollectOutputSettings(TMap<FTG_Id, FTG_OutputSettings>& OutExportSettingsMap) const;
+
+	// Apply an override map of output settings to the graph
+	void SetOutputSettings(const TMap<FTG_Id, FTG_OutputSettings>& OutputSettingsMap);
+
     //////////////////////////////////////////////////////////////////////////
 	// Traverseal of the graph and evaluation 
 	//////////////////////////////////////////////////////////////////////////
diff --git a/Engine/Plugins/Experimental/TextureGraph/Source/TextureGraph/Public/TG_OutputSettings.h b/Engine/Plugins/Experimental/TextureGraph/Source/TextureGraph/Public/TG_OutputSettings.h
--- a/Engine/Plugins/Experimental/TextureGraph/Source/TextureGraph/Public/TG_OutputSettings.h
+++ b/Engine/Plugins/Experimental/TextureGraph/Source/TextureGraph/Public/TG_OutputSettings.h
@@ -85,7 +85,8 @@ struct TEXTUREGRAPH_API FTG_OutputSettings
 	bool bSRGB = false;
 
 	FString GetFullOutputName() { return  FString::Format(TEXT("{0}"), { BaseName.ToString()});}
-
+	friend FArchive& operator<<(FArchive& Ar, FTG_OutputSettings& T);
+	
 	bool operator==(const FTG_OutputSettings& Other) const
 	{
 		return OutputName == Other.OutputName && BaseName == Other.BaseName;
diff --git a/Engine/Plugins/Experimental/TextureGraph/Source/TextureGraph/Public/TextureGraph.h b/Engine/Plugins/Experimental/TextureGraph/Source/TextureGraph/Public/TextureGraph.h
--- a/Engine/Plugins/Experimental/TextureGraph/Source/TextureGraph/Public/TextureGraph.h
+++ b/Engine/Plugins/Experimental/TextureGraph/Source/TextureGraph/Public/TextureGraph.h
@@ -84,8 +84,11 @@ private:
 	TStrongObjectPtr<UTG_Graph>	RuntimeGraph;		// The runtime graph instance to use for the sub-graph
 public:
 	
-	UPROPERTY(EditAnywhere, Category=TextureGraphInstance)
+	UPROPERTY()
 	FTG_VarMap InputParams;
+
+	UPROPERTY()
+	TMap<FTG_Id, FTG_OutputSettings> OutputSettingsMap;
 	
 	/** Parent Texture Graph. */
 	UPROPERTY(EditAnywhere, BlueprintReadOnly, Category=TextureGraphInstance, AssetRegistrySearchable)
@@ -95,9 +98,15 @@ public:
 
 	void CopyParamsToRuntimeGraph();
 	void Update(MixUpdateCyclePtr InCycle) override;
+	bool CheckOutputSettingsMatch(const TObjectPtr<UTextureGraphBase>& Parent) const;
+	void UpdateOutputSettingsFromGraph();
+	
 	void SetParent(const TObjectPtr<UTextureGraphBase>& Parent);
+
+	// Override PreSave method of UObject
+	virtual void PreSave(FObjectPreSaveContext SaveContext) override;
 	
-	virtual const UTG_Graph* Graph() const override { return RuntimeGraph.Get(); };
+	virtual const UTG_Graph* Graph() const override { return RuntimeGraph.Get(); }
 	virtual UTG_Graph* Graph() override;
 	void Initialize();
 };
\ No newline at end of file
diff --git a/Engine/Plugins/Experimental/TextureGraph/Source/TextureGraphEditor/Private/TG_Exporter.cpp b/Engine/Plugins/Experimental/TextureGraph/Source/TextureGraphEditor/Private/TG_Exporter.cpp
--- a/Engine/Plugins/Experimental/TextureGraph/Source/TextureGraphEditor/Private/TG_Exporter.cpp
+++ b/Engine/Plugins/Experimental/TextureGraph/Source/TextureGraphEditor/Private/TG_Exporter.cpp
@@ -237,6 +237,11 @@ TSharedRef<SDockTab> FTG_InstanceImpl::SpawnTab_TG_Properties(const FSpawnTabArg
 		[
 			DetailsView.ToSharedRef()
 		];
+
+	if (TextureGraphInstance)
+	{
+		DetailsView->SetObject(TextureGraphInstance, true);
+	}
 	return DetailsTab.ToSharedRef();
 }
 
@@ -662,7 +667,6 @@ void FTG_InstanceImpl::UpdateExportSettingsUI()
 }
 void FTG_InstanceImpl::UpdateParametersUI()
 {
-	
 	// Create a new object to set for the view.
 	Parameters = NewObject<UTG_Parameters>();
 
@@ -985,6 +989,9 @@ TSharedRef<SDockTab> FTG_ExporterUtility::CreateTGExporterTab(const FSpawnTabArg
 
 void FTG_ExporterUtility::SetTextureGraphToExport(UTextureGraphBase* InTextureGraph)
 {
+	// force open export window
+	FGlobalTabmanager::Get()->TryInvokeTab(FTG_EditorTabs::TextureExporterTabId);
+	
 	if (InTextureGraph->IsA<UTextureGraph>())
 	{
 		TextureGraphInstance = NewObject<UTextureGraphInstance>();
@@ -998,15 +1005,9 @@ void FTG_ExporterUtility::SetTextureGraphToExport(UTextureGraphBase* InTextureGr
 		{
 			TextureGraphInstance->Construct(FString());
 		}
+		TextureGraphInstance->Initialize();
 	}
-	
-	// force open export window
-	FGlobalTabmanager::Get()->TryInvokeTab(FTG_EditorTabs::TextureExporterTabId);
-	
-	
-	
 	Impl->SetTextureGraphToExport(TextureGraphInstance);
-	
 }
 
 #undef LOCTEXT_NAMESPACE 
diff --git a/Engine/Plugins/Experimental/TextureGraph/Source/TextureGraphEditor/Private/TG_InstanceEditor.cpp b/Engine/Plugins/Experimental/TextureGraph/Source/TextureGraphEditor/Private/TG_InstanceEditor.cpp
--- a/Engine/Plugins/Experimental/TextureGraph/Source/TextureGraphEditor/Private/TG_InstanceEditor.cpp
+++ b/Engine/Plugins/Experimental/TextureGraph/Source/TextureGraphEditor/Private/TG_InstanceEditor.cpp
@@ -73,9 +73,13 @@ void FTG_InstanceEditor::InitEditor(const EToolkitMode::Type Mode, const TShared
 	InstanceImpl->Initialize();
 
 	TextureGraphInstance = InTextureGraph;
-	FAssetEditorToolkit::InitAssetEditor(Mode, InitToolkitHost, TG_EditorAppIdentifier, InstanceImpl->GetDefaultLayout(), /*bCreateDefaultToolbar*/ true, /*bCreateDefaultStandaloneMenu*/ true, InTextureGraph);
-    
-	InstanceImpl->SetTextureGraphToExport(InTextureGraph);
+	if (!InTextureGraph->Graph())
+	{
+		TextureGraphInstance->Construct(FString());
+	}
+	FAssetEditorToolkit::InitAssetEditor(Mode, InitToolkitHost, TG_InstanceEditorAppIdentifier, InstanceImpl->GetDefaultLayout(), /*bCreateDefaultToolbar*/ true, /*bCreateDefaultStandaloneMenu*/ true, TextureGraphInstance);
+	
+	InstanceImpl->SetTextureGraphToExport(TextureGraphInstance);
 	
 
 // 	// Add analytics tag
diff --git a/Engine/Plugins/Experimental/TextureGraph/Source/TextureGraphEditor/Private/TextureGraphEditorModule.cpp b/Engine/Plugins/Experimental/TextureGraph/Source/TextureGraphEditor/Private/TextureGraphEditorModule.cpp
--- a/Engine/Plugins/Experimental/TextureGraph/Source/TextureGraphEditor/Private/TextureGraphEditorModule.cpp
+++ b/Engine/Plugins/Experimental/TextureGraph/Source/TextureGraphEditor/Private/TextureGraphEditorModule.cpp
@@ -25,6 +25,7 @@
 #define LOCTEXT_NAMESPACE "FTextureGraphEditorModule"
 
 const FName TG_EditorAppIdentifier = FName(TEXT("TG_EditorApp"));
+const FName TG_InstanceEditorAppIdentifier = FName(TEXT("TG_InstanceEditorApp"));
 
 void FTextureGraphEditorModule::StartupModule()
 {
diff --git a/Engine/Plugins/Experimental/TextureGraph/Source/TextureGraphEditor/Public/TextureGraphEditorModule.h b/Engine/Plugins/Experimental/TextureGraph/Source/TextureGraphEditor/Public/TextureGraphEditorModule.h
--- a/Engine/Plugins/Experimental/TextureGraph/Source/TextureGraphEditor/Public/TextureGraphEditorModule.h
+++ b/Engine/Plugins/Experimental/TextureGraph/Source/TextureGraphEditor/Public/TextureGraphEditorModule.h
@@ -14,6 +14,7 @@ class FTG_EditorGraphNodeFactory;
 class FTG_EditorGraphPanelPinFactory;
 typedef TArray<TSharedPtr<IAssetTypeActions>> AssetTypeActionsArray;
 extern const FName TG_EditorAppIdentifier;
+extern const FName TG_InstanceEditorAppIdentifier;
 
 class TEXTUREGRAPHEDITOR_API FTextureGraphEditorModule : public IModuleInterface
 {
