Change data util tiles array logic

diff --git a/Engine/Plugins/Experimental/TextureGraph/Source/TextureGraph/Private/Expressions/Input/TG_Expression_TexturePath.cpp b/Engine/Plugins/Experimental/TextureGraph/Source/TextureGraph/Private/Expressions/Input/TG_Expression_TexturePath.cpp
--- a/Engine/Plugins/Experimental/TextureGraph/Source/TextureGraph/Private/Expressions/Input/TG_Expression_TexturePath.cpp
+++ b/Engine/Plugins/Experimental/TextureGraph/Source/TextureGraph/Private/Expressions/Input/TG_Expression_TexturePath.cpp
@@ -45,12 +45,12 @@ bool UTG_Expression_TexturePath::ValidateInputPath(FString& ValidatedPath) const
 		{
 			LocalPath = FPaths::ChangeExtension(LocalPath, PathExt);
 		}
-		ValidatedPath = LocalPath;
+		ValidatedPath = FPaths::ConvertRelativePathToFull(LocalPath);
 		return true;
 	}
-	else if(FPaths::FileExists(LocalPath) || FPaths::DirectoryExists(LocalPath))
+	else if(FPaths::FileExists(LocalPath) || FPaths::DirectoryExists(LocalPath) || FPaths::DirectoryExists(Path))
 	{
-		ValidatedPath = LocalPath;
+		ValidatedPath = FPaths::ConvertRelativePathToFull(LocalPath);
 		return true;
 	}
 
@@ -70,40 +70,43 @@ void UTG_Expression_TexturePath::Evaluate(FTG_EvaluationContext* InContext)
 {
 	Super::Evaluate(InContext);
 
+	/// This is our default state
+	Output.SetNum(1);
+	Output.Set(0, FTG_Texture::GetBlack());
+
 	FString LocalPath;
 	bool bValidation = ValidateInputPath(LocalPath);
 
 	if (!Path.IsEmpty() && bValidation)
 	{
-		if (FPaths::FileExists(Path))
+		FPaths::MakePlatformFilename(LocalPath);
+		if (FPaths::FileExists(LocalPath))
 		{
 			Output.SetNum(1);
 			FTG_Texture Item = LoadStaticImage(InContext, LocalPath, nullptr);
 			Output.Set(0, Item);
 		}
-		else if (FPaths::DirectoryExists(Path))
+		else if (FPaths::DirectoryExists(LocalPath))
 		{
 			FFileManagerGeneric FileManager;
 			TArray<FString> Files;
-			FileManager.FindFiles(Files, *Path);
-
-			Output.SetNum(Files.Num());
+			FileManager.FindFiles(Files, *LocalPath);
 
-			/// Now load a texture against all the files
-			for (int32 Index = 0; Index < Files.Num(); Index++)
+			if (!Files.IsEmpty())
 			{
-				const FString& File = Files[Index];
-				FString Filename = FPaths::Combine(Path, File);
-				FTG_Texture Item = LoadStaticImage(InContext, Filename, nullptr);
-				Output.Set(Index, Item);
+				Output.SetNum(Files.Num());
+
+				/// Now load a texture against all the files
+				for (int32 Index = 0; Index < Files.Num(); Index++)
+				{
+					const FString& File = Files[Index];
+					FString Filename = FPaths::Combine(LocalPath, File);
+					FTG_Texture Item = LoadStaticImage(InContext, Filename, nullptr);
+					Output.Set(Index, Item);
+				}
 			}
 		}
 	}
-	else
-	{
-		Output.SetNum(1);
-		Output.Set(0, FTG_Texture::GetBlack());
-	}
 
 	//For the connected pin we will report error here in eveluate because it does not have the updated value during validation.
 	UTG_Pin* PathPin = GetParentNode()->GetPin(GET_MEMBER_NAME_CHECKED(UTG_Expression_TexturePath, Path));
diff --git a/Engine/Plugins/Experimental/TextureGraph/Source/TextureGraphEngine/2D/TargetTextureSet.cpp b/Engine/Plugins/Experimental/TextureGraph/Source/TextureGraphEngine/2D/TargetTextureSet.cpp
--- a/Engine/Plugins/Experimental/TextureGraph/Source/TextureGraphEngine/2D/TargetTextureSet.cpp
+++ b/Engine/Plugins/Experimental/TextureGraph/Source/TextureGraphEngine/2D/TargetTextureSet.cpp
@@ -128,29 +128,28 @@ void TargetTextureSet::Free()
 
 void TargetTextureSet::BindOnTextureUpdate(RenderMaterial_BPPtr InMaterial, FMaterialMappingInfo MaterialMappingInfo) const
 {
-	TextureReadyCallback bindCallback = [this, InMaterial, MaterialMappingInfo](TiledBlobRef texture) mutable
+	TextureReadyCallback bindCallback = [this, InMaterial, MaterialMappingInfo](TiledBlobRef Texture) mutable
 	{
-		Util::OnGameThread([this, texture, InMaterial, MaterialMappingInfo]()
+		Util::OnGameThread([this, Texture, InMaterial, MaterialMappingInfo]()
 			{
 				ResourceBindInfo BindInfo;
 				BindInfo.Dev = Device_FX::Get();
-				if (texture)
+				if (Texture)
 				{
 					BindInfo.Target = MaterialMappingInfo.MaterialInput.ToString();
 					//TODO: Need to move it into the Resource Bind Info
 
 					// Linear textures are now being displayed in linear gamma.
 					// So turning off the conversion to match UE's convention.
-					// InMaterial->SetInt(TEXT("sRGB"), texture.get()->GetDescriptor().bIsSRGB ? 1 : 0);
+					// InMaterial->SetInt(TEXT("sRGB"), Texture.get()->GetDescriptor().bIsSRGB ? 1 : 0);
 					FString GrayScaleParam = "Is" + BindInfo.Target + "GrayScale";
 					
-					InMaterial->SetInt(FName(GrayScaleParam), texture.get()->GetDescriptor().ItemsPerPoint == 1 ? 1 : 0);
+					InMaterial->SetInt(FName(GrayScaleParam), Texture.get()->GetDescriptor().ItemsPerPoint == 1 ? 1 : 0);
 
-					texture->OnFinalise().then([=]()
+					Texture->OnFinalise().then([=]()
 					{
-						texture->Bind(InMaterial.get(), BindInfo);
+						Texture->Bind(InMaterial.get(), BindInfo);
 					});
-					
 				}
 			});
 	};
diff --git a/Engine/Plugins/Experimental/TextureGraph/Source/TextureGraphEngine/2D/Tex.cpp b/Engine/Plugins/Experimental/TextureGraph/Source/TextureGraphEngine/2D/Tex.cpp
--- a/Engine/Plugins/Experimental/TextureGraph/Source/TextureGraphEngine/2D/Tex.cpp
+++ b/Engine/Plugins/Experimental/TextureGraph/Source/TextureGraphEngine/2D/Tex.cpp
@@ -1082,7 +1082,7 @@ AsyncTiledBlobRef Tex::ToSingleBlob(CHashPtr Hash, bool TransferToRT /* = false
 #endif
 }
 
-AsyncTiledBlobRef Tex::ToBlob(int32 XTiles, int32 YTiles, uint32 Width /* = 0 */, uint32 Height /* = 0 */, bool TransferToRT /* = false */)
+AsyncTiledBlobRef Tex::ToBlob(int32 NumRows, int32 NumCols, uint32 Width /* = 0 */, uint32 Height /* = 0 */, bool TransferToRT /* = false */)
 {
 	check(IsInGameThread());
 
@@ -1108,7 +1108,7 @@ AsyncTiledBlobRef Tex::ToBlob(int32 XTiles, int32 YTiles, uint32 Width /* = 0 */
 	} 
 #endif 
 
-	if (!TextureHelper::CanSplitToTiles(CheckSizeX, CheckSizeY, XTiles, YTiles))
+	if (!TextureHelper::CanSplitToTiles(CheckSizeX, CheckSizeY, NumRows, NumCols))
 	{
 		/// TOOD: need to properly calculate Hash over here
 		return ToSingleBlob(nullptr, TransferToRT);
@@ -1119,24 +1119,24 @@ AsyncTiledBlobRef Tex::ToBlob(int32 XTiles, int32 YTiles, uint32 Width /* = 0 */
 	if (Desc.bMipMaps && IsValid(Texture))
 		TransferToRT = true;	//Force transfer if desc has mips generation on and we are not RT.
 
-	uint32 WidthTilesOffset = Width % XTiles;
+	uint32 WidthTilesOffset = Width % NumCols;
 	if (WidthTilesOffset != 0)
 	{
-		Width += (XTiles - WidthTilesOffset);
+		Width += (NumCols - WidthTilesOffset);
 		Desc.Width = Width;
 		TransferToRT = true;
 	} 
 
-	uint32 HeightTilesOffset = Height % YTiles;
+	uint32 HeightTilesOffset = Height % NumRows;
 	if (HeightTilesOffset != 0)
 	{
-		Height += (YTiles - HeightTilesOffset);
+		Height += (NumRows - HeightTilesOffset);
 		Desc.Height = Height;
 		TransferToRT = true;
 	} 
 
-	uint32 TileWidth = Width / XTiles;
-	uint32 TileHeight = Height / YTiles;
+	uint32 TileWidth = Width / NumRows;
+	uint32 TileHeight = Height / NumCols;
 
 	if ((TransferToRT && !IsValid(RT)) || !IsValid(Texture) || IsValidVirtualTexture())		// For cases where we want to transfer UTexture2D to UTextureRenderTarget2D
 	{
@@ -1146,10 +1146,11 @@ AsyncTiledBlobRef Tex::ToBlob(int32 XTiles, int32 YTiles, uint32 Width /* = 0 */
 		Texture = SaveTexture;
 	}
 
-	BlobUPtr* TilesPtr = new BlobUPtr [XTiles * YTiles];
+	//BlobUPtr* TilesPtr = new BlobUPtr [NumRows * NumCols];
+	T_Tiles<BlobPtr>* TilesPtr = new T_Tiles<BlobPtr>(NumRows, NumCols);
 	const BufferDescriptor TileDesc = Desc.ToBufferDescriptor(TileWidth, TileHeight);
 
-	Device_FX::InitTiles_Texture(TilesPtr, XTiles, YTiles, TileDesc, false);
+	Device_FX::InitTiles_Texture(TilesPtr, TileDesc, false);
 
 	return Device_FX::Get()->Use()
 		.then([=, this]() mutable
@@ -1172,15 +1173,15 @@ AsyncTiledBlobRef Tex::ToBlob(int32 XTiles, int32 YTiles, uint32 Width /* = 0 */
 			bool bIsNowRT = RT != nullptr;	// Have we shifted to being an RT now?
 
 			DeviceBufferRef Buffer = bIsNowRT ? Device_FX::Get()->CreateFromRT(RT, Desc.ToBufferDescriptor()) : Device_FX::Get()->CreateFromTexture(Texture, Desc.ToBufferDescriptor());
-			T_Tiles<DeviceBufferRef> TileBuffers(XTiles, YTiles);
+			T_Tiles<DeviceBufferRef> TileBuffers(NumRows, NumCols);
 
-			for (int32 TileX = 0; TileX < XTiles; TileX++)
+			for (int32 RowId = 0; RowId < NumRows; RowId++)
 			{
-				for (int32 TileY = 0; TileY < YTiles; TileY++)
+				for (int32 ColId = 0; ColId < NumCols; ColId++)
 				{
-					BlobUPtr& Tile = TilesPtr[TileX * YTiles + TileY];
+					BlobPtr& Tile = (*TilesPtr)[RowId][ColId];
 					check(Tile);
-					TileBuffers[TileX][TileY] = Tile->GetBufferRef();
+					TileBuffers[RowId][ColId] = Tile->GetBufferRef();
 				}
 			}
 
@@ -1188,19 +1189,19 @@ AsyncTiledBlobRef Tex::ToBlob(int32 XTiles, int32 YTiles, uint32 Width /* = 0 */
 			HashObj = Buffer->Hash(false);
 			check(HashObj && HashObj->IsFinal());
 			
-			CombineSplitArgs SplitArgs { Buffer, TileBuffers};
+			CombineSplitArgs SplitArgs { Buffer, TileBuffers };
 
 			return Device::SplitToTiles_Generic(SplitArgs);
 		})
 		.then([=, this]() mutable
 		{
-			BlobPtrTiles ResultTiles(XTiles, YTiles);
+			BlobPtrTiles ResultTiles(NumRows, NumCols);
 
-			for (int32 TileX = 0; TileX < XTiles; TileX++)
+			for (int32 RowId = 0; RowId < NumRows; RowId++)
 			{
-				for (int32 TileY = 0; TileY < YTiles; TileY++)
+				for (int32 ColId = 0; ColId < NumCols; ColId++)
 				{
-					BlobUPtr& Tile = TilesPtr[TileX * YTiles + TileY];
+					BlobPtr& Tile = (*TilesPtr)[RowId][ColId]; // [RowId * NumCols + ColId] ;
 					check(Tile);
 
 					//check(Tile->IsFinalised());
@@ -1213,7 +1214,7 @@ AsyncTiledBlobRef Tex::ToBlob(int32 XTiles, int32 YTiles, uint32 Width /* = 0 */
 
 					/// Add the tiled blobs to blobber for re-using later
 					/// We do this here since the SplitToTiles function is where the RawObj Hash is updated and we add to blobber after that.
-					ResultTiles[TileX][TileY] = TextureGraphEngine::GetBlobber()->AddResult(TileHash, std::move(Tile));
+					ResultTiles[RowId][ColId] = TextureGraphEngine::GetBlobber()->AddResult(TileHash, Tile);
 				}
 			}
 
@@ -1232,7 +1233,7 @@ AsyncTiledBlobRef Tex::ToBlob(int32 XTiles, int32 YTiles, uint32 Width /* = 0 */
 			
 			auto TiledBlobRef = TextureGraphEngine::GetBlobber()->AddTiledResult(TiledBlobObj);
 
-			delete[] TilesPtr;
+			delete TilesPtr;
 
 			return TiledBlobRef;
 		});
diff --git a/Engine/Plugins/Experimental/TextureGraph/Source/TextureGraphEngine/2D/TextureHelper.cpp b/Engine/Plugins/Experimental/TextureGraph/Source/TextureGraphEngine/2D/TextureHelper.cpp
--- a/Engine/Plugins/Experimental/TextureGraph/Source/TextureGraphEngine/2D/TextureHelper.cpp
+++ b/Engine/Plugins/Experimental/TextureGraph/Source/TextureGraphEngine/2D/TextureHelper.cpp
@@ -1064,10 +1064,10 @@ void TextureHelper::RawFromMem_Tiled(const uint8* SrcData, size_t SrcDataLength,
 			uint8* DstData = new uint8[TileSize];
 			memset(DstData, 0, TileSize);
 
-			Tiles[ColIndex][RowIndex] = DstData;
+			Tiles[RowIndex][ColIndex] = DstData;
 
-			BufferDescriptor thisTileDesc = DstTileDesc;
-			thisTileDesc.Name = TextureHelper::CreateTileName(SrcDesc.Name, ColIndex, RowIndex);
+			BufferDescriptor ThisTileDesc = DstTileDesc;
+			ThisTileDesc.Name = TextureHelper::CreateTileName(SrcDesc.Name, ColIndex, RowIndex);
 		}
 	}
 	
@@ -1079,24 +1079,24 @@ void TextureHelper::RawFromMem_Tiled(const uint8* SrcData, size_t SrcDataLength,
 		int RowIndex = FMath::Abs(Y / TileHeight);
 		for (int ColIndex = 0; ColIndex < NumCols; ColIndex++)
 		{
-			uint8* Tile = Tiles[ColIndex][RowIndex];
+			uint8* Tile = Tiles[RowIndex][ColIndex];
 			uint8* Dst = Tile + ((Y % TileHeight) * DstTilePitch);
 			const uint8* Src = SrcData + (ColIndex * SrcTilePitch) + (Y * SrcPitch);
 			Converter(Dst, DstTilePitch, DstTileDesc, Src, SrcTilePitch, SrcDesc);
 		}
 	});
 
-	for (size_t TileX = 0; TileX < NumRows; TileX++)
+	for (size_t RowId = 0; RowId < NumRows; RowId++)
 	{
-		for (size_t TileY = 0; TileY < NumCols; TileY++)
+		for (size_t ColId = 0; ColId < NumCols; ColId++)
 		{
 			BufferDescriptor ThisTileDesc = DstTileDesc;
-			ThisTileDesc.Name = TextureHelper::CreateTileName(SrcDesc.Name, TileX, TileY);
+			ThisTileDesc.Name = TextureHelper::CreateTileName(SrcDesc.Name, RowId, ColId);
 
 			/// Only TileSize use it for actual hashing calculation
-			HashType TileHashValue = DataUtil::Hash(Tiles[TileX][TileY], TileSize);
+			HashType TileHashValue = DataUtil::Hash(Tiles[RowId][ColId], TileSize);
 			CHashPtr TileHash = std::make_shared<CHash>(TileHashValue, true);
-			RawTiles[TileX][TileY] = std::make_shared<RawBuffer>(Tiles[TileX][TileY], ActualTileSize, ThisTileDesc, TileHash);
+			RawTiles[RowId][ColId] = std::make_shared<RawBuffer>(Tiles[RowId][ColId], ActualTileSize, ThisTileDesc, TileHash);
 		}
 	}
 }
@@ -1143,10 +1143,10 @@ RawBufferPtr TextureHelper::CombineRaw_Tiles(const RawBufferPtrTiles& Tiles, CHa
 		size_t DstPitch = Desc.Pitch();
 
 		Data = new uint8 [DstPitch * Desc.Height];
-		CHashPtrVec TileHashes; 
+		T_Tiles<CHashPtr> TileHashes(0, 0); 
 
 		if (HashValue == nullptr || HashValue->IsNull())
-			TileHashes.resize(NumRows * NumCols);
+			TileHashes.Resize(NumRows, NumCols);
 
 		ParallelFor(TotalTiles, [&](int32 TileIndex)
 		{
@@ -1165,12 +1165,12 @@ RawBufferPtr TextureHelper::CombineRaw_Tiles(const RawBufferPtrTiles& Tiles, CHa
 			if (HashValue == nullptr || HashValue->IsNull())
 			{
 				CHashPtr TileHash = Tile->Hash();
-				TileHashes[ColIndex * NumRows + RowIndex] = TileHash;
+				TileHashes[RowIndex][ColIndex] = TileHash;
 			}
 		});
 
 		if (HashValue == nullptr || HashValue->IsNull())
-			HashValue = CHash::ConstructFromSources(TileHashes);
+			HashValue = CHash::ConstructFromSources(TileHashes.Tiles());
 
 		DataLength = Desc.Size();
 
diff --git a/Engine/Plugins/Experimental/TextureGraph/Source/TextureGraphEngine/Data/Blobber.h b/Engine/Plugins/Experimental/TextureGraph/Source/TextureGraphEngine/Data/Blobber.h
--- a/Engine/Plugins/Experimental/TextureGraph/Source/TextureGraphEngine/Data/Blobber.h
+++ b/Engine/Plugins/Experimental/TextureGraph/Source/TextureGraphEngine/Data/Blobber.h
@@ -86,10 +86,10 @@ public:
 	typedef BinCache<HashType, BlobCacheEntryPtr> BlobLRUCache;
 
 private:
-	typedef std::unordered_map<HashType, BlobCacheEntry> BlobLUT;
-	typedef std::unordered_map<HashType, CHashPtr> HashMappingTable;
-	typedef std::vector<BlobCacheEntryPtr> BlobCacheEntryPtrVec;
-	typedef std::vector<CHashPtr> CHashPtrVec;
+	using BlobLUT					= std::unordered_map<HashType, BlobCacheEntry>;
+	using HashMappingTable			= std::unordered_map<HashType, CHashPtr>;
+	using BlobCacheEntryPtrVec		= std::vector<BlobCacheEntryPtr>;
+	using CHashPtrVec				= std::vector<CHashPtr>;
 
 	//typedef std::map<HashType, std::vector<HashType>> TempHashTable;
 	mutable FCriticalSection		BlobLookupLock;				/// Mutex for the actions (mutable because FScopeLock needs non-const object)
diff --git a/Engine/Plugins/Experimental/TextureGraph/Source/TextureGraphEngine/Data/TiledBlob.cpp b/Engine/Plugins/Experimental/TextureGraph/Source/TextureGraphEngine/Data/TiledBlob.cpp
--- a/Engine/Plugins/Experimental/TextureGraph/Source/TextureGraphEngine/Data/TiledBlob.cpp
+++ b/Engine/Plugins/Experimental/TextureGraph/Source/TextureGraphEngine/Data/TiledBlob.cpp
@@ -32,12 +32,12 @@ TiledBlob::TiledBlob(const BufferDescriptor& InDesc, const BlobPtrTiles& Tiles)
 
 	bIsFinalised = true;
 
-	for (size_t TileX = 0; TileX < Tiles.Rows() && bIsFinalised; TileX++)
+	for (size_t RowId = 0; RowId < Tiles.Rows() && bIsFinalised; RowId++)
 	{
-		for (size_t TileY = 0; TileY < Tiles.Cols() && bIsFinalised; TileY++)
+		for (size_t ColId = 0; ColId < Tiles.Cols() && bIsFinalised; ColId++)
 		{
-			check(Tiles[TileX][TileY]);
-			BlobPtr Tile = Tiles[TileX][TileY];
+			check(Tiles[RowId][ColId]);
+			BlobPtr Tile = Tiles[RowId][ColId];
 
 			/// Cannot have tiled BlobObj as part of another tiled BlobObj
 			check(!Tile->IsTiled());
@@ -59,12 +59,12 @@ std::shared_ptr<TiledBlob> TiledBlob::InitFromTiles(const BufferDescriptor& InDe
 	
 	TiledBlobPtr TBlob = std::make_shared<TiledBlob>(Desc, Tiles);
 
-	for (size_t TileX = 0; TileX < Tiles.Rows(); TileX++)
+	for (size_t RowId = 0; RowId < Tiles.Rows(); RowId++)
 	{
-		for (size_t TileY = 0; TileY < Tiles.Cols(); TileY++)
+		for (size_t ColId = 0; ColId < Tiles.Cols(); ColId++)
 		{
-			check(Tiles[TileX][TileY]);
-			BlobPtr Tile = Tiles[TileX][TileY];
+			check(Tiles[RowId][ColId]);
+			BlobPtr Tile = Tiles[RowId][ColId];
 
 			/// Cannot have tiled BlobObj as part of another tiled BlobObj
 			check(!Tile->IsTiled());
@@ -78,7 +78,7 @@ std::shared_ptr<TiledBlob> TiledBlob::InitFromTiles(const BufferDescriptor& InDe
 }
 
 TiledBlob::TiledBlob(BlobRef BlobObj)
-	: TiledBlob(BlobObj->GetDescriptor(), BlobPtrTiles({{BlobObj}}))
+	: TiledBlob(BlobObj->GetDescriptor(), BlobPtrTiles({BlobObj}, 1, 1))
 {
 	/// Cannot have tiled BlobObj inside another
 	check(BlobObj && !BlobObj->IsTiled());
@@ -92,13 +92,13 @@ TiledBlob::TiledBlob(DeviceBufferRef Buffer)
 	/// no hash to calculate so far
 }
 
-TiledBlob::TiledBlob(const BufferDescriptor& InDesc, size_t NumTilesX, size_t NumTilesY, CHashPtr InHashValue)
+TiledBlob::TiledBlob(const BufferDescriptor& InDesc, size_t NumRows, size_t NumCols, CHashPtr InHashValue)
 	: Blob(InDesc, InHashValue)
-	, Tiles(NumTilesX, NumTilesY)
+	, Tiles(NumRows, NumCols)
 	, HashValue(InHashValue)
 	, Desc(InDesc)
 {
-	check(NumTilesX > 0 && NumTilesY > 0);
+	check(NumRows > 0 && NumCols > 0);
 	
 	if (Desc.Format != BufferFormat::LateBound)
 	{
@@ -154,11 +154,11 @@ bool TiledBlob::HasBlobAsTile(Blob* BlobObj) const
 	check(IsInGameThread());
 	check(BlobObj);
 
-	for (size_t TileX = 0; TileX < Tiles.Rows(); TileX++)
+	for (size_t RowId = 0; RowId < Tiles.Rows(); RowId++)
 	{
-		for (size_t TileY = 0; TileY < Tiles.Cols(); TileY++)
+		for (size_t ColId = 0; ColId < Tiles.Cols(); ColId++)
 		{
-			if (Tiles[TileX][TileY] == BlobObj)
+			if (Tiles[RowId][ColId] == BlobObj)
 				return true;
 		}
 	}
@@ -176,8 +176,8 @@ void TiledBlob::SetTiles(const BlobPtrTiles& InTiles)
 void TiledBlob::SetTile(int32 X, int32 Y, BlobRef Tile)
 {
 #if DEBUG_BLOB_REF_KEEPING == 1
-	if (Tiles[TileX][TileY])
-		Tiles[TileX][TileY]->RemoveOwner(this);
+	if (Tiles[RowId][ColId])
+		Tiles[RowId][ColId]->RemoveOwner(this);
 #endif
 
 	auto Self = shared_from_this();
@@ -365,18 +365,18 @@ AsyncBufferResultPtr TiledBlob::CombineTiles(bool bTouch, bool bIsArray, uint64
 		return cti::make_ready_continuable(std::make_shared<BufferResult>());
 	}
 
-	T_Tiles<DeviceBufferRef> TileBuffers(Tiles.Rows(), Tiles.Cols());
+	T_Tiles<DeviceBufferRef> TileBuffers(Rows(), Cols());
 
-	for (int32 TileX = 0; TileX < Tiles.Rows(); TileX++)
+	for (int32 RowId = 0; RowId < Rows(); RowId++)
 	{
-		for (int32 TileY = 0; TileY < Tiles.Cols(); TileY++)
+		for (int32 ColId = 0; ColId < Cols(); ColId++)
 		{
-			check(Tiles[TileX][TileY]);
+			check(Tiles[RowId][ColId]);
 
-			const BlobPtr Tile = Tiles[TileX][TileY];
+			const BlobPtr Tile = GetTile(RowId, ColId);
 			check(Tile->GetBufferRef());
 
-			TileBuffers[TileX][TileY] = Tile->GetBufferRef();
+			TileBuffers[RowId][ColId] = Tile->GetBufferRef();
 
 			/// Touch the Tiles here
 			if (bTouch)
@@ -391,7 +391,13 @@ AsyncBufferResultPtr TiledBlob::CombineTiles(bool bTouch, bool bIsArray, uint64
 	/// We've already calculated the hash from all the child blobs
 	UE_LOG(LogData, VeryVerbose, TEXT("*** Combine TiledBlob: %s [Hash: %llu] [T0: %s] ***"), *Desc.Name, HashValue->Value(), *Tile0->Name());
 
-	const CombineSplitArgs CombineArgs { Buffer, TileBuffers, bIsArray };
+	const CombineSplitArgs CombineArgs 
+	{ 
+		.Buffer = Buffer, 
+		.Tiles = TileBuffers, 
+		.IsArray = bIsArray,
+		.BufferHash = HashValue
+	};
 
 	return Dev->CombineFromTiles(CombineArgs).then([this](DeviceBufferRef CombinedBuffer)
 	{
@@ -404,11 +410,11 @@ AsyncBufferResultPtr TiledBlob::CombineTiles(bool bTouch, bool bIsArray, uint64
 
 void TiledBlob::TouchTiles(uint64 BatchId)
 {
-	for (size_t TileX = 0; TileX < Tiles.Rows(); TileX++)
+	for (size_t RowId = 0; RowId < Tiles.Rows(); RowId++)
 	{
-		for (size_t TileY = 0; TileY < Tiles.Cols(); TileY++)
+		for (size_t ColId = 0; ColId < Tiles.Cols(); ColId++)
 		{
-			const BlobPtr& Tile = Tiles[TileX][TileY];
+			const BlobPtr& Tile = Tiles[RowId][ColId];
 			Tile->Touch(BatchId);
 		}
 	}
@@ -433,14 +439,14 @@ AscynCHashPtr TiledBlob::CalcHash()
 	std::vector<std::decay_t<AscynCHashPtr>, std::allocator<std::decay_t<AscynCHashPtr>>> Promises;
 	std::unordered_map<HashType, bool> handled;
 
-	size_t NumTilesX = Tiles.Rows();
-	size_t NumTilesY = Tiles.Cols();
+	size_t NumRows = Tiles.Rows();
+	size_t NumCols = Tiles.Cols();
 
-	for (size_t TileX = 0; TileX < NumTilesX; TileX++)
+	for (size_t RowId = 0; RowId < NumRows; RowId++)
 	{
-		for (size_t TileY = 0; TileY < NumTilesY; TileY++)
+		for (size_t ColId = 0; ColId < NumCols; ColId++)
 		{
-			BlobPtr Tile = Tiles[TileX][TileY];
+			BlobPtr Tile = Tiles[RowId][ColId];
 			CHashPtr TileHash = Tile->Hash();
 
 			if (Tile && 
@@ -502,8 +508,8 @@ void TiledBlob::CalcHashNow() const
 
 	HashValue = nullptr;
 
-	const size_t NumTilesX = Tiles.Rows();
-	const size_t NumTilesY = Tiles.Cols();
+	const size_t NumRows = Tiles.Rows();
+	const size_t NumCols = Tiles.Cols();
 	bool CalculateDim = false;
 
 	if (!Desc.Width || !Desc.Height)
@@ -514,16 +520,16 @@ void TiledBlob::CalcHashNow() const
 		CalculateDim = true;
 	}
 
-	CHashPtrVec TileHashes(NumTilesX * NumTilesY);
+	T_Tiles<CHashPtr> TileHashes(NumRows, NumCols);
 
-	for (size_t TileX = 0; TileX < NumTilesX; TileX++)
+	for (size_t RowId = 0; RowId < NumRows; RowId++)
 	{
-		if (CalculateDim && TileX == 0)
-			Desc.Width += Tiles[TileX][0]->GetWidth();
+		if (CalculateDim && RowId == 0)
+			Desc.Width += Tiles[RowId][0]->GetWidth();
 
-		for (size_t TileY = 0; TileY < NumTilesY; TileY++)
+		for (size_t ColId = 0; ColId < NumCols; ColId++)
 		{
-			BlobPtr Tile = Tiles[TileX][TileY];
+			BlobPtr Tile = Tiles[RowId][ColId];
 
 			if (!Tile)
 				return;
@@ -533,14 +539,14 @@ void TiledBlob::CalcHashNow() const
 			if (!TileHash)
 				return;
 
-			TileHashes[TileY * NumTilesX + TileX] = TileHash;
+			TileHashes[RowId][ColId] = TileHash;
 
-			if (CalculateDim && TileY == 0)
-				Desc.Width += Tiles[TileY][0]->GetWidth();
+			if (CalculateDim && ColId == 0)
+				Desc.Width += Tiles[ColId][0]->GetWidth();
 		}
 	}
 
-	const_cast<TiledBlob*>(this)->SetHash(CHash::ConstructFromSources(TileHashes));
+	const_cast<TiledBlob*>(this)->SetHash(CHash::ConstructFromSources(TileHashes.Tiles()));
 }
 
 CHashPtr TiledBlob::Hash() const
@@ -558,16 +564,16 @@ AsyncDeviceBufferRef TiledBlob::TransferTo(Device* Target)
 	/// Go through my Tile blobs and make sure they are compatible with the Target
 	std::vector<AsyncDeviceBufferRef> Promises;
 
-	for (size_t TileX = 0; TileX < Tiles.Rows(); TileX++)
+	for (size_t RowId = 0; RowId < Tiles.Rows(); RowId++)
 	{
-		for (size_t TileY = 0; TileY < Tiles.Cols(); TileY++)
+		for (size_t ColId = 0; ColId < Tiles.Cols(); ColId++)
 		{
-			BlobPtr Tile = Tiles[TileX][TileY];
+			BlobPtr Tile = Tiles[RowId][ColId];
 			check(Tile);
 
 			if (!Tile->GetBufferRef()->IsCompatible(Target))
 			{
-				UE_LOG(LogData, Log, TEXT("TransferTo %s TiledBlob: %s Tile [%d, %d]"), *TargetName, *Desc.Name, TileX, TileY);
+				UE_LOG(LogData, Log, TEXT("TransferTo %s TiledBlob: %s Tile [%d, %d]"), *TargetName, *Desc.Name, RowId, ColId);
 				Promises.emplace_back(Tile->TransferTo(Target));
 			}
 		}
@@ -629,15 +635,15 @@ AsyncBufferResultPtr TiledBlob::Bind(const BlobTransform* Transform, const Resou
 
 			//T_Tiles<DeviceBufferRef> Tiles(Tiles.Rows(), Tiles.Cols());
 
-			//for (int32 TileX = 0; TileX < Tiles.Rows(); TileX++)
+			//for (int32 RowId = 0; RowId < Tiles.Rows(); RowId++)
 			//{
-			//	for (int32 TileY = 0; TileY < Tiles.Cols(); TileY++)
+			//	for (int32 ColId = 0; ColId < Tiles.Cols(); ColId++)
 			//	{
-			//		check(Tiles[TileX][TileY]->Buffer());
-			//		Tiles[TileX][TileY] = Tiles[TileX][TileY]->Buffer();
+			//		check(Tiles[RowId][ColId]->Buffer());
+			//		Tiles[RowId][ColId] = Tiles[RowId][ColId]->Buffer();
 
 			//		/// Touch the Tiles here
-			//		Tiles[TileX][TileY]->Touch();
+			//		Tiles[RowId][ColId]->Touch();
 			//	}
 			//}
 
@@ -666,13 +672,13 @@ AsyncBufferResultPtr TiledBlob::Bind(const BlobTransform* Transform, const Resou
 			{
 				RawBufferPtrTiles Tiles(Tiles.Rows(), Tiles.Cols());
 
-				for (size_t TileX = 0; TileX < Tiles.Rows(); TileX++)
+				for (size_t RowId = 0; RowId < Tiles.Rows(); RowId++)
 				{
-					for (size_t TileY = 0; TileY < Tiles.Cols(); TileY++)
+					for (size_t ColId = 0; ColId < Tiles.Cols(); ColId++)
 					{
-						RawBufferPtr TileRaw = Tiles[TileX][TileY]->Raw();
+						RawBufferPtr TileRaw = Tiles[RowId][ColId]->Raw();
 						check(TileRaw);
-						Tiles[TileX][TileY] = TileRaw;
+						Tiles[RowId][ColId] = TileRaw;
 					}
 				}
 
@@ -798,20 +804,20 @@ AsyncBufferResultPtr TiledBlob::Unbind(const BlobTransform* Transform, const Res
 AsyncBufferResultPtr TiledBlob::TileBuffer(DeviceBufferRef Buffer, BlobPtrTiles& Tiles)
 {
 	SCOPE_CYCLE_COUNTER(STAT_TiledBlob_TileBuffer)
-	const size_t NumTilesX = Tiles.Rows();
-	const size_t NumTilesY = Tiles.Cols();
-	T_Tiles<DeviceBufferRef>* TilesBuffers = new T_Tiles<DeviceBufferRef>(NumTilesX, NumTilesY);
+	const size_t NumRows = Tiles.Rows();
+	const size_t NumCols = Tiles.Cols();
+	T_Tiles<DeviceBufferRef>* TilesBuffers = new T_Tiles<DeviceBufferRef>(NumRows, NumCols);
 		
 	int totalTiles = Tiles.Rows() * Tiles.Cols();
 
 	/// Incase the Buffer is present in Tiles, it would be copied over. This is to avoid allocation of RenderTarget
 	/// This allocation is being done on prepare resources due to optimization (e.g lazy evaluation)
-	for (size_t TileX = 0; TileX < NumTilesX; TileX++)
+	for (size_t RowId = 0; RowId < NumRows; RowId++)
 	{
-		for (size_t TileY = 0; TileY < NumTilesY; TileY++)
+		for (size_t ColId = 0; ColId < NumCols; ColId++)
 		{
-			check(Tiles[TileX][TileY]);
-			TilesBuffers->Tiles()[TileX][TileY] = Tiles[TileX][TileY]->GetBufferRef();
+			check(Tiles[RowId][ColId]);
+			(*TilesBuffers)(RowId, ColId) = Tiles[RowId][ColId]->GetBufferRef();
 		}
 	}
 
@@ -821,13 +827,13 @@ AsyncBufferResultPtr TiledBlob::TileBuffer(DeviceBufferRef Buffer, BlobPtrTiles&
 
 	return Buffer->GetOwnerDevice()->SplitToTiles(splitArgs).then([Buffer, &Tiles, TilesBuffers](DeviceBufferRef) mutable
 	{
-		for (size_t TileX = 0; TileX < Tiles.Rows(); TileX++)
+		for (size_t RowId = 0; RowId < Tiles.Rows(); RowId++)
 		{
-			for (size_t TileY = 0; TileY < Tiles.Cols(); TileY++)
+			for (size_t ColId = 0; ColId < Tiles.Cols(); ColId++)
 			{
-				DeviceBufferRef TileBuffer = (*TilesBuffers)[TileX][TileY];
+				DeviceBufferRef TileBuffer = (*TilesBuffers)[RowId][ColId];
 				BlobPtr Tile = std::make_shared<Blob>(TileBuffer);
-				Tiles[TileX][TileY] = BlobRef(Tile);
+				Tiles[RowId][ColId] = BlobRef(Tile);
 			}
 		}
 
@@ -843,13 +849,13 @@ AsyncBufferResultPtr TiledBlob::Flush(const ResourceBindInfo& BindInfo)
 {
 	if (!Tiles.Rows() || !Tiles.Cols())
 	{
-		int32 NumTilesX = BindInfo.NumTilesX;
-		int32 NumTilesY = BindInfo.NumTilesY;
+		int32 NumRows = BindInfo.NumTilesX;
+		int32 NumCols = BindInfo.NumTilesY;
 
-		check(NumTilesX > 0 && NumTilesY > 0);
+		check(NumRows > 0 && NumCols > 0);
 		check(Buffer);
 
-		Tiles.Resize(NumTilesX, NumTilesY);
+		Tiles.Resize(NumRows, NumCols);
 	}
 
 	bReady = false;
@@ -877,11 +883,11 @@ bool TiledBlob::IsValid() const
 	/// Otherwise, iterate over the Tiles and check for the validity of each of the Tiles
 	bool Valid = true;
 	
-	for (size_t TileX = 0; TileX < Rows() && Valid; TileX++)
+	for (size_t RowId = 0; RowId < Rows() && Valid; RowId++)
 	{
-		for (size_t TileY = 0; TileY < Cols() && Valid; TileY++)
+		for (size_t ColId = 0; ColId < Cols() && Valid; ColId++)
 		{
-			Valid &= Tiles[TileX][TileY] && Tiles[TileX][TileY]->IsValid();
+			Valid &= Tiles[RowId][ColId] && Tiles[RowId][ColId]->IsValid();
 		}
 	}
 
@@ -909,13 +915,13 @@ bool TiledBlob::CanCalculateHash() const
 {
 	check(IsInGameThread());
 
-	for (size_t TileX = 0; TileX < Tiles.Rows(); TileX++)
+	for (size_t RowId = 0; RowId < Tiles.Rows(); RowId++)
 	{
-		for (size_t TileY = 0; TileY < Tiles.Cols(); TileY++)
+		for (size_t ColId = 0; ColId < Tiles.Cols(); ColId++)
 		{
-			check(Tiles[TileX][TileY]);
+			check(Tiles[RowId][ColId]);
 
-			BlobPtr TileBlob = Tiles[TileX][TileY];
+			BlobPtr TileBlob = Tiles[RowId][ColId];
 			DeviceBufferRef TileBuffer = TileBlob->GetBufferRef();
 			CHashPtr TileHash = TileBuffer->Hash();
 
@@ -988,12 +994,12 @@ void TiledBlob::SetLODLevel(int32 Level, BlobPtr Blob_, BlobPtrW LODParent_, Blo
 	{
 		Util::OnGameThread([this, BlobObj, Self, Level, LODParent_, LODSource_, AddToBlobber]()
 		{
-			for (size_t TileX = 0; TileX < Rows(); TileX++)
+			for (size_t RowId = 0; RowId < Rows(); RowId++)
 			{
-				for (size_t TileY = 0; TileY < Cols(); TileY++)
+				for (size_t ColId = 0; ColId < Cols(); ColId++)
 				{
-					BlobRef mipBlobXY = BlobObj->GetTile(TileX, TileY);
-					BlobRef baseBlobXY = Self->GetTile(TileX, TileY);
+					BlobRef mipBlobXY = BlobObj->GetTile(RowId, ColId);
+					BlobRef baseBlobXY = Self->GetTile(RowId, ColId);
 
 					baseBlobXY->SetLODLevel(Level, mipBlobXY.get(), LODParent_, LODSource_, AddToBlobber);
 				}
@@ -1009,12 +1015,12 @@ void TiledBlob::FinaliseNow(bool bNoCalcHash, CHashPtr FixedHash)
 	UE_LOG(LogJob, Verbose, TEXT("Finalised tiled BlobObj: %s"), *Desc.Name);
 
 	/// Debug: check that all Tiles are valid
-	for (int32 TileX = 0; TileX < Tiles.Rows(); TileX++)
+	for (int32 RowId = 0; RowId < Tiles.Rows(); RowId++)
 	{
-		for (int32 TileY = 0; TileY < Tiles.Cols(); TileY++)
+		for (int32 ColId = 0; ColId < Tiles.Cols(); ColId++)
 		{
-			check(Tiles[TileX][TileY]);
-			BlobPtr Tile = Tiles[TileX][TileY];
+			check(Tiles[RowId][ColId]);
+			BlobPtr Tile = Tiles[RowId][ColId];
 			Tile->FinaliseNow(bNoCalcHash, nullptr);
 		}
 	}
@@ -1035,7 +1041,7 @@ TiledBlob_Promise::TiledBlob_Promise(DeviceBufferRef Buffer) : TiledBlob(Buffer)
 {
 }
 
-TiledBlob_Promise::TiledBlob_Promise(const BufferDescriptor& Desc, size_t NumTilesX, size_t NumTilesY, CHashPtr Hash) : TiledBlob(Desc, NumTilesX, NumTilesY, Hash)
+TiledBlob_Promise::TiledBlob_Promise(const BufferDescriptor& Desc, size_t NumRows, size_t NumCols, CHashPtr Hash) : TiledBlob(Desc, NumRows, NumCols, Hash)
 {
 }
 
@@ -1158,15 +1164,15 @@ AsyncBufferResultPtr TiledBlob_Promise::Finalise(bool bNoCalcHash, CHashPtr Fixe
 	return TiledBlob::TileBuffer(Buffer, Tiles)
 		.then([this, bNoCalcHash, FixedHash](BufferResultPtr Result) mutable
 		{
-			CHashPtrVec TileHashes(Tiles.Rows() * Tiles.Cols());
+			T_Tiles<CHashPtr> TileHashes(Tiles.Rows(), Tiles.Cols());
 
-			for (int32 TileX = 0; TileX < Tiles.Rows(); TileX++)
+			for (int32 RowId = 0; RowId < Tiles.Rows(); RowId++)
 			{
-				for (int32 TileY = 0; TileY < Tiles.Cols(); TileY++)
+				for (int32 ColId = 0; ColId < Tiles.Cols(); ColId++)
 				{
-					check(Tiles[TileX][TileY]);
+					check(Tiles[RowId][ColId]);
 
-					BlobPtr Tile = Tiles[TileX][TileY];
+					BlobPtr Tile = Tiles[RowId][ColId];
 					if (!bNoCalcHash)
 						Tile->CalcHash();
 					else
@@ -1174,12 +1180,12 @@ AsyncBufferResultPtr TiledBlob_Promise::Finalise(bool bNoCalcHash, CHashPtr Fixe
 						CHashPtr TileHash = Tile->Hash();
 						check(TileHash);
 
-						TileHashes[TileY * Tiles.Rows() + TileX] = TileHash;
+						TileHashes[RowId][ColId] = TileHash;
 					}
 				}
 			}
 
-			CHashPtr fullHash = CHash::ConstructFromSources(TileHashes);
+			CHashPtr FullHash = CHash::ConstructFromSources(TileHashes.Tiles());
 			
 			/// Clear the Buffer here. We don't need it anymore
 			ResetBuffer();
@@ -1188,7 +1194,7 @@ AsyncBufferResultPtr TiledBlob_Promise::Finalise(bool bNoCalcHash, CHashPtr Fixe
 			bTiledTarget = true;
 
 			/// Tiled blobs always recalculate their has to match those of the children
-			FinaliseNow(true, fullHash);
+			FinaliseNow(true, FullHash);
 
 			return Result;
 		});
@@ -1210,9 +1216,9 @@ void TiledBlob_Promise::OnFinaliseInternal(BlobReadyCallback Callback) const
 	Callbacks.push_back(std::move(Callback));
 }
 
-void TiledBlob_Promise::SetTile(int32 TileX, int32 TileY, BlobRef InTile)
+void TiledBlob_Promise::SetTile(int32 RowId, int32 ColId, BlobRef InTile)
 {
-	check(IsValidTileIndex(TileX, TileY));
+	check(IsValidTileIndex(RowId, ColId));
 	check(!InTile.expired());
 
 	BlobPtr Tile = InTile.lock();
@@ -1220,7 +1226,7 @@ void TiledBlob_Promise::SetTile(int32 TileX, int32 TileY, BlobRef InTile)
 	/// Cannot have tiled BlobObj as part of another tiled BlobObj
 	check(!Tile->IsTiled());
 
-	TiledBlob::SetTile(TileX, TileY, InTile);
+	TiledBlob::SetTile(RowId, ColId, InTile);
 
 	if (bMakeSingleBlob && Rows() == 1 && Cols() == 1)
 		Buffer = Tile->GetBufferRef();
diff --git a/Engine/Plugins/Experimental/TextureGraph/Source/TextureGraphEngine/Data/TiledBlob.h b/Engine/Plugins/Experimental/TextureGraph/Source/TextureGraphEngine/Data/TiledBlob.h
--- a/Engine/Plugins/Experimental/TextureGraph/Source/TextureGraphEngine/Data/TiledBlob.h
+++ b/Engine/Plugins/Experimental/TextureGraph/Source/TextureGraphEngine/Data/TiledBlob.h
@@ -124,18 +124,18 @@ public:
 	FORCEINLINE bool				TiledTarget() const { return bTiledTarget; }
 	FORCEINLINE bool&				TiledTarget() { return bTiledTarget; }
 	FORCEINLINE bool				IsWellDefined() { return !IsLateBound() && GetWidth() > 0 && GetHeight() > 0; }
-	FORCEINLINE bool				IsValidTileIndex(int32 TileX, int32 TileY) const
+	FORCEINLINE bool				IsValidTileIndex(int32 RowId, int32 ColId) const
 	{
-		return TileX < static_cast<int32>(Tiles.Rows()) && TileY < static_cast<int32>(Tiles.Cols());
+		return RowId < static_cast<int32>(Tiles.Rows()) && ColId < static_cast<int32>(Tiles.Cols());
 	}
 
-	FORCEINLINE BlobRef				GetTile(int32 TileX, int32 TileY) const 
+	FORCEINLINE BlobRef				GetTile(int32 RowId, int32 ColId) const 
 	{
 		if (Tiles.Rows() == 1 && Tiles.Cols() == 1)
 			return BlobRef(Tiles[0][0], false);
 
-		check(IsValidTileIndex(TileX, TileY));
-		return BlobRef(Tiles[TileX][TileY], false);
+		check(IsValidTileIndex(RowId, ColId));
+		return BlobRef(Tiles[RowId][ColId], false);
 	}
 	FORCEINLINE BlobPtr				GetSingleBlob() const { return SingleBlob; }
 
@@ -180,7 +180,7 @@ public:
 	virtual AsyncBufferResultPtr	Unbind(const BlobTransform* Transform, const ResourceBindInfo& BindInfo) override;
 	virtual AsyncBufferResultPtr	Flush(const ResourceBindInfo& BindInfo) override;
 	virtual bool					IsFinalised() const override { return bIsFinalised; }
-	virtual void					SetTile(int32 TileX, int32 TileY, BlobRef Tile) override;
+	virtual void					SetTile(int32 RowId, int32 ColId, BlobRef Tile) override;
 	virtual AsyncBufferResultPtr	MakeSingleBlob() override;
 	virtual void					CopyResolveLateBound(BlobPtr RHS) override;
 	virtual bool					IsPromise() const override { return true; }
diff --git a/Engine/Plugins/Experimental/TextureGraph/Source/TextureGraphEngine/Device/Device.cpp b/Engine/Plugins/Experimental/TextureGraph/Source/TextureGraphEngine/Device/Device.cpp
--- a/Engine/Plugins/Experimental/TextureGraph/Source/TextureGraphEngine/Device/Device.cpp
+++ b/Engine/Plugins/Experimental/TextureGraph/Source/TextureGraphEngine/Device/Device.cpp
@@ -687,30 +687,30 @@ AsyncDeviceBufferRef Device::CombineFromTiles(const CombineSplitArgs& CombineArg
 
 	check(Tiles.Rows() > 1 && Tiles.Cols() > 1);
 
-	RawBufferPtrTiles rawTiles(Tiles.Rows(), Tiles.Cols());
-	CHashPtrVec Hashes(Tiles.Rows() * Tiles.Cols());
+	RawBufferPtrTiles RawTiles(Tiles.Rows(), Tiles.Cols());
+	T_Tiles<CHashPtr> Hashes(Tiles.Rows(), Tiles.Cols());
 
-	for (size_t TileX = 0; TileX < Tiles.Rows(); TileX++)
+	for (size_t RowId = 0; RowId < Tiles.Rows(); RowId++)
 	{
-		for (size_t TileY = 0; TileY < Tiles.Cols(); TileY++)
+		for (size_t ColId = 0; ColId < Tiles.Cols(); ColId++)
 		{
-			RawBufferPtr TileRawBuffer = Tiles[TileX][TileY]->Raw_Now();
+			RawBufferPtr TileRawBuffer = Tiles[RowId][ColId]->Raw_Now();
 			check(TileRawBuffer);
-			rawTiles[TileX][TileY] = TileRawBuffer;
+			RawTiles[RowId][ColId] = TileRawBuffer;
 
-			CHashPtr TileHash = Tiles[TileX][TileY]->Hash();
+			CHashPtr TileHash = Tiles[RowId][ColId]->Hash();
 
 			if (TileHash == 0)
 				TileHash = TileRawBuffer->Hash();
 
-			Hashes[TileY * Tiles.Rows() + TileX] = TileHash;
+			Hashes[RowId][ColId] = TileHash;
 		}
 	}
 
-	CHashPtr Hash = CHash::ConstructFromSources(Hashes);
+	CHashPtr Hash = CHash::ConstructFromSources(Hashes.Tiles());
 	
 	/// We've already calculated the Hash from all the child blobs
-	RawBufferPtr Raw = TextureHelper::CombineRaw_Tiles(rawTiles, Hash, Buffer->IsTransient());
+	RawBufferPtr Raw = TextureHelper::CombineRaw_Tiles(RawTiles, Hash, Buffer->IsTransient());
 
 	Buffer->UpdateRaw(Raw);
 
@@ -726,8 +726,8 @@ AsyncDeviceBufferRef Device::SplitToTiles_Generic(const CombineSplitArgs& SplitA
 	RawBufferPtr Raw = Buffer->Raw_Now();
 	check(Raw);
 
-	int32 TileWidth = Raw->Width() / Tiles.Rows();
-	int32 TileHeight = Raw->Height() / Tiles.Cols();
+	int32 TileWidth = Raw->Width() / Tiles.Cols();
+	int32 TileHeight = Raw->Height() / Tiles.Rows();
 
 	check(TileWidth > 0);
 	check(TileHeight > 0);
@@ -738,12 +738,12 @@ AsyncDeviceBufferRef Device::SplitToTiles_Generic(const CombineSplitArgs& SplitA
 	check(RawTileBuffers.Cols() == Tiles.Cols());
 
 	std::vector<AsyncBufferResultPtr> Promises;
-	for (int32 TileX = 0; TileX < Tiles.Rows(); TileX++)
+	for (int32 RowId = 0; RowId < Tiles.Rows(); RowId++)
 	{
-		for (int32 TileY = 0; TileY < Tiles.Cols(); TileY++)
+		for (int32 ColId = 0; ColId < Tiles.Cols(); ColId++)
 		{
-			DeviceBufferRef Tile = Tiles[TileX][TileY];
-			RawBufferPtr RawTile = RawTileBuffers[TileX][TileY];
+			DeviceBufferRef Tile = Tiles[RowId][ColId];
+			RawBufferPtr RawTile = RawTileBuffers[RowId][ColId];
 			check(RawTile);
 
 			AsyncBufferResultPtr TilePromise = Tile->UpdateRaw(RawTile);
diff --git a/Engine/Plugins/Experimental/TextureGraph/Source/TextureGraphEngine/Device/Device.h b/Engine/Plugins/Experimental/TextureGraph/Source/TextureGraphEngine/Device/Device.h
--- a/Engine/Plugins/Experimental/TextureGraph/Source/TextureGraphEngine/Device/Device.h
+++ b/Engine/Plugins/Experimental/TextureGraph/Source/TextureGraphEngine/Device/Device.h
@@ -37,6 +37,7 @@ struct CombineSplitArgs
 	DeviceBufferRef					Buffer;						/// Main buffer that should be used for combining to or split from
 	const T_Tiles<DeviceBufferRef>& Tiles;						/// Tiles that needs to be merged or split
 	bool							IsArray = false;			/// if operating buffer should be treated as Array 
+	CHashPtr						BufferHash;					/// Precalculated buffer hash that might be available. Otherwise nullptr
 };
 
 //////////////////////////////////////////////////////////////////////////
diff --git a/Engine/Plugins/Experimental/TextureGraph/Source/TextureGraphEngine/Device/FX/Device_FX.cpp b/Engine/Plugins/Experimental/TextureGraph/Source/TextureGraphEngine/Device/FX/Device_FX.cpp
--- a/Engine/Plugins/Experimental/TextureGraph/Source/TextureGraphEngine/Device/FX/Device_FX.cpp
+++ b/Engine/Plugins/Experimental/TextureGraph/Source/TextureGraphEngine/Device/FX/Device_FX.cpp
@@ -202,36 +202,15 @@ AsyncDeviceBufferRef Device_FX::CombineFromTiles(const CombineSplitArgs& Combine
 	const T_Tiles<DeviceBufferRef>& Tiles = CombineArgs.Tiles;
 	auto Buffer = CombineArgs.Buffer;
 	auto IsArray = CombineArgs.IsArray;
+	CHashPtr BufferHash = CombineArgs.BufferHash;
 
 	check(Tiles.Rows() > 1 && Tiles.Cols() > 1);
-
-	/// check whether the buffers are compatible
 	bool bIsCompatible = true;
 	BufferDescriptor Desc = Tiles[0][0]->Descriptor();
-	CHashPtrVec TileHashes(Tiles.Rows() * Tiles.Cols());
-
-	for (int32 TileX = 0; TileX < Tiles.Rows() && bIsCompatible; TileX++)
-	{
-		for (int32 TileY = 0; TileY < Tiles.Cols() && bIsCompatible; TileY++)
-		{
-			DeviceBufferRef Tile = Tiles[TileX][TileY];
-			CHashPtr TileHash = Tile->Hash(false);
-			check(TileHash && TileHash->IsValid());
-			bIsCompatible &= Tile->IsCompatible(this);
-			TileHashes[TileY * Tiles.Cols() + TileX] = TileHash;
-		}
-	}
-
-	CHashPtr BufferHash = CHash::ConstructFromSources(TileHashes);
-	if (Buffer)
-		Buffer->SetHash(BufferHash);
 
 	Desc.Width = Desc.Width * Tiles.Cols();
 	Desc.Height = Desc.Height * Tiles.Rows();
 
-	/// Tiles must have been made compatible beforehand
-	check(bIsCompatible);
-
 	/// If any of the buffers is not compatible, then we fallback to the default implementation
 	if (!bIsCompatible)
 		return Device::CombineFromTiles(CombineArgs);
@@ -275,11 +254,11 @@ AsyncDeviceBufferRef Device_FX::FillTextureArray_Deferred(DeviceBufferRef Buffer
 			FRHICommandListImmediate& DevRHI = RHI();
 
 			int32 Index = 0;
-			for (int32 TileX = 0; TileX < Tiles.Rows(); TileX++)
+			for (int32 RowId = 0; RowId < Tiles.Rows(); RowId++)
 			{
-				for (int32 TileY = 0; TileY < Tiles.Cols(); TileY++, Index++)
+				for (int32 ColId = 0; ColId < Tiles.Cols(); ColId++, Index++)
 				{
-					DeviceBuffer_FX* TileBuffer = static_cast<DeviceBuffer_FX*>(Tiles[TileX][TileY].get());
+					DeviceBuffer_FX* TileBuffer = static_cast<DeviceBuffer_FX*>(Tiles[RowId][ColId].get());
 					TexPtr TileTex = TileBuffer->GetTexture();
 
 					check(TileTex);
@@ -348,14 +327,14 @@ AsyncDeviceBufferRef Device_FX::DrawTilesToBuffer_Deferred(DeviceBufferRef Buffe
 
 				SetGraphicsPipelineState(DevRHI, GraphicsPSOInit, 0);
 
-				for (int32 TileX = 0; TileX < Tiles.Rows(); TileX++)
+				for (int32 RowId = 0; RowId < Tiles.Rows(); RowId++)
 				{
-					for (int32 TileY = 0; TileY < Tiles.Cols(); TileY++)
+					for (int32 ColId = 0; ColId < Tiles.Cols(); ColId++)
 					{
-						DeviceBuffer_FX* TileBuffer = static_cast<DeviceBuffer_FX*>(Tiles[TileX][TileY].get());
+						DeviceBuffer_FX* TileBuffer = static_cast<DeviceBuffer_FX*>(Tiles(RowId, ColId).get());
 						TexPtr TileTex = TileBuffer->GetTexture();
 
-						UE_LOG(LogDevice, Verbose, TEXT("Combine tiles: %s [%d, %d] => %llu [%s]"), *FXBuffer->GetName(), TileX, TileY, TileBuffer->Hash(false)->Value(), *TileBuffer->GetName());
+						UE_LOG(LogDevice, Verbose, TEXT("Combine tiles: %s [%d, %d] => %llu [%s]"), *FXBuffer->GetName(), RowId, ColId, TileBuffer->Hash(false)->Value(), *TileBuffer->GetName());
 
 						check(TileTex);
 
@@ -368,8 +347,8 @@ AsyncDeviceBufferRef Device_FX::DrawTilesToBuffer_Deferred(DeviceBufferRef Buffe
 				//// This was here but shouldn't be needed if default state is SRV 
 				///		DevRHI.Transition(FRHITransitionInfo(TileResource, ERHIAccess::Unknown, ERHIAccess::SRVGraphics));
 
-						int32 X1 = TileX * SrcTextureWidth;
-						int32 Y1 = TileY * SrcTextureHeight;
+						int32 X1 = ColId * SrcTextureWidth;
+						int32 Y1 = RowId * SrcTextureHeight;
 
 						DevRHI.SetViewport(X1, Y1, 0, X1 + ViewportWidth, Y1 + ViewportHeight, 0);
 
@@ -412,27 +391,27 @@ AsyncDeviceBufferRef Device_FX::SplitToTiles_Internal(const CombineSplitArgs& Sp
 
 			/// This texture must have been created beforehand AND it must be a render target
 			TexPtr SrcTex = FXBuffer->GetTexture();
-			FTextureRHIRef sourceResource = ((UTextureRenderTarget2D*)SrcTex->GetTexture())->GetRenderTargetResource()->GetRenderTargetTexture();
+			FTextureRHIRef SourceRes = ((UTextureRenderTarget2D*)SrcTex->GetTexture())->GetRenderTargetResource()->GetRenderTargetTexture();
 
-			for (int32 TileX = 0; TileX < Tiles.Rows(); TileX++)
+			for (int32 RowId = 0; RowId < Tiles.Rows(); RowId++)
 			{
-				for (int32 TileY = 0; TileY < Tiles.Cols(); TileY++)
+				for (int32 ColId = 0; ColId < Tiles.Cols(); ColId++)
 				{
 					/// Using RHI to split Tiles to Buffer
-					DeviceBuffer_FX* TileBuffer = static_cast<DeviceBuffer_FX*>(Tiles[TileX][TileY].get());
+					DeviceBuffer_FX* TileBuffer = static_cast<DeviceBuffer_FX*>(Tiles[RowId][ColId].get());
 					TexPtr TileTex = TileBuffer->GetTexture();
 					check(TileTex);
 					check(TileTex->GetRenderTarget());
 
-					FTextureRHIRef texture2Dres = ((FTextureRenderTarget2DResource*)TileTex->GetRenderTarget()->GetResource())->GetTextureRHI();
+					FTextureRHIRef Texture2DRes = ((FTextureRenderTarget2DResource*)TileTex->GetRenderTarget()->GetResource())->GetTextureRHI();
 
 					FRHICopyTextureInfo CopyInfo;
-					CopyInfo.SourcePosition = FIntVector(TileX * TileTex->GetWidth(), TileY * TileTex->GetHeight(), 0);
+					CopyInfo.SourcePosition = FIntVector(RowId * TileTex->GetHeight(), ColId * TileTex->GetWidth(), 0);
 					CopyInfo.DestPosition = FIntVector(0, 0, 0);
 					CopyInfo.Size = FIntVector(TileTex->GetWidth(), TileTex->GetHeight(), 0);
 
 					FRHICommandListImmediate& DevRHI = GRHICommandList.GetImmediateCommandList();
-					DevRHI.CopyTexture(sourceResource, texture2Dres, CopyInfo);
+					DevRHI.CopyTexture(SourceRes, Texture2DRes, CopyInfo);
 				}
 			}
 
@@ -461,14 +440,14 @@ AsyncDeviceBufferRef Device_FX::SplitToTiles(const CombineSplitArgs& SplitArgs)
 
 	/// first of all we need to ensure that all the Tiles are ready
 
-	for (int32 TileX = 0; TileX < Tiles.Rows(); TileX++)
+	for (int32 RowId = 0; RowId < Tiles.Rows(); RowId++)
 	{
-		for (int32 TileY = 0; TileY < Tiles.Cols(); TileY++)
+		for (int32 ColId = 0; ColId < Tiles.Cols(); ColId++)
 		{
-			if (!Tiles[TileX][TileY])
+			if (!Tiles[RowId][ColId])
 			{
-				DeviceBufferRef Tile = Tiles[TileX][TileY];
-				TileDesc.Name = FString::Format(TEXT("{0}-{1}.{2}"), { Buffer->Descriptor().Name, TileX, TileY });
+				DeviceBufferRef Tile = Tiles[RowId][ColId];
+				TileDesc.Name = FString::Format(TEXT("{0}-{1}.{2}"), { Buffer->Descriptor().Name, RowId, ColId });
 				DeviceBuffer_FX* FXBuffer = new DeviceBuffer_FX(this, TileDesc, nullptr);
 				DeviceBufferRef FXBufferRef = AddNewRef_Internal(FXBuffer);
 
@@ -888,21 +867,24 @@ Device_FX* Device_FX::Get()
 	if (TextureGraphEngine::IsDestroying())
 		throw std::runtime_error("Engine is shutting down!");
 
-	Device* dev = TextureGraphEngine::GetDeviceManager()->GetDevice(DeviceType::FX);
-	check(dev);
-	return static_cast<Device_FX*>(dev);
+	Device* Dev = TextureGraphEngine::GetDeviceManager()->GetDevice(DeviceType::FX);
+	check(Dev);
+	return static_cast<Device_FX*>(Dev);
 }
 
-void Device_FX::InitTiles_Texture(BlobUPtr* Tiles, size_t NumRows, size_t NumCols, const BufferDescriptor& InTileDesc, bool bInitRaw)
+void Device_FX::InitTiles_Texture(T_Tiles<BlobPtr>* Tiles, const BufferDescriptor& InTileDesc, bool bInitRaw)
 {
+	size_t NumRows = Tiles->Rows();
+	size_t NumCols = Tiles->Cols();
+
 	BufferDescriptor TileDesc = InTileDesc;
 	TileDesc.Format = TextureHelper::FindOptimalSupportedFormat(InTileDesc.Format);
 
-	for (int32 TileX = 0; TileX < NumRows; TileX++)
+	for (int32 RowId = 0; RowId < NumRows; RowId++)
 	{
-		for (int32 TileY = 0; TileY < NumCols; TileY++)
+		for (int32 ColId = 0; ColId < NumCols; ColId++)
 		{
-			TileDesc.Name = TextureHelper::CreateTileName(InTileDesc.Name, TileX, TileY);
+			TileDesc.Name = TextureHelper::CreateTileName(InTileDesc.Name, RowId, ColId);
 
 			TexDescriptor TextureTileDesc(TileDesc);
 			TexPtr TileTex = std::make_shared<Tex>(TextureTileDesc);
@@ -910,20 +892,20 @@ void Device_FX::InitTiles_Texture(BlobUPtr* Tiles, size_t NumRows, size_t NumCol
 			/// prepare over here
 			//DeviceBuffer_FX* FXBuffer = new DeviceBuffer_FX(Device_FX::Get(), TileTex, TileDesc.ToBufferDescriptor(), (CHashPtr)nullptr);
 			//DeviceBuffer_FXPtr FXBuffer = std::make_shared<DeviceBuffer_FX>;
-			BlobUPtr TileBlob = std::make_unique<Blob>(Device_FX::Get()->CreateFromTex(TileTex, bInitRaw));
-			Tiles[TileX * NumCols + TileY] = std::move(TileBlob);
+			BlobPtr TileBlob = std::make_shared<Blob>(Device_FX::Get()->CreateFromTex(TileTex, bInitRaw));
+			(*Tiles)[RowId][ColId] = TileBlob;
 		}
 	}
 }
 
 void Device_FX::InitTiles_RenderTargets(BlobUPtr* Tiles, size_t NumRows, size_t NumCols, const BufferDescriptor& InTileDesc, bool bInitRaw)
 {
-	for (int32 TileX = 0; TileX < NumRows; TileX++)
+	for (int32 RowId = 0; RowId < NumRows; RowId++)
 	{
-		for (int32 TileY = 0; TileY < NumCols; TileY++)
+		for (int32 ColId = 0; ColId < NumCols; ColId++)
 		{
 			TexDescriptor TileDesc(InTileDesc);
-			TileDesc.Name = TextureHelper::CreateTileName(InTileDesc.Name, TileX, TileY);
+			TileDesc.Name = TextureHelper::CreateTileName(InTileDesc.Name, RowId, ColId);
 
 			TexPtr TileTex = std::make_shared<Tex>(TileDesc);
 			TileTex->InitRT();
@@ -931,7 +913,7 @@ void Device_FX::InitTiles_RenderTargets(BlobUPtr* Tiles, size_t NumRows, size_t
 			/// prepare over here
 			//DeviceBuffer_FXPtr FXBuffer = std::make_shared<DeviceBuffer_FX>(Device_FX::Get(), TileTex, TileDesc.ToBufferDescriptor(), (CHashPtr)nullptr);
 			BlobUPtr TileBlob = std::make_unique<Blob>(Device_FX::Get()->CreateFromTex(TileTex, bInitRaw));
-			Tiles[TileX * NumCols + TileY] = std::move(TileBlob);
+			Tiles[RowId * NumCols + ColId] = std::move(TileBlob);
 		}
 	}
 }
diff --git a/Engine/Plugins/Experimental/TextureGraph/Source/TextureGraphEngine/Device/FX/Device_FX.h b/Engine/Plugins/Experimental/TextureGraph/Source/TextureGraphEngine/Device/FX/Device_FX.h
--- a/Engine/Plugins/Experimental/TextureGraph/Source/TextureGraphEngine/Device/FX/Device_FX.h
+++ b/Engine/Plugins/Experimental/TextureGraph/Source/TextureGraphEngine/Device/FX/Device_FX.h
@@ -77,8 +77,8 @@ public:
 
 	virtual AsyncDeviceBufferRef	CombineFromTiles(const CombineSplitArgs& CombineArgs) override;
 
-	AsyncDeviceBufferRef			DrawTilesToBuffer_Deferred(DeviceBufferRef buffer, const T_Tiles<DeviceBufferRef>& tiles);
-	AsyncDeviceBufferRef			FillTextureArray_Deferred(DeviceBufferRef buffer, const T_Tiles<DeviceBufferRef>& tiles);
+	AsyncDeviceBufferRef			DrawTilesToBuffer_Deferred(DeviceBufferRef buffer, const T_Tiles<DeviceBufferRef>& Tiles);
+	AsyncDeviceBufferRef			FillTextureArray_Deferred(DeviceBufferRef buffer, const T_Tiles<DeviceBufferRef>& Tiles);
 
 	virtual AsyncDeviceBufferRef	SplitToTiles(const CombineSplitArgs& splitArgs) override;
 
@@ -100,7 +100,7 @@ public:
 	//////////////////////////////////////////////////////////////////////////
 	/// Static functions
 	//////////////////////////////////////////////////////////////////////////
-	static void						InitTiles_Texture(BlobUPtr* tiles, size_t NumRows, size_t NumCols, const BufferDescriptor& InTileDesc, bool bInitRaw);
+	static void						InitTiles_Texture(T_Tiles<BlobPtr>* Tiles, const BufferDescriptor& InTileDesc, bool bInitRaw);
 	static void						InitTiles_RenderTargets(BlobUPtr* Tiles, size_t NumRows, size_t NumCols, const BufferDescriptor& InTileDesc, bool bInitRaw);
 
 	static Device_FX*				Get();
diff --git a/Engine/Plugins/Experimental/TextureGraph/Source/TextureGraphEngine/Helper/DataUtil.cpp b/Engine/Plugins/Experimental/TextureGraph/Source/TextureGraphEngine/Helper/DataUtil.cpp
--- a/Engine/Plugins/Experimental/TextureGraph/Source/TextureGraphEngine/Helper/DataUtil.cpp
+++ b/Engine/Plugins/Experimental/TextureGraph/Source/TextureGraphEngine/Helper/DataUtil.cpp
@@ -138,7 +138,7 @@ CHash::CHash(HashType Value, bool bInIsFinal) : bIsFinal(bInIsFinal), Timestamp(
 		HashValue = Value;
 }
 
-CHash::CHash(CHashPtrVec Sources) : Timestamp(FDateTime::Now())
+CHash::CHash(const CHashPtrVec& Sources) : Timestamp(FDateTime::Now())
 {
 	check(Sources.size() > 0);
 	bIsFinal = true;
@@ -171,7 +171,7 @@ CHash::CHash(CHashPtrVec Sources) : Timestamp(FDateTime::Now())
 	}
 }
 
-CHashPtr CHash::ConstructFromSources(CHashPtrVec Sources)
+CHashPtr CHash::ConstructFromSources(const CHashPtrVec& Sources)
 {
 	CHashPtr HashValue = CHashPtr(new CHash(Sources));
 
diff --git a/Engine/Plugins/Experimental/TextureGraph/Source/TextureGraphEngine/Helper/DataUtil.h b/Engine/Plugins/Experimental/TextureGraph/Source/TextureGraphEngine/Helper/DataUtil.h
--- a/Engine/Plugins/Experimental/TextureGraph/Source/TextureGraphEngine/Helper/DataUtil.h
+++ b/Engine/Plugins/Experimental/TextureGraph/Source/TextureGraphEngine/Helper/DataUtil.h
@@ -173,7 +173,7 @@ private:
 	void							HandleLinkUpdated(CHashPtr LinkUpdated);
 
 									/// DO NOT MAKE THIS C'TOR PUBLIC. Use ConstructFromSources instead
-	explicit						CHash(CHashPtrVec Sources);
+	explicit						CHash(const CHashPtrVec& Sources);
 
 public:
 									CHash(HashType Value, bool bInIsFinal);
@@ -183,7 +183,7 @@ public:
 	bool							operator == (const CHash& RHS) const;
 	HashTypeVec						GetIntermediateHashes() const;
 
-	static CHashPtr					ConstructFromSources(CHashPtrVec Sources);
+	static CHashPtr					ConstructFromSources(const CHashPtrVec& Sources);
 	static CHashPtr					UpdateHash(CHashPtr NewHash, CHashPtr PrevHash);
 
 	//////////////////////////////////////////////////////////////////////////
@@ -215,22 +215,23 @@ class TEXTUREGRAPHENGINE_API T_Tiles
 {
 public:
 	typedef std::vector<Type>		TypeVec;
-	typedef std::vector<TypeVec>	Type_Tiles;
 
 protected:
-	Type_Tiles						TilesVec;			/// The tiles 
+	TypeVec							TilesVec;			/// The tiles 
 	size_t							NumRows = 0;		/// Number of rows
-	size_t							NumCols = 0;			/// Number of columns
+	size_t							NumCols = 0;		/// Number of columns
 
 public:
 	T_Tiles(size_t InNumRows, size_t InNumCols) : NumRows(InNumRows), NumCols(InNumCols)
 	{
-		TilesVec.resize(InNumRows);
-		for (size_t TileX = 0; TileX < InNumRows; TileX++)
-			TilesVec[TileX].resize(InNumCols);
+		if (NumRows > 0 && NumCols > 0)
+			TilesVec.resize(NumRows * NumCols);
 	}
 
-	T_Tiles(const Type_Tiles& tiles) : TilesVec(tiles), NumRows(TilesVec.size()), NumCols(TilesVec.size() ? TilesVec[0].size() : 0)
+	T_Tiles(const TypeVec& RHS, size_t InNumRows, size_t InNumCols) 
+		: TilesVec(RHS)
+		, NumRows(InNumRows)
+		, NumCols(InNumCols)
 	{
 	}
 
@@ -238,34 +239,35 @@ public:
 	{
 		NumRows = InNumRows;
 		NumCols = InNumCols;
-		TilesVec.resize(NumRows);
-		for (size_t TileX = 0; TileX < InNumRows; TileX++)
-			TilesVec[TileX].resize(InNumCols);
+		if (NumRows > 0 && NumCols > 0)
+			TilesVec.resize(NumRows * NumCols);
 	}
 
+	FORCEINLINE size_t				Length() const { return NumRows * NumCols; }
+
 	template<typename Functor>
 	FORCEINLINE void				ForEach(Functor Func) const
 	{
-		for (size_t TileX = 0; TileX < NumRows; ++TileX)
+		for (size_t RowId = 0; RowId < NumRows; ++RowId)
 		{
-			for (size_t TileY = 0; TileY < NumCols; ++TileY)
+			for (size_t ColId = 0; ColId < NumCols; ++ColId)
 			{
-				Func(TileX, TileY);
+				Func(RowId, ColId);
 			}
 		}
 	}
 	//////////////////////////////////////////////////////////////////////////
 	/// Inline functions
 	//////////////////////////////////////////////////////////////////////////
-	FORCEINLINE const TypeVec&		operator[](size_t TileX) const { check(TileX < NumRows); return TilesVec[TileX]; }
-	FORCEINLINE TypeVec&			operator[](size_t TileX) { check(TileX < NumRows); return TilesVec[TileX]; }
-	FORCEINLINE const Type&			operator()(size_t TileX, size_t TileY) const { check(TileX < NumRows && TileY < NumCols); const TypeVec& Row = TilesVec[TileX]; return Row[TileY]; }
-	FORCEINLINE Type&				operator()(size_t TileX, size_t TileY) { check(TileX < NumRows && TileY < NumCols); TypeVec& Row = TilesVec[TileX]; return Row[TileY]; }
+	FORCEINLINE const Type*			operator[](size_t RowId) const { check(RowId < NumRows); return (TilesVec.data() + (RowId * NumCols)); }
+	FORCEINLINE Type*				operator[](size_t RowId) { check(RowId < NumRows); return (TilesVec.data() + (RowId * NumCols)); }
+	FORCEINLINE const Type&			operator()(size_t RowId, size_t ColId) const { check(RowId < NumRows && ColId < NumCols); return TilesVec[RowId * NumCols + ColId]; }
+	FORCEINLINE Type&				operator()(size_t RowId, size_t ColId) { check(RowId < NumRows && ColId < NumCols); return TilesVec[RowId * NumCols + ColId]; }
 	FORCEINLINE size_t				Rows() const { return NumRows; }
 	FORCEINLINE size_t				Cols() const { return NumCols; }
-	FORCEINLINE const Type_Tiles&	Tiles() const { return TilesVec; }
-	FORCEINLINE Type_Tiles&			Tiles() { return TilesVec; }
+	FORCEINLINE const TypeVec&		Tiles() const { return TilesVec; }
+	FORCEINLINE TypeVec&			Tiles() { return TilesVec; }
 	FORCEINLINE void				Clear() { TilesVec.clear(); }
 };
 
-typedef T_Tiles<bool>				TileInvalidateMatrix;
+using TileInvalidateMatrix			= T_Tiles<int32>;
diff --git a/Engine/Plugins/Experimental/TextureGraph/Source/TextureGraphEngine/Job/Job.cpp b/Engine/Plugins/Experimental/TextureGraph/Source/TextureGraphEngine/Job/Job.cpp
--- a/Engine/Plugins/Experimental/TextureGraph/Source/TextureGraphEngine/Job/Job.cpp
+++ b/Engine/Plugins/Experimental/TextureGraph/Source/TextureGraphEngine/Job/Job.cpp
@@ -56,7 +56,7 @@ Job::~Job()
 	Args.Empty();
 }
 
-AsyncJobArgResultPtr Job::BindArgs(JobRunInfo InRunInfo, BlobTransformPtr TransformObj, int32 TileX, int32 TileY)
+AsyncJobArgResultPtr Job::BindArgs(JobRunInfo InRunInfo, BlobTransformPtr TransformObj, int32 RowId, int32 ColId)
 {
 	if (!Args.Num())
 		return cti::make_ready_continuable(std::make_shared<JobArgResult>());
@@ -66,7 +66,7 @@ AsyncJobArgResultPtr Job::BindArgs(JobRunInfo InRunInfo, BlobTransformPtr Transf
 
 	for (size_t ArgIndex = 0; ArgIndex < Args.Num(); ArgIndex++)
 	{
-		JobArgBindInfo ArgBindInfo = { this, InRunInfo.Batch, ArgIndex, TileX, TileY, InRunInfo.Dev, TransformObj, InRunInfo.Batch->GetCycle()->LODLevel() };
+		JobArgBindInfo ArgBindInfo = { this, InRunInfo.Batch, ArgIndex, RowId, ColId, InRunInfo.Dev, TransformObj, InRunInfo.Batch->GetCycle()->LODLevel() };
 
 		if (!Args[ArgIndex]->Unbounded())
 			Promises.emplace_back(std::forward<AsyncJobArgResultPtr>(Args[ArgIndex]->Bind(ArgBindInfo)));
@@ -80,7 +80,7 @@ AsyncJobArgResultPtr Job::BindArgs(JobRunInfo InRunInfo, BlobTransformPtr Transf
 	});
 }
 
-AsyncJobArgResultPtr Job::UnbindArgs(JobRunInfo InRunInfo, BlobTransformPtr TransformObj, int32 TileX, int32 TileY)
+AsyncJobArgResultPtr Job::UnbindArgs(JobRunInfo InRunInfo, BlobTransformPtr TransformObj, int32 RowId, int32 ColId)
 {
 	if (!Args.Num())
 		return cti::make_ready_continuable(std::make_shared<JobArgResult>());
@@ -90,7 +90,7 @@ AsyncJobArgResultPtr Job::UnbindArgs(JobRunInfo InRunInfo, BlobTransformPtr Tran
 
 	for (size_t ArgIndex = 0; ArgIndex < Args.Num(); ArgIndex++)
 	{
-		JobArgBindInfo ArgBindInfo = { this, InRunInfo.Batch, ArgIndex, TileX, TileY, InRunInfo.Dev, TransformObj, InRunInfo.Batch->GetCycle()->LODLevel() };
+		JobArgBindInfo ArgBindInfo = { this, InRunInfo.Batch, ArgIndex, RowId, ColId, InRunInfo.Dev, TransformObj, InRunInfo.Batch->GetCycle()->LODLevel() };
 
 		if (!Args[ArgIndex]->Unbounded())
 			Promises.emplace_back(std::forward<AsyncJobArgResultPtr>(Args[ArgIndex]->Unbind(ArgBindInfo)));
@@ -112,13 +112,13 @@ AsyncInt Job::Bind_Or_Unbind_All_Generic(JobRunInfo InRunInfo, Bind_Unbind_Func
 
 		std::vector<AsyncJobArgResultPtr> Promises;
 
-		for (size_t TileX = 0; TileX < TileInvalidationMatrix.Rows(); TileX++)
+		for (size_t RowId = 0; RowId < TileInvalidationMatrix.Rows(); RowId++)
 		{
-			for (size_t TileY = 0; TileY < TileInvalidationMatrix.Rows(); TileY++)
+			for (size_t ColId = 0; ColId < TileInvalidationMatrix.Rows(); ColId++)
 			{
-				if (TileInvalidationMatrix[TileX][TileY])
+				if (TileInvalidationMatrix[RowId][ColId])
 				{
-					Promises.push_back(BFunc(InRunInfo, TileResults[TileX][TileY]->Transform, TileX, TileY));
+					Promises.push_back(BFunc(InRunInfo, TileResults[RowId][ColId]->Transform, RowId, ColId));
 				}
 			}
 		}
@@ -204,7 +204,7 @@ CHashPtr Job::Hash() const
 	return HashValue;
 }
 
-CHashPtr Job::CalcTileHash(int32 TileX, int32 TileY) const
+CHashPtr Job::CalcTileHash(int32 RowId, int32 ColId) const
 {
 	CHashPtrVec ArgHashes;
 	ArgHashes.reserve(Args.Num() + 1); /// +1 for the Transform
@@ -217,8 +217,8 @@ CHashPtr Job::CalcTileHash(int32 TileX, int32 TileY) const
 		{
 			CHashPtr ArgHash;
 
-			if (TileX >= 0 && TileY >= 0)
-				ArgHash = Arg->TileHash(TileX, TileY);
+			if (RowId >= 0 && ColId >= 0)
+				ArgHash = Arg->TileHash(RowId, ColId);
 			else
 				ArgHash = Arg->Hash();
 
@@ -236,7 +236,7 @@ CHashPtr Job::CalcTileHash(int32 TileX, int32 TileY) const
 	/// For the full job hash calculation we can append something to differentiate it from 
 	/// any of the tiles. Otherwise we won't be able to differentiate between the two
 	/// when someone queries it from the blobber. 
-	if (TileX < 0 && TileY < 0)
+	if (RowId < 0 && ColId < 0)
 	{
 		ArgHashes.push_back(std::make_shared<CHash>(DataUtil::Hash_Int32(-1), true));
 		ArgHashes.push_back(std::make_shared<CHash>(DataUtil::Hash_Int32(-1), true));
@@ -268,9 +268,9 @@ CHashPtr Job::CalcTileHash(int32 TileX, int32 TileY) const
 	return CHash::ConstructFromSources(ArgHashes);
 }
 
-CHashPtr Job::TileHash(int32 TileX, int32 TileY) const
+CHashPtr Job::TileHash(int32 RowId, int32 ColId) const
 {
-	return TextureGraphEngine::GetBlobber()->AddGloballyUniqueHash(CalcTileHash(TileX, TileY));
+	return TextureGraphEngine::GetBlobber()->AddGloballyUniqueHash(CalcTileHash(RowId, ColId));
 }
 
 BufferDescriptor Job::GetCombinedDesc(BufferDescriptor& ArgsDescCombined, size_t& Count) const
@@ -375,12 +375,12 @@ bool Job::CanHandleTiles() const
 	return true;
 }
 
-void Job::SetTileResult(int32 TileX, int32 TileY, JobResultPtr TileResult)
+void Job::SetTileResult(int32 RowId, int32 ColId, JobResultPtr TileResult)
 {
 	check(TileResult->HashValue);
 
-	if (TileX >= 0 && TileY >= 0)
-		TileResults[TileX][TileY] = TileResult;
+	if (RowId >= 0 && ColId >= 0)
+		TileResults[RowId][ColId] = TileResult;
 	else
 		FinalJobResult = TileResult;
 }
@@ -574,17 +574,17 @@ AsyncJobResultPtr Job::FinaliseTiles(JobRunInfo InRunInfo)
 	{
 		if (CanHandleTiles())
 		{
-			for (size_t TileX = 0; TileX < TileResults.Rows(); TileX++)
+			for (size_t RowId = 0; RowId < TileResults.Rows(); RowId++)
 			{
-				for (size_t TileY = 0; TileY < TileResults.Rows(); TileY++)
+				for (size_t ColId = 0; ColId < TileResults.Rows(); ColId++)
 				{
-					JobResultPtr JobResult = TileResults[TileX][TileY];
+					JobResultPtr JobResult = TileResults[RowId][ColId];
 					auto BlobObj = JobResult->BlobObj.lock();
 
 					if (BlobObj)
 					{
 						UE_LOG(LogJob, VeryVerbose, TEXT("FINISHED: %llu => %s"), JobResult->HashValue->Value(), *BlobObj->Name());
-						result->SetTile(TileX, TileY, BlobObj);
+						result->SetTile(RowId, ColId, BlobObj);
 					}
 				}
 			}
@@ -630,27 +630,27 @@ bool Job::CheckCulled(JobRunInfo InRunInfo)
 	size_t NumCols = ResultPromise->Cols();
 	BlobPtrTiles Tiles(NumRows, NumCols);
 
-	for (size_t TileX = 0; TileX < ResultPromise->Rows(); TileX++)
+	for (size_t RowId = 0; RowId < ResultPromise->Rows(); RowId++)
 	{
-		for (size_t TileY = 0; TileY < ResultPromise->Cols(); TileY++)
+		for (size_t ColId = 0; ColId < ResultPromise->Cols(); ColId++)
 		{
 			/// Check if any of the arguments are late bound
 			for (const JobArgPtr& Arg : Args)
 			{
-				if (Arg->IsLateBound(TileX, TileY))
+				if (Arg->IsLateBound(RowId, ColId))
 					return false;
 			}
 
-			CHashPtr ResultTileHash = TileHash(TileX, TileY);
+			CHashPtr ResultTileHash = TileHash(RowId, ColId);
 			BlobRef TileBlob = TextureGraphEngine::GetBlobber()->FindSingle(ResultTileHash->Value());
 
 			/// Couldn't find the result. Must run the job
 			if (!TileBlob)
 				return false;
 
-			UE_LOG(LogJob, VeryVerbose, TEXT("[Job - %s] matched tile: %d, %d"), *GetName(), (int32)TileX, (int32)TileY);
+			UE_LOG(LogJob, VeryVerbose, TEXT("[Job - %s] matched tile: %d, %d"), *GetName(), (int32)RowId, (int32)ColId);
 
-			Tiles[TileX][TileY] = TileBlob;
+			Tiles[RowId][ColId] = TileBlob;
 		}
 	}
 
@@ -765,11 +765,11 @@ AsyncPrepareResult Job::PrepareTargets(JobBatch* Batch)
 	bool bIsDiscard = IsDiscard(); /// In the case we are replaying the job for debug purpose, do not cache the result
 	std::unordered_map<HashType, BlobRef> DuplicatesThisJob;
 
-	for (size_t TileX = 0; TileX < ResultPromise->Rows(); TileX++)
+	for (size_t RowId = 0; RowId < ResultPromise->Rows(); RowId++)
 	{
-		for (size_t TileY = 0; TileY < ResultPromise->Cols(); TileY++)
+		for (size_t ColId = 0; ColId < ResultPromise->Cols(); ColId++)
 		{
-			CHashPtr ResultTileHash = bCanHandleTiles ? TileHash(TileX, TileY) : TileHash(-1, -1);
+			CHashPtr ResultTileHash = bCanHandleTiles ? TileHash(RowId, ColId) : TileHash(-1, -1);
 			BlobRef TileBlob = !(Batch->IsNoCache() || bIsNoCache) ? TextureGraphEngine::GetBlobber()->FindSingle(ResultTileHash->Value()) : BlobRef();
 
 			if (!TileBlob && DuplicatesThisJob.size())
@@ -783,7 +783,7 @@ AsyncPrepareResult Job::PrepareTargets(JobBatch* Batch)
 			if (!TileBlob || !bCanHandleTiles)
 			{
 				CHashPtr TileHashCopy;
-				ResultTileDesc.Name = TextureHelper::CreateTileName(ResultDesc.Name, TileX, TileY);
+				ResultTileDesc.Name = TextureHelper::CreateTileName(ResultDesc.Name, RowId, ColId);
 				ResultTileDesc.bIsTransient = bIsDiscard;
 
 				TileHashCopy = std::make_shared<CHash>(ResultTileHash); /// Make a copy of the hash
@@ -796,7 +796,7 @@ AsyncPrepareResult Job::PrepareTargets(JobBatch* Batch)
 
 				check(TileBlobNew->GetWidth() > 0 && TileBlobNew->GetHeight() > 0);
 
-				TileInvalidationMatrix[TileX][TileY] = true;
+				TileInvalidationMatrix[RowId][ColId] = true;
 
 				BlobCacheOptions CacheOpt;
 				CacheOpt.Discard = bIsDiscard;
@@ -806,17 +806,17 @@ AsyncPrepareResult Job::PrepareTargets(JobBatch* Batch)
 			}
 			else
 			{
-				UE_LOG(LogJob, VeryVerbose, TEXT("Transform: %s (Tile: %d, %d) => %s"), *Transform->GetName(), (int32)TileX, (int32)TileY, *TileBlob->Name());
-				TileInvalidationMatrix[TileX][TileY] = false;
+				UE_LOG(LogJob, VeryVerbose, TEXT("Transform: %s (Tile: %d, %d) => %s"), *Transform->GetName(), (int32)RowId, (int32)ColId, *TileBlob->Name());
+				TileInvalidationMatrix[RowId][ColId] = false;
 			}
 
 			if (bCanHandleTiles)
 			{
 				JobResultPtr TileResult = std::make_shared<JobResult>(TileBlob, ResultTileHash);
 
-				if (TileInvalidationMatrix[TileX][TileY])
+				if (TileInvalidationMatrix[RowId][ColId])
 				{
-					FString TransformCloneName = FString::Printf(TEXT("%s-%llu,%llu"), *Transform->GetName(), TileX, TileY);
+					FString TransformCloneName = FString::Printf(TEXT("%s-%llu,%llu"), *Transform->GetName(), RowId, ColId);
 					BlobTransformPtr TransformClone = Transform->DuplicateInstance(TransformCloneName);
 
 					if (TransformClone)
@@ -826,10 +826,10 @@ AsyncPrepareResult Job::PrepareTargets(JobBatch* Batch)
 						DuplicatesThisJob[ResultTileHash->Value()] = TileBlob;
 				}
 
-				TileResults[TileX][TileY] = TileResult;
+				TileResults[RowId][ColId] = TileResult;
 			}
 
-			ResultPromise->SetTile(TileX, TileY, TileBlob);
+			ResultPromise->SetTile(RowId, ColId, TileBlob);
 		}
 	}
 
@@ -879,21 +879,21 @@ AsyncPrepareResult Job::PrepareResources(JobBatch* Batch)
 		});
 	}
 
-AsyncJobResultPtr Job::RunTile(JobRunInfo InRunInfo, int32 TileX, int32 TileY)
+AsyncJobResultPtr Job::RunTile(JobRunInfo InRunInfo, int32 RowId, int32 ColId)
 {
-	if (TileX >= 0 && TileY >= 0)
+	if (RowId >= 0 && ColId >= 0)
 	{
-		if (!TileInvalidationMatrix[TileX][TileY])
-			return cti::make_ready_continuable(TileResults[TileX][TileY]);
+		if (!TileInvalidationMatrix[RowId][ColId])
+			return cti::make_ready_continuable(TileResults[RowId][ColId]);
 	}
 
 	CHashPtr JobHash = FinalJobResult->HashValue;
 	BlobTransformPtr ResultTransform = FinalJobResult->Transform;
 
-	if (TileX >= 0 && TileY >= 0)
+	if (RowId >= 0 && ColId >= 0)
 	{
-		ResultTransform = TileResults[TileX][TileY]->Transform;
-		JobHash = TileResults[TileX][TileY]->HashValue;
+		ResultTransform = TileResults[RowId][ColId]->Transform;
+		JobHash = TileResults[RowId][ColId]->HashValue;
 	}
 
 	check(ResultTransform);
@@ -901,12 +901,12 @@ AsyncJobResultPtr Job::RunTile(JobRunInfo InRunInfo, int32 TileX, int32 TileY)
 	UE_LOG(LogJob, VeryVerbose, TEXT("No cached results for hash: %llu [Recalculating ...]"), JobHash->Value());
 
 	/// wait for the Transform to finish
-	ExecTransform(InRunInfo, ResultTransform, TileX, TileY, JobHash).apply(cti::transforms::wait());
+	ExecTransform(InRunInfo, ResultTransform, RowId, ColId, JobHash).apply(cti::transforms::wait());
 
 	if (ResultTransform->GeneratesData() && CanHandleTiles())
 	{
-		TileResults[TileX][TileY]->HashValue = JobHash;
-		return cti::make_ready_continuable(TileResults[TileX][TileY]);
+		TileResults[RowId][ColId]->HashValue = JobHash;
+		return cti::make_ready_continuable(TileResults[RowId][ColId]);
 	}
 	else
 	{
@@ -917,7 +917,7 @@ AsyncJobResultPtr Job::RunTile(JobRunInfo InRunInfo, int32 TileX, int32 TileY)
 	return cti::make_ready_continuable(FinalJobResult);
 }
 
-AsyncTransformResultPtr Job::ExecTransform(JobRunInfo InRunInfo, BlobTransformPtr TransformObj, int32 TileX, int32 TileY, CHashPtr jobHash)
+AsyncTransformResultPtr Job::ExecTransform(JobRunInfo InRunInfo, BlobTransformPtr TransformObj, int32 RowId, int32 ColId, CHashPtr JobHash)
 {
 	TransformArgs TransArgs;
 	TransArgs.Dev = RunInfo.Dev;
@@ -930,9 +930,9 @@ AsyncTransformResultPtr Job::ExecTransform(JobRunInfo InRunInfo, BlobTransformPt
 	{
 		TransArgs.Target = GetResultRef();
 
-		if (TileX >= 0 && TileY >= 0 && Result->TiledTarget())
+		if (RowId >= 0 && ColId >= 0 && Result->TiledTarget())
 		{
-			BlobRef TileResult = Result->GetTile(TileX, TileY);
+			BlobRef TileResult = Result->GetTile(RowId, ColId);
 			
 			/// This should've been prepared in the PrepareResources function
 			check(TileResult);
@@ -944,11 +944,11 @@ AsyncTransformResultPtr Job::ExecTransform(JobRunInfo InRunInfo, BlobTransformPt
 			TransArgs.Target = TileResult.get();
 
 			/// Set the job result preemptively
-			JobResultPtr TileJobResult = std::make_shared<JobResult>(TransArgs.Target, jobHash);
+			JobResultPtr TileJobResult = std::make_shared<JobResult>(TransArgs.Target, JobHash);
 			TileJobResult->Transform = TransformObj;
 			TileJobResult->BlobObj = TileResult;
 
-			SetTileResult(TileX, TileY, TileJobResult);
+			SetTileResult(RowId, ColId, TileJobResult);
 		}
 	}
 
@@ -1002,25 +1002,25 @@ AsyncJobResultPtr Job::Run(JobRunInfo InRunInfo)
 		}
 		else
 		{
-			for (int32 TileX = 0; TileX < (int32)LocalMatrix.Rows(); TileX++)
+			for (int32 RowId = 0; RowId < (int32)LocalMatrix.Rows(); RowId++)
 			{
-				for (int32 TileY = 0; TileY < (int32)LocalMatrix.Cols(); TileY++)
+				for (int32 ColId = 0; ColId < (int32)LocalMatrix.Cols(); ColId++)
 				{
-					LocalMatrix[TileX][TileY] = true;
+					LocalMatrix(RowId, ColId) = true;
 				}
 			}
 		}
 
 		std::vector<std::decay_t<AsyncJobResultPtr>, std::allocator<std::decay_t<AsyncJobResultPtr>>> Promises;
-		for (int32 TileX = 0; TileX < (int32)LocalMatrix.Rows(); TileX++)
+		for (int32 RowId = 0; RowId < (int32)LocalMatrix.Rows(); RowId++)
 		{
-			for (int32 TileY = 0; TileY < (int32)LocalMatrix.Cols(); TileY++)
+			for (int32 ColId = 0; ColId < (int32)LocalMatrix.Cols(); ColId++)
 			{
-				bool InvalidateThisTile = LocalMatrix[TileX][TileY];
-				if (InvalidateThisTile)
+				int32 InvalidateThisTile = LocalMatrix[RowId][ColId];
+				if (InvalidateThisTile != 0)
 				{
 					Promises.reserve(1 + Promises.size());
-					Promises.emplace_back(std::forward<AsyncJobResultPtr>(RunTile(InRunInfo, TileX, TileY)));
+					Promises.emplace_back(std::forward<AsyncJobResultPtr>(RunTile(InRunInfo, RowId, ColId)));
 				}
 				else
 				{
@@ -1336,16 +1336,16 @@ void Job::ResetForReplay(bool noCache)
 	}
 }
 
-bool Job::GetTileInvalidation(int32 TileX, int32 TileY) const
+bool Job::GetTileInvalidation(int32 RowId, int32 ColId) const
 {
-	if (TileX >= 0 && TileY >= 0)
+	if (RowId >= 0 && ColId >= 0)
 	{
-		check(TileX < (int32)TileInvalidationMatrix.Rows());
-		check(TileY < (int32)TileInvalidationMatrix.Cols());
+		check(RowId < (int32)TileInvalidationMatrix.Rows());
+		check(ColId < (int32)TileInvalidationMatrix.Cols());
 
-		return TileInvalidationMatrix[TileX][TileY];
+		return TileInvalidationMatrix[RowId][ColId] != 0;
 	}
 
-	return TileInvalidationMatrix[0][0];
+	return TileInvalidationMatrix(0, 0) != 0;
 }
 
diff --git a/Engine/Plugins/Experimental/TextureGraph/Source/TextureGraphEngine/Job/Job.h b/Engine/Plugins/Experimental/TextureGraph/Source/TextureGraphEngine/Job/Job.h
--- a/Engine/Plugins/Experimental/TextureGraph/Source/TextureGraphEngine/Job/Job.h
+++ b/Engine/Plugins/Experimental/TextureGraph/Source/TextureGraphEngine/Job/Job.h
@@ -82,8 +82,8 @@ protected:
 	typedef TMap<HashType, IndexPair> JobLUT;
 
 	/// Binds/Unbinds args for individual tiles or the entire blob if its a non-tiled job
-	virtual AsyncJobArgResultPtr	BindArgs(JobRunInfo NewRunInfo, BlobTransformPtr TransformObj, int32 TileX, int32 TileY);
-	virtual AsyncJobArgResultPtr	UnbindArgs(JobRunInfo NewRunInfo, BlobTransformPtr TransformObj, int32 TileX, int32 TileY);
+	virtual AsyncJobArgResultPtr	BindArgs(JobRunInfo NewRunInfo, BlobTransformPtr TransformObj, int32 RowId, int32 ColId);
+	virtual AsyncJobArgResultPtr	UnbindArgs(JobRunInfo NewRunInfo, BlobTransformPtr TransformObj, int32 RowId, int32 ColId);
 
 	/// Binds/Unbinds args for all the tiles at once and returns a promise
 	virtual AsyncInt				BindArgs_All(JobRunInfo RunInfo);
@@ -92,13 +92,13 @@ protected:
 	typedef std::function			<AsyncJobArgResultPtr(JobRunInfo, BlobTransformPtr, int32, int32)> Bind_Unbind_Func;
 	virtual AsyncInt				Bind_Or_Unbind_All_Generic(JobRunInfo NewRunInfo, Bind_Unbind_Func BFunc);
 
-	virtual AsyncJobResultPtr		RunTile(JobRunInfo RunInfo, int32 TileX, int32 TileY);
+	virtual AsyncJobResultPtr		RunTile(JobRunInfo RunInfo, int32 RowId, int32 ColId);
 	virtual AsyncJobResultPtr		RunSingle(JobRunInfo RunInfo);
 	virtual AsyncJobResultPtr		FinaliseTiles(JobRunInfo RunInfo);
-	void							SetTileResult(int32 TileX, int32 TileY, JobResultPtr TileResult);
+	void							SetTileResult(int32 RowId, int32 ColId, JobResultPtr TileResult);
 	virtual AsyncPrepareResult		PrepareTargets(JobBatch* batch);
 
-	AsyncTransformResultPtr			ExecTransform(JobRunInfo RunInfo, BlobTransformPtr TransformObj , int32 TileX, int32 TileY, CHashPtr jobHash);
+	AsyncTransformResultPtr			ExecTransform(JobRunInfo RunInfo, BlobTransformPtr TransformObj , int32 RowId, int32 ColId, CHashPtr jobHash);
 	virtual void					MarkJobDone();
 	virtual bool					IsDiscard() const;
 
@@ -125,8 +125,8 @@ public:
 	virtual void					AddResultToBlobber();
 
 	virtual CHashPtr				Hash() const;
-	virtual CHashPtr				TileHash(int32 TileX, int32 TileY) const;
-	virtual CHashPtr				CalcTileHash(int32 TileX, int32 TileY) const;
+	virtual CHashPtr				TileHash(int32 RowId, int32 ColId) const;
+	virtual CHashPtr				CalcTileHash(int32 RowId, int32 ColId) const;
 	virtual bool					CanHandleTiles() const;
 	virtual bool					CheckDefaultArgs() const;
 
@@ -158,7 +158,7 @@ public:
 	virtual Device*					GetTargetDevice() const override;
 	virtual ENamedThreads::Type		GetExecutionThread() const override;
 	virtual bool					IsAsync() const override;
-	bool							GetTileInvalidation(int32 TileX, int32 TileY) const;
+	bool							GetTileInvalidation(int32 RowId, int32 ColId) const;
 
 	virtual void					GetDependencies(JobPtrVec& Prior, JobPtrVec& After, JobRunInfo RunInfo);
 
@@ -191,10 +191,10 @@ public:
 	FORCEINLINE const JobRunInfo&	GetRunInfo() const { return RunInfo; }
 	FORCEINLINE const RenderMesh*	GetMesh() const { return Mesh; }
 	FORCEINLINE void				SetMesh(RenderMesh* NewMesh) { Mesh = NewMesh; }
-	FORCEINLINE JobResultPtr		GetTileResult(int32 TileX, int32 TileY) const
+	FORCEINLINE JobResultPtr		GetTileResult(int32 RowId, int32 ColId) const
 	{
-		if (TileX >= 0 && TileY >= 0)
-			return TileResults[TileX][TileY];
+		if (RowId >= 0 && ColId >= 0)
+			return TileResults[RowId][ColId];
 		return FinalJobResult;
 	}
 
diff --git a/Engine/Plugins/Experimental/TextureGraph/Source/TextureGraphEngine/Job/JobArgs.cpp b/Engine/Plugins/Experimental/TextureGraph/Source/TextureGraphEngine/Job/JobArgs.cpp
--- a/Engine/Plugins/Experimental/TextureGraph/Source/TextureGraphEngine/Job/JobArgs.cpp
+++ b/Engine/Plugins/Experimental/TextureGraph/Source/TextureGraphEngine/Job/JobArgs.cpp
@@ -35,11 +35,11 @@ JobArg_Blob::JobArg_Blob(TiledBlobPtr BlobObj, const ResourceBindInfo& BindInfo)
 
 #if DEBUG_BLOB_REF_KEEPING == 1
 	/// Check whether all the tiles are valid and available
-	//for (size_t TileX = 0; TileX < BlobObj->Rows(); TileX++)
+	//for (size_t RowId = 0; RowId < BlobObj->Rows(); RowId++)
 	//{
-	//	for (size_t TileY = 0; TileY < BlobObj->Cols(); TileY++) 
+	//	for (size_t ColId = 0; ColId < BlobObj->Cols(); ColId++) 
 	//	{
-	//		BlobRef Tile = BlobObj->Tile(TileX, TileY);
+	//		BlobRef Tile = BlobObj->Tile(RowId, ColId);
 	//		check(!Tile.expired());
 	//	}
 	//}
@@ -168,13 +168,13 @@ AsyncJobArgResultPtr JobArg_Blob::Bind(JobArgBindInfo JobBindInfo)
 			{
 				// Collect the individual tile textures from the root tiled BlobObj in an array
 				std::vector<TexPtr> TileTexObjs;
-				for (int TileX = 0; TileX < 2; TileX++)
+				for (int RowId = 0; RowId < 2; RowId++)
 				{
-					for (int TileY = 0; TileY < 2; TileY++)
+					for (int ColId = 0; ColId < 2; ColId++)
 					{
-						if (RootBlob->IsValidTileIndex(TileX + 2 * JobBindInfo.TileX, TileY + 2 * JobBindInfo.TileY))
+						if (RootBlob->IsValidTileIndex(RowId + 2 * JobBindInfo.RowId, ColId + 2 * JobBindInfo.ColId))
 						{
-							auto Tile = RootBlob->GetTile(TileX + 2 * JobBindInfo.TileX, TileY + 2 * JobBindInfo.TileY);
+							auto Tile = RootBlob->GetTile(RowId + 2 * JobBindInfo.RowId, ColId + 2 * JobBindInfo.ColId);
 							auto TileBuffer = std::static_pointer_cast<DeviceBuffer_FX>(Tile->GetBufferRef().GetPtr());
 							TileTexObjs.push_back(TileBuffer->GetTexture());
 						}
@@ -202,30 +202,30 @@ AsyncJobArgResultPtr JobArg_Blob::Bind(JobArgBindInfo JobBindInfo)
 			{
 				// Collect the individual tile textures from the root tiled BlobObj in an array
 				std::vector<TexPtr> Textures;
-				auto MainX = JobBindInfo.TileX;
-				auto MainY = JobBindInfo.TileY;
-				auto GridX = RootBlob->Cols();
-				auto GridY = RootBlob->Rows();
+				auto MainRow = JobBindInfo.RowId;
+				auto MainCol = JobBindInfo.ColId;
+				auto GridRow = RootBlob->Rows();
+				auto GridCol = RootBlob->Cols();
 
-				for (int TileX = 0; TileX < 3; TileX++)
+				for (int RowId = 0; RowId < 3; RowId++)
 				{
-					for (int TileY = 0; TileY < 3; TileY++)
+					for (int ColId = 0; ColId < 3; ColId++)
 					{
 						// Compute the true tile index
-						auto InnerTileX = (MainX - 1 + TileX);
-						auto InnerTileY = (MainY - 1 + TileY);
-						if (InnerTileX < 0)
-							InnerTileX = GridX - 1;
-						if (InnerTileX >= GridX)
-							InnerTileX = 0;
-						if (InnerTileY < 0)
-							InnerTileY = GridY - 1;
-						if (InnerTileY >= GridY)
-							InnerTileY = 0;
-
-						if (RootBlob->IsValidTileIndex(InnerTileX, InnerTileY))
+						auto InnerRowId = (MainRow - 1 + RowId);
+						auto InnerColId = (MainCol - 1 + ColId);
+						if (InnerRowId < 0)
+							InnerRowId = GridRow - 1;
+						if (InnerRowId >= GridRow)
+							InnerRowId = 0;
+						if (InnerColId < 0)
+							InnerColId = GridCol - 1;
+						if (InnerColId >= GridCol)
+							InnerColId = 0;
+
+						if (RootBlob->IsValidTileIndex(InnerRowId, InnerColId))
 						{
-							auto Tile = RootBlob->GetTile(InnerTileX, InnerTileY);
+							auto Tile = RootBlob->GetTile(InnerRowId, InnerColId);
 							auto TileBuffer = std::static_pointer_cast<DeviceBuffer_FX>(Tile->GetBufferRef().GetPtr());
 							Textures.push_back(TileBuffer->GetTexture());
 						}
@@ -289,8 +289,8 @@ AsyncJobArgResultPtr JobArg_Blob::Bind(JobArgBindInfo JobBindInfo)
 		check(BlobToBind);
 
 		bool bArgCanHandleTiles = CanHandleTiles();
-		if (bArgCanHandleTiles && transform->CanHandleTiles() && JobBindInfo.TileX >= 0 && JobBindInfo.TileY >= 0)
-			BlobToBind = BlobObjRef.get()->GetTile(JobBindInfo.TileX, JobBindInfo.TileY).lock();
+		if (bArgCanHandleTiles && transform->CanHandleTiles() && JobBindInfo.RowId >= 0 && JobBindInfo.ColId >= 0)
+			BlobToBind = BlobObjRef.get()->GetTile(JobBindInfo.RowId, JobBindInfo.ColId).lock();
 
 		FString TransformName = CurrentJob->GetTransform()->GetName();
 
@@ -332,9 +332,9 @@ AsyncJobArgResultPtr JobArg_Blob::Unbind(JobArgBindInfo JobBindInfo)
 		BlobPtr blobToBind = GetRootBlob(JobBindInfo);
 
 		bool bCanArgHandleTiles = CanHandleTiles();
-		if (bCanArgHandleTiles && transform->CanHandleTiles() && JobBindInfo.TileX >= 0 && JobBindInfo.TileY >= 0)
+		if (bCanArgHandleTiles && transform->CanHandleTiles() && JobBindInfo.RowId >= 0 && JobBindInfo.ColId >= 0)
 		{
-			blobToBind = BlobObjRef.get()->GetTile(JobBindInfo.TileX, JobBindInfo.TileY).lock();
+			blobToBind = BlobObjRef.get()->GetTile(JobBindInfo.RowId, JobBindInfo.ColId).lock();
 		}
 
 		ArgBindInfo.BatchId = JobBindInfo.Batch->GetBatchId();
@@ -348,10 +348,10 @@ AsyncJobArgResultPtr JobArg_Blob::Unbind(JobArgBindInfo JobBindInfo)
 	}
 }
 
-bool JobArg_Blob::IsLateBound(uint32 TileX, uint32 TileY) const
+bool JobArg_Blob::IsLateBound(uint32 RowId, uint32 ColId) const
 {
 	TiledBlobPtr BlobObj = BlobObjRef.get();
-	BlobRef TileXY = BlobObj->GetTile(TileX, TileY);
+	BlobRef TileXY = BlobObj->GetTile(RowId, ColId);
 
 	/// If the tile isn't even there then it's definitely late bound
 	if (TileXY.IsNull())
@@ -366,10 +366,10 @@ const BufferDescriptor* JobArg_Blob::GetDescriptor() const
 	return BlobObjRef ? &BlobObjRef->GetDescriptor() : nullptr;
 }
 
-CHashPtr JobArg_Blob::TileHash(uint32 TileX, uint32 TileY) const
+CHashPtr JobArg_Blob::TileHash(uint32 RowId, uint32 ColId) const
 {
 	TiledBlobPtr BlobObj = BlobObjRef.get();
-	return BlobObj->GetTile(TileX, TileY)->Hash();
+	return BlobObj->GetTile(RowId, ColId)->Hash();
 }
 
 CHashPtr JobArg_Blob::Hash() const
@@ -402,7 +402,7 @@ CHashPtr JobArg_Mesh::Hash() const
 	return Mesh->Hash(); 
 }
 
-CHashPtr JobArg_Mesh::TileHash(uint32 TileX, uint32 TileY) const
+CHashPtr JobArg_Mesh::TileHash(uint32 RowId, uint32 ColId) const
 {
 	return Mesh->Hash();
 }
@@ -416,15 +416,15 @@ AsyncJobArgResultPtr JobArg_TileInfo::Bind(JobArgBindInfo JobBindInfo)
 
 	TiledBlobPtr Result = JobObj->GetResult();
 	float TileWidth = Result->GetWidth() / Result->Cols();
-	float tileHeight = Result->GetHeight() / Result->Rows();
+	float TileHeight = Result->GetHeight() / Result->Rows();
 
-	Value.TileX = JobBindInfo.TileX;
+	Value.TileX = JobBindInfo.ColId;
 	Value.TileCountX = Result->Cols();
 	Value.TileWidth = TileWidth;
 	
-	Value.TileY = JobBindInfo.TileY;
+	Value.TileY = JobBindInfo.RowId;
 	Value.TileCountY = Result->Rows();
-	Value.TileHeight = tileHeight;
+	Value.TileHeight = TileHeight;
 
 	Transform->BindStruct<FTileInfo>(Value, ArgBindInfo);
 	
@@ -443,14 +443,14 @@ CHashPtr JobArg_TileInfo::Hash() const
 	return HashValue;
 }
 
-CHashPtr JobArg_TileInfo::TileHash(uint32 TileX, uint32 TileY) const
+CHashPtr JobArg_TileInfo::TileHash(uint32 RowId, uint32 ColId) const
 {
 	HashTypeVec StructHash =
 	{
-		MX_HASH_VAL_DEF(TileX),
+		MX_HASH_VAL_DEF(RowId),
 		MX_HASH_VAL_DEF(Value.TileCountX),
 		MX_HASH_VAL_DEF(Value.TileWidth),
-		MX_HASH_VAL_DEF(TileY),
+		MX_HASH_VAL_DEF(ColId),
 		
 		MX_HASH_VAL_DEF(Value.TileCountY),
 		MX_HASH_VAL_DEF(Value.TileHeight),
@@ -467,12 +467,12 @@ CHashPtr JobArg_ForceTiling::Hash() const
 	return TileHash(-1, -1);
 }
 
-CHashPtr JobArg_ForceTiling::TileHash(uint32 TileX, uint32 TileY) const
+CHashPtr JobArg_ForceTiling::TileHash(uint32 RowId, uint32 ColId) const
 {
 	HashTypeVec StructHash =
 	{
-		MX_HASH_VAL_DEF(TileX),
-		MX_HASH_VAL_DEF(TileY),
+		MX_HASH_VAL_DEF(RowId),
+		MX_HASH_VAL_DEF(ColId),
 	};
 	
 	return std::make_shared<CHash>(DataUtil::Hash(StructHash), true);
diff --git a/Engine/Plugins/Experimental/TextureGraph/Source/TextureGraphEngine/Job/JobArgs.h b/Engine/Plugins/Experimental/TextureGraph/Source/TextureGraphEngine/Job/JobArgs.h
--- a/Engine/Plugins/Experimental/TextureGraph/Source/TextureGraphEngine/Job/JobArgs.h
+++ b/Engine/Plugins/Experimental/TextureGraph/Source/TextureGraphEngine/Job/JobArgs.h
@@ -31,8 +31,8 @@ struct TEXTUREGRAPHENGINE_API JobArgBindInfo
 	Job*							JobObj = nullptr;	/// The job that this is associated with
 	const JobBatch*					Batch = nullptr;	/// What is the batch that we're going to be part of (can be nullptr)
 	size_t							JobIndex = 0;		/// The index of this job
-	int32							TileX = -1;			/// The x-tile index that we're binding for or -1 if it needs to be bound as one
-	int32							TileY = -1;			/// The y-tile index tile that we're binding for or -1 if it needs to be bound as one
+	int32							RowId = -1;			/// The row index that we're binding for or -1 if it needs to be bound as one
+	int32							ColId = -1;			/// The col index tile that we're binding for or -1 if it needs to be bound as one
 	Device*							Dev = nullptr;		/// The device to which this is going to be bound to
 	BlobTransformPtr				Transform;			/// The transformation function to which the argument will be bound
 	uint32							LODLevel = 0;		/// What is the lod level that we'll be binding this to
diff --git a/Engine/Plugins/Experimental/TextureGraph/Source/TextureGraphEngine/Model/Mix/MixUpdateCycle.cpp b/Engine/Plugins/Experimental/TextureGraph/Source/TextureGraphEngine/Model/Mix/MixUpdateCycle.cpp
--- a/Engine/Plugins/Experimental/TextureGraph/Source/TextureGraphEngine/Model/Mix/MixUpdateCycle.cpp
+++ b/Engine/Plugins/Experimental/TextureGraph/Source/TextureGraphEngine/Model/Mix/MixUpdateCycle.cpp
@@ -33,20 +33,14 @@ void MixTargetUpdate::InvalidateTile(int32 Row, int32 Col)
 
 void MixTargetUpdate::InvalidateAllTiles()
 {
-	for (size_t RowIndex = 0; RowIndex < InvalidationMatrix.Rows(); RowIndex++)
-	{
-		std::vector<bool>& Column = InvalidationMatrix[RowIndex];
-		std::fill(Column.begin(), Column.end(), true);
-	}
+	for (size_t Index = 0; Index < InvalidationMatrix.Length(); Index++)
+		InvalidationMatrix.Tiles()[Index] = 1;
 }
 
 void MixTargetUpdate::InvalidateNoTiles()
 {
-	for (size_t RowIndex = 0; RowIndex < InvalidationMatrix.Rows(); RowIndex++)
-	{
-		std::vector<bool>& Column = InvalidationMatrix[RowIndex];
-		std::fill(Column.begin(), Column.end(), false);
-	}
+	for (size_t Index = 0; Index < InvalidationMatrix.Length(); Index++)
+		InvalidationMatrix.Tiles()[Index] = 0;
 }
 
 //////////////////////////////////////////////////////////////////////////
diff --git a/Engine/Plugins/Experimental/TextureGraph/Source/TextureGraphEngine/Model/Mix/MixUpdateCycle.h b/Engine/Plugins/Experimental/TextureGraph/Source/TextureGraphEngine/Model/Mix/MixUpdateCycle.h
--- a/Engine/Plugins/Experimental/TextureGraph/Source/TextureGraphEngine/Model/Mix/MixUpdateCycle.h
+++ b/Engine/Plugins/Experimental/TextureGraph/Source/TextureGraphEngine/Model/Mix/MixUpdateCycle.h
@@ -54,7 +54,7 @@ public:
 	FORCEINLINE const 
 		TileInvalidateMatrix&		GetInvalidationMatrix() const { return InvalidationMatrix; }
 	FORCEINLINE bool				CheckIsValid(int32 Row, int32 Col) const { return (size_t)Row < InvalidationMatrix.Rows() && (size_t)Col < InvalidationMatrix.Cols(); }
-	FORCEINLINE bool				IsInvalidated(int32 Row, int32 Col) const { check(CheckIsValid(Row, Col)); return InvalidationMatrix[Row][Col]; }
+	FORCEINLINE bool				IsInvalidated(int32 Row, int32 Col) const { check(CheckIsValid(Row, Col)); return InvalidationMatrix(Row, Col) != 0; }
 	FORCEINLINE MixTextureSet&		GetLastRender() { return LastRender; }
 	FORCEINLINE int32				GetTargetId() const { return TargetId; }
 	FORCEINLINE JobPtrW				GetLastAddedJob() const { return LastAddedJob; }
diff --git a/Engine/Plugins/Experimental/TextureGraph/Source/TextureGraphEngine/Transform/Expressions/T_ArrayGrid.cpp b/Engine/Plugins/Experimental/TextureGraph/Source/TextureGraphEngine/Transform/Expressions/T_ArrayGrid.cpp
--- a/Engine/Plugins/Experimental/TextureGraph/Source/TextureGraphEngine/Transform/Expressions/T_ArrayGrid.cpp
+++ b/Engine/Plugins/Experimental/TextureGraph/Source/TextureGraphEngine/Transform/Expressions/T_ArrayGrid.cpp
@@ -48,31 +48,39 @@ public:
 		bool bIsWellDefined = true;
 		uint32 MaxRowHeight = 0;
 		uint32 MaxColWidth = 0;
+		size_t HashIndex = 0;
 
 		/// First of all we combine all the tiles into a single result
-		for (int32 X = 0; X < NumCols; X++)
+		for (int32 ArrRowIter = 0; ArrRowIter < NumRows; ArrRowIter++)
 		{
-			int32 ColIndex = X * InNumCols;
-			for (int32 Y = 0; Y < NumRows; Y++)
+			int32 RowIndex = ArrRowIter * InNumCols;
+			for (int32 ArrColIter = 0; ArrColIter < NumCols; ArrColIter++)
 			{
-				int32 TileIndex = Y * NumCols + X;
+				int32 TileIndex = ArrRowIter * NumCols + ArrColIter;
 				check(TileIndex < Inputs.Num());
 				TiledBlobPtr Input = Inputs[TileIndex];
-				int32 RowIndex = Y * InNumRows;
+				int32 ColIndex = ArrColIter * InNumCols;
 
 				MaxRowHeight = std::max(MaxRowHeight, Input->GetHeight());
 				MaxColWidth = std::max(MaxColWidth, Input->GetWidth());
 
-				InputHashes[TileIndex] = Input->Hash();
-				check(InputHashes[TileIndex]);
+				check(HashIndex < InputHashes.size());
+				check(Input->Hash());
+				InputHashes[HashIndex++] = Input->Hash();
+
 				bIsWellDefined &= Input->IsWellDefined();
 
-				for (int32 TileX = 0; TileX < Input->Rows(); TileX++)
+				/// Ok here we copy the tiles over to the result tiles. However, there can be situations where the 
+				/// input has fewer tiles than the output (1x1 as opposed to 8x8 for instance). In this case
+				/// we just copy clamped (RowId, ColId) over to the destination
+				for (int32 RowId = 0; RowId < InNumRows; RowId++)
 				{
-					for (int32 TileY = 0; TileY < Input->Cols(); TileY++)
+					int32 InputRowId = std::min<int32>(RowId, (int32)Input->Rows() - 1);
+					for (int32 ColId = 0; ColId < InNumCols; ColId++)
 					{
-						BlobPtr Tile = Input->GetTile(TileX, TileY);
-						ResultTiles[RowIndex + TileX][ColIndex + TileY] = Tile;
+						int32 InputColId = std::min<int32>(ColId, (int32)Input->Cols() - 1);
+						BlobPtr Tile = Input->GetTile(InputRowId, InputColId);
+						ResultTiles[RowIndex + RowId][ColIndex + ColId] = Tile;
 
 						if (Tile)
 							bIsWellDefined &= Tile->IsWellDefined();
@@ -132,6 +140,13 @@ TiledBlobPtr T_ArrayGrid::Create(MixUpdateCyclePtr Cycle, BufferDescriptor& InDe
 
 	int32 InNumRows = Cycle->GetMix()->GetNumXTiles();
 	int32 InNumCols = Cycle->GetMix()->GetNumYTiles();
+
+	for (int32 InputIndex = 0; InputIndex < Inputs.Num(); InputIndex++)
+	{
+		InNumRows = std::min<int32>(InNumRows, (int32)Inputs[InputIndex]->Rows());
+		InNumCols = std::min<int32>(InNumCols, (int32)Inputs[InputIndex]->Cols());
+	}
+
 	CHashPtrVec InputHashes(Inputs.Num());
 	bool bIsWellDefined = true;
 	uint32 MaxRowHeight = 0;
@@ -149,6 +164,10 @@ TiledBlobPtr T_ArrayGrid::Create(MixUpdateCyclePtr Cycle, BufferDescriptor& InDe
 		Desc.Width = Desc.Height = 0;
 		TiledBlob_PromisePtr CombinedResultTiled = std::make_shared<TiledBlob_Promise>(Desc, InNumRows * NumRows, InNumCols * NumCols, Hash);
 		JobUPtr JobObj = std::make_unique<Job_ArrayGrid>(Cycle->GetMix(), TargetId, Inputs, NumRows, NumCols, MaxColWidth, MaxRowHeight, InDesiredOutputDesc, CombinedResultTiled);
+
+		for (int32 InputIndex = 0; InputIndex < Inputs.Num(); InputIndex++)
+			JobObj->AddArg(WithUnbounded(ARG_BLOB(Inputs[InputIndex], "Input")));
+
 		TiledBlobPtr CombinedResultOne = T_CombineTiledBlob::Create(Cycle, InDesiredOutputDesc, TargetId, CombinedResultTiled, std::move(JobObj));
 		FinalResult = T_SplitToTiles::Create(Cycle, TargetId, CombinedResultOne);
 	}
@@ -161,7 +180,9 @@ TiledBlobPtr T_ArrayGrid::Create(MixUpdateCyclePtr Cycle, BufferDescriptor& InDe
 		/// The we combine those tiles into a single result
 		TiledBlobPtr CombinedResultTiled = TiledBlob::InitFromTiles(DesiredOutputDesc, ResultTiles);
 		TiledBlobPtr CombinedResultOne = T_CombineTiledBlob::Create(Cycle, DesiredOutputDesc, TargetId, CombinedResultTiled);
-		FinalResult = T_SplitToTiles::Create(Cycle, TargetId, CombinedResultOne);
+
+		if (DesiredOutputDesc.Width > 1 && DesiredOutputDesc.Height > 1)
+			FinalResult = T_SplitToTiles::Create(Cycle, TargetId, CombinedResultOne);
 	}
 
 	return FinalResult;
diff --git a/Engine/Plugins/Experimental/TextureGraph/Source/TextureGraphEngine/Transform/Utility/T_CombineTiledBlob.cpp b/Engine/Plugins/Experimental/TextureGraph/Source/TextureGraphEngine/Transform/Utility/T_CombineTiledBlob.cpp
--- a/Engine/Plugins/Experimental/TextureGraph/Source/TextureGraphEngine/Transform/Utility/T_CombineTiledBlob.cpp
+++ b/Engine/Plugins/Experimental/TextureGraph/Source/TextureGraphEngine/Transform/Utility/T_CombineTiledBlob.cpp
@@ -38,6 +38,20 @@ TiledBlobPtr T_CombineTiledBlob::Create(MixUpdateCyclePtr Cycle, BufferDescripto
 	}
 
 	return SourceTex;
+#if 0
+	RenderMaterial_FXPtr RenderMaterial = TextureGraphEngine::GetMaterialManager()->CreateMaterialOfType_FX<Fx_FullScreenCopy>(TEXT("T_CombineTiledBlob"));
+	JobUPtr JobObj = std::make_unique<Job>(Cycle->GetMix(), TargetId, std::static_pointer_cast<BlobTransform>(RenderMaterial));
+	JobObj
+		->AddArg(ARG_BLOB(SourceTex, "SourceTexture"))
+		->SetTiled(false)
+		;
+
+	TiledBlobPtr Result = JobObj->InitResult(JobObj->GetName(), &DesiredOutputDesc, 1, 1);
+	Cycle->AddJob(TargetId, std::move(JobObj));
+	Result->MakeSingleBlob();
+
+	return Result;
+#endif
 }
 
 
@@ -53,7 +67,6 @@ Device* CombineTiledBlob_Transform::TargetDevice(size_t DevIndex) const
 	return Device_FX::Get();
 }
 
-
 bool CombineTiledBlob_Transform::GeneratesData() const 
 {
 	return true;
@@ -73,16 +86,16 @@ AsyncTransformResultPtr	CombineTiledBlob_Transform::Exec(const TransformArgs& Ar
 
 	T_Tiles<DeviceBufferRef> Tiles(Source->Rows(), Source->Cols());
 
-	for (int32 TileX = 0; TileX < Source->Rows(); TileX++)
+	for (int32 RowId = 0; RowId < Source->Rows(); RowId++)
 	{
-		for (int32 TileY = 0; TileY < Source->Cols(); TileY++)
+		for (int32 ColId = 0; ColId < Source->Cols(); ColId++)
 		{
-			check(Source->GetTile(TileX, TileY));
+			check(Source->GetTile(RowId, ColId));
 
-			BlobPtr Tile = Source->GetTile(TileX, TileY);
+			BlobPtr Tile = Source->GetTile(RowId, ColId);
 			check(Tile->GetBufferRef());
 
-			Tiles[TileX][TileY] = Tile->GetBufferRef();
+			Tiles[RowId][ColId] = Tile->GetBufferRef();
 		}
 	}
 
