TG Instances Blueprint crash fix

diff --git a/Engine/Plugins/Experimental/TextureGraph/Source/TextureGraph/Private/Blueprint/TG_AsyncExportTask.cpp b/Engine/Plugins/Experimental/TextureGraph/Source/TextureGraph/Private/Blueprint/TG_AsyncExportTask.cpp
--- a/Engine/Plugins/Experimental/TextureGraph/Source/TextureGraph/Private/Blueprint/TG_AsyncExportTask.cpp
+++ b/Engine/Plugins/Experimental/TextureGraph/Source/TextureGraph/Private/Blueprint/TG_AsyncExportTask.cpp
@@ -24,7 +24,7 @@ UTG_AsyncExportTask* UTG_AsyncExportTask::TG_AsyncExportTask( UTextureGraphBase*
 	if (InTextureGraph != nullptr)
 	{
 		Task->OrignalTextureGraphPtr = InTextureGraph;
-		Task->TextureGraphPtr = (UTextureGraph*)StaticDuplicateObject(Task->OrignalTextureGraphPtr, GetTransientPackage(), NAME_None, RF_Standalone, UTextureGraph::StaticClass());
+		Task->TextureGraphPtr = Cast<UTextureGraphBase>(StaticDuplicateObject(Task->OrignalTextureGraphPtr, GetTransientPackage(), NAME_None, RF_Standalone));
 		FTG_HelperFunctions::InitTargets(Task->TextureGraphPtr);
 		Task->RegisterWithTGAsyncTaskManger();
 	}
diff --git a/Engine/Plugins/Experimental/TextureGraph/Source/TextureGraph/Private/Blueprint/TG_AsyncRenderTask.cpp b/Engine/Plugins/Experimental/TextureGraph/Source/TextureGraph/Private/Blueprint/TG_AsyncRenderTask.cpp
--- a/Engine/Plugins/Experimental/TextureGraph/Source/TextureGraph/Private/Blueprint/TG_AsyncRenderTask.cpp
+++ b/Engine/Plugins/Experimental/TextureGraph/Source/TextureGraph/Private/Blueprint/TG_AsyncRenderTask.cpp
@@ -27,7 +27,7 @@ UTG_AsyncRenderTask* UTG_AsyncRenderTask::TG_AsyncRenderTask( UTextureGraphBase*
 	if (InTextureGraph != nullptr)
 	{
 		Task->OriginalTextureGraphPtr = InTextureGraph;
-		Task->TextureGraphPtr = (UTextureGraph*)StaticDuplicateObject(Task->OriginalTextureGraphPtr, GetTransientPackage(), NAME_None, RF_Standalone, UTextureGraph::StaticClass());
+		Task->TextureGraphPtr = Cast<UTextureGraphBase>(StaticDuplicateObject(Task->OriginalTextureGraphPtr, GetTransientPackage(), NAME_None, RF_Standalone));
 		FTG_HelperFunctions::InitTargets(Task->TextureGraphPtr);
 		Task->RegisterWithTGAsyncTaskManger();
 		FWorldDelegates::OnWorldCleanup.AddUObject(Task, &UTG_AsyncRenderTask::OnWorldCleanup);
