TG Instance fixes for output toggles and overrides not working with blueprints

diff --git a/Engine/Plugins/Experimental/TextureGraph/Source/TextureGraph/Private/Blueprint/TG_AsyncExportTask.cpp b/Engine/Plugins/Experimental/TextureGraph/Source/TextureGraph/Private/Blueprint/TG_AsyncExportTask.cpp
--- a/Engine/Plugins/Experimental/TextureGraph/Source/TextureGraph/Private/Blueprint/TG_AsyncExportTask.cpp
+++ b/Engine/Plugins/Experimental/TextureGraph/Source/TextureGraph/Private/Blueprint/TG_AsyncExportTask.cpp
@@ -14,17 +14,19 @@ UTG_AsyncExportTask::UTG_AsyncExportTask(const FObjectInitializer& ObjectInitial
 
 }
 
-UTG_AsyncExportTask* UTG_AsyncExportTask::TG_AsyncExportTask( UTextureGraphBase* InTextureGraph, const bool OverwriteTextures, const bool bSave)
+UTG_AsyncExportTask* UTG_AsyncExportTask::TG_AsyncExportTask( UTextureGraphBase* InTextureGraph, const bool OverwriteTextures, const bool bSave, const bool bExportAll)
 {
 	UTG_AsyncExportTask* Task = NewObject<UTG_AsyncExportTask>();
 	Task->SetFlags(RF_Standalone);
-	Task->OverwriteTextures = OverwriteTextures;
+	Task->bOverwriteTextures = OverwriteTextures;
 	Task->bSave = bSave;
+	Task->bExportAll = bExportAll;
 
 	if (InTextureGraph != nullptr)
 	{
 		Task->OrignalTextureGraphPtr = InTextureGraph;
 		Task->TextureGraphPtr = Cast<UTextureGraphBase>(StaticDuplicateObject(Task->OrignalTextureGraphPtr, GetTransientPackage(), NAME_None, RF_Standalone));
+		Task->TextureGraphPtr->Initialize();
 		FTG_HelperFunctions::InitTargets(Task->TextureGraphPtr);
 		Task->RegisterWithTGAsyncTaskManger();
 	}
@@ -47,7 +49,7 @@ void UTG_AsyncExportTask::Activate()
 	TargetExportSettings = FExportSettings();
 	TargetExportSettings.OnDone.BindUFunction(this, "OnExportDone");
 
-	FTG_HelperFunctions::ExportAsync(TextureGraphPtr, "", "", TargetExportSettings, false, OverwriteTextures, true, bSave);
+	FTG_HelperFunctions::ExportAsync(TextureGraphPtr, "", "", TargetExportSettings, false, bOverwriteTextures, bExportAll, bSave);
 }
 
 void UTG_AsyncExportTask::OnExportDone()
diff --git a/Engine/Plugins/Experimental/TextureGraph/Source/TextureGraph/Private/Blueprint/TG_AsyncRenderTask.cpp b/Engine/Plugins/Experimental/TextureGraph/Source/TextureGraph/Private/Blueprint/TG_AsyncRenderTask.cpp
--- a/Engine/Plugins/Experimental/TextureGraph/Source/TextureGraph/Private/Blueprint/TG_AsyncRenderTask.cpp
+++ b/Engine/Plugins/Experimental/TextureGraph/Source/TextureGraph/Private/Blueprint/TG_AsyncRenderTask.cpp
@@ -28,6 +28,7 @@ UTG_AsyncRenderTask* UTG_AsyncRenderTask::TG_AsyncRenderTask( UTextureGraphBase*
 	{
 		Task->OriginalTextureGraphPtr = InTextureGraph;
 		Task->TextureGraphPtr = Cast<UTextureGraphBase>(StaticDuplicateObject(Task->OriginalTextureGraphPtr, GetTransientPackage(), NAME_None, RF_Standalone));
+		Task->TextureGraphPtr->Initialize();
 		FTG_HelperFunctions::InitTargets(Task->TextureGraphPtr);
 		Task->RegisterWithTGAsyncTaskManger();
 		FWorldDelegates::OnWorldCleanup.AddUObject(Task, &UTG_AsyncRenderTask::OnWorldCleanup);
diff --git a/Engine/Plugins/Experimental/TextureGraph/Source/TextureGraph/Private/Expressions/Output/TG_Expression_Output.cpp b/Engine/Plugins/Experimental/TextureGraph/Source/TextureGraph/Private/Expressions/Output/TG_Expression_Output.cpp
--- a/Engine/Plugins/Experimental/TextureGraph/Source/TextureGraph/Private/Expressions/Output/TG_Expression_Output.cpp
+++ b/Engine/Plugins/Experimental/TextureGraph/Source/TextureGraph/Private/Expressions/Output/TG_Expression_Output.cpp
@@ -146,5 +146,8 @@ void UTG_Expression_Output::InitializeOutputSettings()
 
 void UTG_Expression_Output::SetShouldExport(bool InShouldExport)
 {
-	bShouldExport = InShouldExport;
+	OutputSettings.bShouldExport = InShouldExport;
+
+	UTG_Pin* Settings = GetParentNode()->GetPin(GET_MEMBER_NAME_CHECKED(UTG_Expression_Output, OutputSettings));
+	Settings->FromString(OutputSettings.ToString());
 }
\ No newline at end of file
diff --git a/Engine/Plugins/Experimental/TextureGraph/Source/TextureGraph/Private/TextureGraph.cpp b/Engine/Plugins/Experimental/TextureGraph/Source/TextureGraph/Private/TextureGraph.cpp
--- a/Engine/Plugins/Experimental/TextureGraph/Source/TextureGraph/Private/TextureGraph.cpp
+++ b/Engine/Plugins/Experimental/TextureGraph/Source/TextureGraph/Private/TextureGraph.cpp
@@ -52,29 +52,6 @@ void UTextureGraphInstance::CopyParamsToRuntimeGraph()
 	}
 }
 
-void UTextureGraphInstance::Update(MixUpdateCyclePtr InCycle)
-{
-	// Graph Evaluate
-	SceneTargetUpdatePtr Target = std::make_shared<MixTargetUpdate>(InCycle->GetMix(), 0);
-	InCycle->AddTarget(Target);
-	
-	T_InvalidateTiles::Create(InCycle, 0);
-
-	// Now Evaluate the Graph!
-	FTG_EvaluationContext EvaluationContext;
-	EvaluationContext.Cycle = InCycle;
-	EvaluationContext.Cycle->PushMix(this);
-	
-	if (Graph()->Validate(InCycle))
-		Graph()->Evaluate(&EvaluationContext);
-	
-	EvaluationContext.Cycle->PopMix();
-	
-	/// This will be the final result of the rendering
-	T_UpdateTargets::Create(InCycle, 0, true);
-	
-}
-
 void UTextureGraphInstance::SetParent(const TObjectPtr<UTextureGraphBase>& Parent)
 {
 	this->ParentTextureGraph = Parent;
diff --git a/Engine/Plugins/Experimental/TextureGraph/Source/TextureGraph/Public/Blueprint/TG_AsyncExportTask.h b/Engine/Plugins/Experimental/TextureGraph/Source/TextureGraph/Public/Blueprint/TG_AsyncExportTask.h
--- a/Engine/Plugins/Experimental/TextureGraph/Source/TextureGraph/Public/Blueprint/TG_AsyncExportTask.h
+++ b/Engine/Plugins/Experimental/TextureGraph/Source/TextureGraph/Public/Blueprint/TG_AsyncExportTask.h
@@ -21,7 +21,7 @@ class TEXTUREGRAPH_API UTG_AsyncExportTask : public UTG_AsyncTask
 
 public:
 	UFUNCTION(BlueprintCallable, Category = "TextureGraph" , meta=(DisplayName="Texture Graph Export (Async)", BlueprintInternalUseOnly = "true"))
-	static UTG_AsyncExportTask* TG_AsyncExportTask(UTextureGraphBase* TextureGraph, const bool OverwriteTextures, const bool bSave);
+	static UTG_AsyncExportTask* TG_AsyncExportTask(UTextureGraphBase* TextureGraph, const bool OverwriteTextures, const bool bSave, const bool bExportAll);
 
 	virtual void Activate() override;
 
@@ -35,9 +35,11 @@ private:
 	UFUNCTION()
 	void OnExportDone();
 	
-	bool OverwriteTextures = true;
+	bool bOverwriteTextures = true;
 
 	bool bSave = false;
+	
+	bool bExportAll = false;
 
 	UTextureGraphBase* OrignalTextureGraphPtr;
 
diff --git a/Engine/Plugins/Experimental/TextureGraph/Source/TextureGraph/Public/Expressions/Output/TG_Expression_Output.h b/Engine/Plugins/Experimental/TextureGraph/Source/TextureGraph/Public/Expressions/Output/TG_Expression_Output.h
--- a/Engine/Plugins/Experimental/TextureGraph/Source/TextureGraph/Public/Expressions/Output/TG_Expression_Output.h
+++ b/Engine/Plugins/Experimental/TextureGraph/Source/TextureGraph/Public/Expressions/Output/TG_Expression_Output.h
@@ -33,7 +33,7 @@ public:
 	UPROPERTY(EditAnywhere, Category = NoCategory, meta = (TGType = "TG_Input", CollapsableChildProperties,ShowOnlyInnerProperties, FullyExpand, NoResetToDefault, PinDisplayName = "Settings") )
 	FTG_OutputSettings OutputSettings;
 
-	UPROPERTY()
+	UE_DEPRECATED(5.6, "Use OutputSettings.bShouldExport instead")
 	bool bShouldExport = true;
 
 	virtual FTG_Name GetDefaultName() const override { return TEXT("Output");}
@@ -46,5 +46,5 @@ public:
 
 	void SetShouldExport(bool InShouldExport);
 
-	bool GetShouldExport() const { return bShouldExport; }
+	bool GetShouldExport() const { return OutputSettings.bShouldExport; }
 };
diff --git a/Engine/Plugins/Experimental/TextureGraph/Source/TextureGraph/Public/TG_OutputSettings.h b/Engine/Plugins/Experimental/TextureGraph/Source/TextureGraph/Public/TG_OutputSettings.h
--- a/Engine/Plugins/Experimental/TextureGraph/Source/TextureGraph/Public/TG_OutputSettings.h
+++ b/Engine/Plugins/Experimental/TextureGraph/Source/TextureGraph/Public/TG_OutputSettings.h
@@ -84,6 +84,9 @@ struct TEXTUREGRAPH_API FTG_OutputSettings
 	UPROPERTY(BlueprintReadWrite, EditAnywhere, Category = "Advanced", DisplayName = "sRGB", Meta = (NoResetToDefault, EditCondition = "TexturePresetType == ETG_TexturePresetType::None"))
 	bool bSRGB = false;
 
+	UPROPERTY()
+	bool bShouldExport = true;
+
 	FString GetFullOutputName() { return  FString::Format(TEXT("{0}"), { BaseName.ToString()});}
 	friend FArchive& operator<<(FArchive& Ar, FTG_OutputSettings& T);
 	
diff --git a/Engine/Plugins/Experimental/TextureGraph/Source/TextureGraph/Public/TextureGraph.h b/Engine/Plugins/Experimental/TextureGraph/Source/TextureGraph/Public/TextureGraph.h
--- a/Engine/Plugins/Experimental/TextureGraph/Source/TextureGraph/Public/TextureGraph.h
+++ b/Engine/Plugins/Experimental/TextureGraph/Source/TextureGraph/Public/TextureGraph.h
@@ -44,7 +44,7 @@ public:
 	
 	void InvalidateAll() override;
 	void Update(MixUpdateCyclePtr InCycle) override;
-
+	virtual void Initialize() {}; 
 
 	TEXTUREGRAPH_API void TriggerUpdate(bool Tweaking);
 	TEXTUREGRAPH_API void FlushInvalidations();
@@ -97,7 +97,6 @@ public:
 	virtual void Construct(FString Name) override;
 
 	void CopyParamsToRuntimeGraph();
-	void Update(MixUpdateCyclePtr InCycle) override;
 	bool CheckOutputSettingsMatch(const TObjectPtr<UTextureGraphBase>& Parent) const;
 	void UpdateOutputSettingsFromGraph();
 	
@@ -108,5 +107,5 @@ public:
 	
 	virtual const UTG_Graph* Graph() const override { return RuntimeGraph.Get(); }
 	virtual UTG_Graph* Graph() override;
-	void Initialize();
+	virtual void Initialize() override;
 };
\ No newline at end of file
diff --git a/Engine/Plugins/Experimental/TextureGraph/Source/TextureGraphEditor/Private/TG_Exporter.cpp b/Engine/Plugins/Experimental/TextureGraph/Source/TextureGraphEditor/Private/TG_Exporter.cpp
--- a/Engine/Plugins/Experimental/TextureGraph/Source/TextureGraphEditor/Private/TG_Exporter.cpp
+++ b/Engine/Plugins/Experimental/TextureGraph/Source/TextureGraphEditor/Private/TG_Exporter.cpp
@@ -515,7 +515,6 @@ void FTG_InstanceImpl::SetTextureGraphToExport(UTextureGraphInstance* InTextureG
 	GetPreviewSettingsView()->SetObject(TextureGraphInstance->GetSettings(), true);
 
 	UpdatePreviewMesh();
-	OnGraphChanged(TextureGraphInstance->Graph(), nullptr, false);
 	TextureGraphEngine::RegisterErrorReporter(TextureGraphInstance, std::make_shared<FTextureGraphErrorReporter>());
 
 }
diff --git a/Engine/Plugins/Experimental/TextureGraph/Source/TextureGraphEditor/Private/TG_InstanceEditor.cpp b/Engine/Plugins/Experimental/TextureGraph/Source/TextureGraphEditor/Private/TG_InstanceEditor.cpp
--- a/Engine/Plugins/Experimental/TextureGraph/Source/TextureGraphEditor/Private/TG_InstanceEditor.cpp
+++ b/Engine/Plugins/Experimental/TextureGraph/Source/TextureGraphEditor/Private/TG_InstanceEditor.cpp
@@ -107,12 +107,12 @@ FLinearColor FTG_InstanceEditor::GetWorldCentricTabColorScale() const
 
 FName FTG_InstanceEditor::GetToolkitFName() const
 {
-	return FName("TG_Editor");
+	return FName("TG_InstanceEditor");
 }
 
 FText FTG_InstanceEditor::GetBaseToolkitName() const
 {
-	return LOCTEXT("AppLabel", "TG_ Editor");
+	return LOCTEXT("AppLabel", "TG_ InstanceEditor");
 }
 
 FString FTG_InstanceEditor::GetWorldCentricTabPrefix() const
