Fix some missing/incorrect Vars to VariantArray conversions

diff --git a/Engine/Plugins/Experimental/TextureGraph/Source/TextureGraph/Private/TG_GraphEvaluation.cpp b/Engine/Plugins/Experimental/TextureGraph/Source/TextureGraph/Private/TG_GraphEvaluation.cpp
--- a/Engine/Plugins/Experimental/TextureGraph/Source/TextureGraph/Private/TG_GraphEvaluation.cpp
+++ b/Engine/Plugins/Experimental/TextureGraph/Source/TextureGraph/Private/TG_GraphEvaluation.cpp
@@ -405,17 +405,52 @@ void FTG_VariantToFVector2f_Converter_Internal(FVector2f& Output, const FTG_Vari
 	}
 }
 
-void FTG_VariantToFTG_VariantArray_Converter_Internal(FTG_VariantArray& Output, const FTG_Variant& Input, int Index)
+void FTG_VariantToFTG_VariantArray_Converter(FTG_Evaluation::VarConverterInfo& Info)
 {
-	check(Index >= 0 && Index < Output.Num());
-	Output.Set(Index, Input);
+	FTG_VariantArray& Output = Info.OutVar->EditAs<FTG_VariantArray>();
+
+	if (!Info.InVar->IsArray())
+	{
+		FTG_Variant& Input = Info.InVar->GetAs<FTG_Variant>();
+		int32 Index = 0;
+		if (Info.Index >= 0)
+		{
+			Index = Info.Index;
+		}
+
+		if (Output.Num() <= Index)
+			Output.SetNum(Index + 1);
+
+		Output.Set(Index, Input);
+	}
+	else
+	{
+		Output = Info.InVar->GetAs<FTG_VariantArray>();
+	}
 }
 
-void FTG_VariantToFTG_VariantArray_Converter(FTG_Evaluation::VarConverterInfo& Info)
+void FTG_TextureToFTG_VariantArray_Converter(FTG_Evaluation::VarConverterInfo& Info)
 {
-	FTG_Variant& Input = Info.InVar->GetAs<FTG_Variant>();
 	FTG_VariantArray& Output = Info.OutVar->EditAs<FTG_VariantArray>();
-	FTG_VariantToFTG_VariantArray_Converter_Internal(Output, Input, Info.Index);
+
+	if (!Info.InVar->IsArray())
+	{
+		FTG_Texture InTexture = Info.InVar->GetAs<FTG_Texture>();
+		int32 Index = 0;
+		if (Info.Index >= 0)
+		{
+			Index = Info.Index;
+		}
+
+		if (Output.Num() <= Index)
+			Output.SetNum(Index + 1);
+
+		Output.Set(Index, InTexture);
+	}
+	else
+	{
+		Output = Info.InVar->GetAs<FTG_VariantArray>();
+	}
 }
 
 void FTG_VariantToFVector2f_Converter(FTG_Evaluation::VarConverterInfo& Info)
@@ -534,6 +569,7 @@ FTG_Evaluation::ConverterMap FTG_Evaluation::DefaultConverters
 		VAR_CONVERTER(FTG_Variant.Color, FTG_VariantArray, FTG_VariantToFTG_VariantArray_Converter),
 		VAR_CONVERTER(FTG_Variant.Vector, FTG_VariantArray, FTG_VariantToFTG_VariantArray_Converter),
 		VAR_CONVERTER(FTG_Variant.Texture, FTG_VariantArray, FTG_VariantToFTG_VariantArray_Converter),
+		VAR_CONVERTER(FTG_Texture, FTG_VariantArray, FTG_TextureToFTG_VariantArray_Converter),
 
 		VAR_CONVERTER_NULL(FTG_Variant.Scalar, FTG_Variant),
 		VAR_CONVERTER_NULL(FTG_Variant.Color, FTG_Variant),
diff --git a/Engine/Plugins/Experimental/TextureGraph/Source/TextureGraphEngine/Transform/Expressions/T_ArrayGrid.cpp b/Engine/Plugins/Experimental/TextureGraph/Source/TextureGraphEngine/Transform/Expressions/T_ArrayGrid.cpp
--- a/Engine/Plugins/Experimental/TextureGraph/Source/TextureGraphEngine/Transform/Expressions/T_ArrayGrid.cpp
+++ b/Engine/Plugins/Experimental/TextureGraph/Source/TextureGraphEngine/Transform/Expressions/T_ArrayGrid.cpp
@@ -84,6 +84,8 @@ public:
 
 						if (Tile)
 							bIsWellDefined &= Tile->IsWellDefined();
+						else
+							bIsWellDefined = false;
 					}
 				}
 			}
