Non-async blueprint support through delegate callback

diff --git a/Engine/Plugins/Experimental/TextureGraph/Source/TextureGraph/Private/Blueprint/TG_AsyncExportTask.cpp b/Engine/Plugins/Experimental/TextureGraph/Source/TextureGraph/Private/Blueprint/TG_AsyncExportTask.cpp
--- a/Engine/Plugins/Experimental/TextureGraph/Source/TextureGraph/Private/Blueprint/TG_AsyncExportTask.cpp
+++ b/Engine/Plugins/Experimental/TextureGraph/Source/TextureGraph/Private/Blueprint/TG_AsyncExportTask.cpp
@@ -57,6 +57,7 @@ void UTG_AsyncExportTask::OnExportDone()
 	TargetExportSettings.ExportPreset.clear();
 	TargetExportSettings.OnDone.Unbind();
 	OnDone.Broadcast();
+	OnStaticDone.ExecuteIfBound(0);
 	TextureGraphPtr->FlushInvalidations();
 	ClearFlags(RF_Standalone);
 	SetReadyToDestroy();
diff --git a/Engine/Plugins/Experimental/TextureGraph/Source/TextureGraph/Private/Blueprint/TG_AsyncRenderTask.cpp b/Engine/Plugins/Experimental/TextureGraph/Source/TextureGraph/Private/Blueprint/TG_AsyncRenderTask.cpp
--- a/Engine/Plugins/Experimental/TextureGraph/Source/TextureGraph/Private/Blueprint/TG_AsyncRenderTask.cpp
+++ b/Engine/Plugins/Experimental/TextureGraph/Source/TextureGraph/Private/Blueprint/TG_AsyncRenderTask.cpp
@@ -78,6 +78,7 @@ void UTG_AsyncRenderTask::Activate()
 	.then([this](bool bRtResult)
 	{
 		OnDone.Broadcast(OutputRts);
+		OnStaticDone.ExecuteIfBound(OutputRts);
 		bRenderComplete = true;
 
 		UE_LOG(LogTextureGraph, Log, TEXT("UTG_AsyncRenderTask:: OnDone : bShouldDestroyOnRenderComplete %i"), bShouldDestroyOnRenderComplete ? 1 : 0);
diff --git a/Engine/Plugins/Experimental/TextureGraph/Source/TextureGraph/Private/Blueprint/TG_BlueprintFunctionLibrary.cpp b/Engine/Plugins/Experimental/TextureGraph/Source/TextureGraph/Private/Blueprint/TG_BlueprintFunctionLibrary.cpp
--- a/Engine/Plugins/Experimental/TextureGraph/Source/TextureGraph/Private/Blueprint/TG_BlueprintFunctionLibrary.cpp
+++ b/Engine/Plugins/Experimental/TextureGraph/Source/TextureGraph/Private/Blueprint/TG_BlueprintFunctionLibrary.cpp
@@ -267,40 +267,19 @@ FTG_OutputSettings UTG_BlueprintFunctionLibrary::GetSettingsParameterValue(UObje
 	return ParameterValue;
 }
 
-#if 0 /// This requires a lot of additional work. Will attempt after summer break
-TArray<UTextureRenderTarget2D*> UTG_BlueprintFunctionLibrary::RenderTextureGraph(UObject* WorldContextObject, UTextureGraph* InTextureGraph)
+void UTG_BlueprintFunctionLibrary::RenderTextureGraph(UObject* WorldContextObject, UTextureGraphBase* InTextureGraph, FOnTextureGraphRenderDone CompletionDelegate)
 {
-	check(IsInGameThread());
-	TArray<UTextureRenderTarget2D*> Results;
-	UTextureGraph* TextureGraphPtr = (UTextureGraph*)StaticDuplicateObject(InTextureGraph, GetTransientPackage(), NAME_None, ~RF_Standalone, UTextureGraph::StaticClass());
-	FTG_HelperFunctions::InitTargets(TextureGraphPtr);
-
-	JobBatchPtr Batch = FTG_HelperFunctions::InitRenderBatch(TextureGraphPtr);
-	if (!Batch)
-		return Results;
-
-	/// Add to the scheduler
-	TextureGraphEngine::GetScheduler()->AddBatch(Batch);
-
-	std::atomic_bool* bIsMixRendered = new std::atomic_bool(false);
-	Batch->OnDone([=, &bIsMixRendered](JobBatch*) mutable
-	{
-		*bIsMixRendered = true;
-	});
-
-	while (!*bIsMixRendered)
-	{
-		if (!TextureGraphEngine::IsDestroying())
-			TextureGraphEngine::Update(0);
-	}
-
-	delete bIsMixRendered;
+	UTG_AsyncRenderTask* Task = UTG_AsyncRenderTask::TG_AsyncRenderTask(InTextureGraph);
+	Task->OnStaticDone = CompletionDelegate;
+	Task->Activate();
+}
 
-	//FTG_HelperFunctions::RenderAsync(TextureGraphPtr);
-	//TextureGraphEngine::GetDeviceManager()->WaitForQueuedTasks(ENamedThreads::GameThread);
-	return Results;
+void UTG_BlueprintFunctionLibrary::ExportTextureGraph(UObject* WorldContextObject, UTextureGraphBase* InTextureGraph, FOnTextureGraphExportDone CompletionDelegate)
+{
+	UTG_AsyncExportTask* Task = UTG_AsyncExportTask::TG_AsyncExportTask(InTextureGraph, true, true, true);
+	Task->OnStaticDone = CompletionDelegate;
+	Task->Activate();
 }
-#endif 
 
 void UTG_BlueprintFunctionLibrary::AddParamWarning(FName ParamName, UObject* ObjectPtr, FString FunctionName)
 {
diff --git a/Engine/Plugins/Experimental/TextureGraph/Source/TextureGraph/Public/Blueprint/TG_AsyncExportTask.h b/Engine/Plugins/Experimental/TextureGraph/Source/TextureGraph/Public/Blueprint/TG_AsyncExportTask.h
--- a/Engine/Plugins/Experimental/TextureGraph/Source/TextureGraph/Public/Blueprint/TG_AsyncExportTask.h
+++ b/Engine/Plugins/Experimental/TextureGraph/Source/TextureGraph/Public/Blueprint/TG_AsyncExportTask.h
@@ -13,6 +13,7 @@
 #include "TG_AsyncExportTask.generated.h"
 
 DECLARE_DYNAMIC_MULTICAST_DELEGATE(FTSExportTaskDelegate);
+DECLARE_DYNAMIC_DELEGATE_OneParam(FOnTextureGraphExportDone, int, Result);
 
 UCLASS()
 class TEXTUREGRAPH_API UTG_AsyncExportTask : public UTG_AsyncTask
@@ -29,6 +30,7 @@ public:
 
 	UPROPERTY(BlueprintAssignable, Category = "TextureGraph")
 	FTSExportTaskDelegate OnDone;
+	FOnTextureGraphExportDone OnStaticDone;
 
 private:
 
diff --git a/Engine/Plugins/Experimental/TextureGraph/Source/TextureGraph/Public/Blueprint/TG_AsyncRenderTask.h b/Engine/Plugins/Experimental/TextureGraph/Source/TextureGraph/Public/Blueprint/TG_AsyncRenderTask.h
--- a/Engine/Plugins/Experimental/TextureGraph/Source/TextureGraph/Public/Blueprint/TG_AsyncRenderTask.h
+++ b/Engine/Plugins/Experimental/TextureGraph/Source/TextureGraph/Public/Blueprint/TG_AsyncRenderTask.h
@@ -15,6 +15,7 @@
 #include "TG_AsyncRenderTask.generated.h"
 
 DECLARE_DYNAMIC_MULTICAST_DELEGATE_OneParam(FTSRenderTaskDelegate, const TArray<UTextureRenderTarget2D*>&, OutputRts);
+DECLARE_DYNAMIC_DELEGATE_OneParam(FOnTextureGraphRenderDone, const TArray<UTextureRenderTarget2D*>&, OutputRts);
 
 UCLASS()
 class TEXTUREGRAPH_API UTG_AsyncRenderTask : public UTG_AsyncTask
@@ -33,6 +34,7 @@ public:
 
 	UPROPERTY(BlueprintAssignable, Category = "TextureGraph")
 	FTSRenderTaskDelegate OnDone;
+	FOnTextureGraphRenderDone OnStaticDone;
 
 private:
 	AsyncBool FinalizeAllOutputBlobs();
diff --git a/Engine/Plugins/Experimental/TextureGraph/Source/TextureGraph/Public/Blueprint/TG_BlueprintFunctionLibrary.h b/Engine/Plugins/Experimental/TextureGraph/Source/TextureGraph/Public/Blueprint/TG_BlueprintFunctionLibrary.h
--- a/Engine/Plugins/Experimental/TextureGraph/Source/TextureGraph/Public/Blueprint/TG_BlueprintFunctionLibrary.h
+++ b/Engine/Plugins/Experimental/TextureGraph/Source/TextureGraph/Public/Blueprint/TG_BlueprintFunctionLibrary.h
@@ -8,8 +8,13 @@
 #include "Kismet/BlueprintFunctionLibrary.h"
 #include "TextureGraph.h"
 #include "TG_OutputSettings.h"
+#include "TG_AsyncRenderTask.h"
+#include "TG_AsyncExportTask.h"
 #include "TG_BlueprintFunctionLibrary.generated.h"
 
+//DECLARE_DYNAMIC_DELEGATE_OneParam(FOnTextureGraphRenderDone, UTextureRenderTarget2D*, Results);
+//DECLARE_DELEGATE_TwoParams(FOnGroomBindingAssetBuildCompleteNative, UGroomBindingAsset*, EGroomBindingAssetBuildResult);
+
 UCLASS( meta=(ScriptName="TextureScreptingLibrary"))
 class UTG_BlueprintFunctionLibrary : public UBlueprintFunctionLibrary
 {
@@ -76,10 +81,12 @@ public:
 	static FTG_OutputSettings GetSettingsParameterValue(UObject* WorldContextObject, UTextureGraph* TextureGraph, FName ParameterName , int& Width, int& Height);
 
 	/** Render the texture graph and return an array of texture render targets. */
-#if 0 /// This requires a lot of additional work. Will attempt after summer break
-	UFUNCTION(BlueprintCallable, Category = "TextureGraph", meta = (Keywords = "SetColorParameterValue", WorldContext = "WorldContextObject"))
-	static TArray<UTextureRenderTarget2D*> RenderTextureGraph(UObject* WorldContextObject, UTextureGraph* InTextureGraph);
-#endif 
+	UFUNCTION(BlueprintCallable, Category = "TextureGraph", meta = (Keywords = "RenderTextureGraph", WorldContext = "WorldContextObject"))
+	static void RenderTextureGraph(UObject* WorldContextObject, UTextureGraphBase* InTextureGraph, FOnTextureGraphRenderDone CompletionDelegate);
+
+	/** Export texture graph. */
+	UFUNCTION(BlueprintCallable, Category = "TextureGraph", meta = (Keywords = "ExportTextureGraph", WorldContext = "WorldContextObject"))
+	static void ExportTextureGraph(UObject* WorldContextObject, UTextureGraphBase* InTextureGraph, FOnTextureGraphExportDone CompletionDelegate);
 
 	static void AddParamWarning(FName ParamName, UObject* ObjectPtr, FString FunctionName);
 	static void AddError(UObject* ObjectPtr, FString FunctionName, FString Error);
diff --git a/Engine/Plugins/Experimental/TextureGraph/Source/TextureGraphEngine/Transform/Utility/T_LoadStaticResource.cpp b/Engine/Plugins/Experimental/TextureGraph/Source/TextureGraphEngine/Transform/Utility/T_LoadStaticResource.cpp
--- a/Engine/Plugins/Experimental/TextureGraph/Source/TextureGraphEngine/Transform/Utility/T_LoadStaticResource.cpp
+++ b/Engine/Plugins/Experimental/TextureGraph/Source/TextureGraphEngine/Transform/Utility/T_LoadStaticResource.cpp
@@ -16,6 +16,19 @@ Job_LoadStaticImageResource::Job_LoadStaticImageResource(UMixInterface* InMix, U
 	Name = FString::Printf(TEXT("T_LoadStaticImageResource_%s"), *InSource->GetAssetUUID());
 }
 
+FString Job_LoadStaticImageResource::GetReferencerName() const
+{
+	if (Source)
+		return Source->GetAssetUUID();
+	return TEXT("Job_LoadStaticImageResource");
+}
+
+void Job_LoadStaticImageResource::AddReferencedObjects(FReferenceCollector& Collector)
+{
+	if (Source)
+		Collector.AddReferencedObject(Source);
+}
+
 cti::continuable<int32> Job_LoadStaticImageResource::PreExecAsync(ENamedThreads::Type execThread, ENamedThreads::Type returnThread)
 {
 	auto JobHash = Job::Hash();
diff --git a/Engine/Plugins/Experimental/TextureGraph/Source/TextureGraphEngine/Transform/Utility/T_LoadStaticResource.h b/Engine/Plugins/Experimental/TextureGraph/Source/TextureGraphEngine/Transform/Utility/T_LoadStaticResource.h
--- a/Engine/Plugins/Experimental/TextureGraph/Source/TextureGraphEngine/Transform/Utility/T_LoadStaticResource.h
+++ b/Engine/Plugins/Experimental/TextureGraph/Source/TextureGraphEngine/Transform/Utility/T_LoadStaticResource.h
@@ -7,14 +7,20 @@
 
 class UStaticImageResource;
 
-class TEXTUREGRAPHENGINE_API Job_LoadStaticImageResource : public Job
+class TEXTUREGRAPHENGINE_API Job_LoadStaticImageResource : public Job, public FGCObject
 {
-	UStaticImageResource*			Source;			/// The source for the static image resource
+	TObjectPtr<UStaticImageResource> Source;			/// The source for the static image resource
 
 public:
 									Job_LoadStaticImageResource(UMixInterface* InMix, UStaticImageResource* InSource, int32 TargetId, UObject* InErrorOwner = nullptr, uint16 Priority = (uint16)E_Priority::kHigh);
 	virtual							~Job_LoadStaticImageResource() override = default;
 
+	//////////////////////////////////////////////////////////////////////////
+	/// FGCObject
+	//////////////////////////////////////////////////////////////////////////
+	virtual void					AddReferencedObjects(FReferenceCollector& Collector) override;
+	virtual FString					GetReferencerName() const override;
+
 protected:
 	virtual void					GetDependencies(JobPtrVec& prior, JobPtrVec& after, JobRunInfo runInfo) override { RunInfo = runInfo; Prev.clear(); }
 	virtual AsyncPrepareResult		PrepareTargets(JobBatch* batch) override { return cti::make_ready_continuable(0); }
