TextureGraphInstance crash fix for OutputSettings

diff --git a/Engine/Plugins/Experimental/TextureGraph/Source/TextureGraph/Private/TG_Var.cpp b/Engine/Plugins/Experimental/TextureGraph/Source/TextureGraph/Private/TG_Var.cpp
--- a/Engine/Plugins/Experimental/TextureGraph/Source/TextureGraph/Private/TG_Var.cpp
+++ b/Engine/Plugins/Experimental/TextureGraph/Source/TextureGraph/Private/TG_Var.cpp
@@ -740,7 +740,7 @@ void VarArchiveSerializer_UTexture(FTG_Var::VarArchiveSerialInfo& Info)
 }
 void VarArchiveSerializer_FTG_OutputSettings(FTG_Var::VarArchiveSerialInfo& Info)
 {
-	Generic_Simple_ArSerializer<TObjectPtr<UTexture>>(Info);
+	Generic_Simple_ArSerializer<FTG_OutputSettings>(Info);
 }
 //////////////////////////////////////////////////////////////////////////////////////////////////////////////////////////////////////
 FTG_Var::VarArchiveSerializerMap FTG_Var::DefaultArchiveSerializers
diff --git a/Engine/Plugins/Experimental/TextureGraph/Source/TextureGraphEditor/Private/Customizations/TG_OutputSettingsCustomization.h b/Engine/Plugins/Experimental/TextureGraph/Source/TextureGraphEditor/Private/Customizations/TG_OutputSettingsCustomization.h
--- a/Engine/Plugins/Experimental/TextureGraph/Source/TextureGraphEditor/Private/Customizations/TG_OutputSettingsCustomization.h
+++ b/Engine/Plugins/Experimental/TextureGraph/Source/TextureGraphEditor/Private/Customizations/TG_OutputSettingsCustomization.h
@@ -248,6 +248,13 @@ public:
 		ExternalObjects.Empty();
 		ExternalObjects.Add(OutputExpression);
 		
+		TObjectPtr<UTG_Pin> TGPin = OutputExpression->GetParentNode()->GetPin(GET_MEMBER_NAME_CHECKED(UTG_Expression_Output, OutputSettings));
+		FProperty* Property = TGPin->GetExpressionProperty();
+		bool bCanEditChange = OutputExpression->CanEditChange(Property);
+		FName NAME_EditCondition(TEXT("EditCondition"));
+		Property->SetMetaData(NAME_EditCondition, bCanEditChange ? TEXT("true") : TEXT("false"));
+
+		
 		IDetailPropertyRow* Row = ChildBuilder.AddExternalObjectProperty(ExternalObjects, GET_MEMBER_NAME_CHECKED(UTG_Expression_Output, OutputSettings), FAddPropertyParams().HideRootObjectNode(true).CreateCategoryNodes(false));
 		Row->DisplayName(FText::FromName(OutputExpression->GetTitleName()));
 		
diff --git a/Engine/Plugins/Experimental/TextureGraph/Source/TextureGraphEditor/Private/EdGraph/TG_EdGraphNode.cpp b/Engine/Plugins/Experimental/TextureGraph/Source/TextureGraphEditor/Private/EdGraph/TG_EdGraphNode.cpp
--- a/Engine/Plugins/Experimental/TextureGraph/Source/TextureGraphEditor/Private/EdGraph/TG_EdGraphNode.cpp
+++ b/Engine/Plugins/Experimental/TextureGraph/Source/TextureGraphEditor/Private/EdGraph/TG_EdGraphNode.cpp
@@ -515,6 +515,31 @@ void UTG_EdGraphNode::UpdateInputPinsVisibility() const
 		}
 	}
 }
+
+void UTG_EdGraphNode::UpdatePinsEditCondition() const
+{
+	auto Expression = Node->GetExpression();
+	// go through all pins and match them with their property hide states
+	for (FTG_Id OtherPinId : Node->GetInputPinIds())
+	{
+		TObjectPtr<UTG_Pin> TGPin = Node->GetGraph()->GetPin(OtherPinId);
+		
+		// check hide state
+		if (TGPin)
+		{
+			// Rely on the FProperty meta information for potential conditional editability / visibility
+			FProperty* Property = TGPin->GetExpressionProperty();
+			if (Property)
+			{
+				bool bCanEditChange = Expression->CanEditChange(Property);
+				
+				// update Metadata with EditCondition to also update details view
+				FName NAME_EditCondition(TEXT("EditCondition"));
+				Property->SetMetaData(NAME_EditCondition, bCanEditChange ? TEXT("true") : TEXT("false"));
+			}
+		}
+	}
+}
 #endif
 
 void UTG_EdGraphNode::PinDefaultValueChangedWithTweaking(UEdGraphPin* Pin, bool bIsTweaking)
diff --git a/Engine/Plugins/Experimental/TextureGraph/Source/TextureGraphEditor/Private/TG_Editor.cpp b/Engine/Plugins/Experimental/TextureGraph/Source/TextureGraphEditor/Private/TG_Editor.cpp
--- a/Engine/Plugins/Experimental/TextureGraph/Source/TextureGraphEditor/Private/TG_Editor.cpp
+++ b/Engine/Plugins/Experimental/TextureGraph/Source/TextureGraphEditor/Private/TG_Editor.cpp
@@ -1205,16 +1205,9 @@ void FTG_Editor::OnSelectedNodesChanged(const TSet<class UObject*>& NewSelection
 		{
 			SelectedObjects.Add(GraphNode->GetDetailsObject());
 
-			// Temp fix - Commenting this code for now.
-			// Calling UpdateInputPinsVisibility on selection is causing disconnection of the pins.
-			// Disabling this code will cause UE-215729 bug.
-			// But that only effects details panel so for now this is fine.
-			// Problem: Visibility and interaction of the pin are being updated from this code which are dependent on each other
-			// Ideally they shouldnt be dependent on each other.
-			
-			// calling this here to internally set EditCondition metadata properly for Details panel
+			// calling this here to set EditCondition metadata properly for Details panel
 			// this is driven from the CanEditChange() method of the expression which can be customized as required
-			// GraphNode->UpdateInputPinsVisibility();
+			GraphNode->UpdatePinsEditCondition();
 		}  
 	}
 	
diff --git a/Engine/Plugins/Experimental/TextureGraph/Source/TextureGraphEditor/Public/EdGraph/TG_EdGraphNode.h b/Engine/Plugins/Experimental/TextureGraph/Source/TextureGraphEditor/Public/EdGraph/TG_EdGraphNode.h
--- a/Engine/Plugins/Experimental/TextureGraph/Source/TextureGraphEditor/Public/EdGraph/TG_EdGraphNode.h
+++ b/Engine/Plugins/Experimental/TextureGraph/Source/TextureGraphEditor/Public/EdGraph/TG_EdGraphNode.h
@@ -66,6 +66,7 @@ public:
 #if WITH_EDITOR
 	void UpdatePinVisibility(UEdGraphPin* Pin, UTG_Pin* TGPin) const;
 	void UpdateInputPinsVisibility() const;
+	void UpdatePinsEditCondition() const;
 #endif
 	
 	/** Called when the DefaultValue of one of the pins of this node is changed in the editor */
