Fixed issue with output settings not getting saved.

diff --git a/Engine/Plugins/Experimental/TextureGraph/Source/TextureGraph/Private/Expressions/Arrays/TG_Expression_ArrayGrid.cpp b/Engine/Plugins/Experimental/TextureGraph/Source/TextureGraph/Private/Expressions/Arrays/TG_Expression_ArrayGrid.cpp
--- a/Engine/Plugins/Experimental/TextureGraph/Source/TextureGraph/Private/Expressions/Arrays/TG_Expression_ArrayGrid.cpp
+++ b/Engine/Plugins/Experimental/TextureGraph/Source/TextureGraph/Private/Expressions/Arrays/TG_Expression_ArrayGrid.cpp
@@ -62,14 +62,18 @@ void UTG_Expression_ArrayGrid::Evaluate(FTG_EvaluationContext* InContext)
 
 	if (OutRows * OutCols > Inputs.Num())
 	{
-		Output = FTG_Texture::GetBlack();
 		FTextureGraphErrorReporter* ErrorReporter = TextureGraphEngine::GetErrorReporter(InContext->Cycle->GetMix());
 		if (ErrorReporter)
 		{
-			ErrorReporter->ReportWarning((int32)ETextureGraphErrorType::NODE_WARNING, FString::Printf(TEXT("Number of rows and columns for the tiling do not match the number of inputs. Grid: %d x %d [Max Inputs: %d]"), 
+			ErrorReporter->ReportWarning((int32)ETextureGraphErrorType::NODE_WARNING, FString::Printf(TEXT("Number of rows and columns for the tiling do not match the number of inputs. Grid: %d x %d [Max Inputs: %d]. Padding with black images."), 
 				OutRows, OutCols, Inputs.Num()), GetParentNode());
 		}
-		return;
+
+		int32 NumDesiredEntires = OutRows * OutCols;
+		for (int32 Index = Inputs.Num(); Index < NumDesiredEntires; Index++)
+		{
+			Inputs.Add(TextureHelper::GetBlack());
+		}
 	}
 
 	BufferDescriptor DesiredDesc = Output.Descriptor;
diff --git a/Engine/Plugins/Experimental/TextureGraph/Source/TextureGraph/Private/TG_GraphEvaluation.cpp b/Engine/Plugins/Experimental/TextureGraph/Source/TextureGraph/Private/TG_GraphEvaluation.cpp
--- a/Engine/Plugins/Experimental/TextureGraph/Source/TextureGraph/Private/TG_GraphEvaluation.cpp
+++ b/Engine/Plugins/Experimental/TextureGraph/Source/TextureGraph/Private/TG_GraphEvaluation.cpp
@@ -560,6 +560,7 @@ FTG_Evaluation::ConverterMap FTG_Evaluation::DefaultConverters
 		VAR_CONVERTER(FTG_Variant.Color, FTG_Texture, FTG_VariantToFTG_Texture_Converter),
 		VAR_CONVERTER(FTG_Variant.Vector, FTG_Texture, FTG_VariantToFTG_Texture_Converter),
 		VAR_CONVERTER(FTG_Variant.Texture, FTG_Texture, FTG_VariantToFTG_Texture_Converter),
+		VAR_CONVERTER(FTG_Variant, FTG_Texture, FTG_VariantToFTG_Texture_Converter),
 
 		VAR_CONVERTER(FTG_Variant.Scalar, FVector2f, FTG_VariantToFVector2f_Converter),
 		VAR_CONVERTER(FTG_Variant.Color, FVector2f, FTG_VariantToFVector2f_Converter),
@@ -664,19 +665,50 @@ void FTG_Evaluation::EvaluateNodeArray(UTG_Node* InNode, const TArray<TObjectPtr
 	check(!ArrayInputs.IsEmpty());
 	check(MaxCount > 0);
 
-	/// Non array pins only need to have transfer once
+	/// Initialize an array for each of the output pins/vars
+	TMap<FName, FTG_VariantArray> ArrayOutputs;
+
+	/// Non array pins only need to have transfer once. We use the same loop to construct our
+	/// array pins and copy over the value of the pins that were evaluated as arrays in prevous
+	/// update cycles. This way we retain the texture descriptors
 	for (UTG_Pin* Pin : NonArrayPins)
 	{
+		FTG_Argument Arg = Pin->GetArgument();
+		const auto VarId = Pin->GetVarId();
+		FTG_Var* Var = InContext->Graph->GetVar(VarId);
+
+		if (Arg.IsOutput())
+		{
+			FTG_VariantArray ArrayOutput;
+			if (Var->IsArray())
+				ArrayOutput = Var->GetAs<FTG_VariantArray>();
+			ArrayOutput.SetNum(MaxCount);
+			ArrayOutputs.Add(Arg.GetName(), std::move(ArrayOutput));
+		}
+
 		FTG_Evaluation::TransferVarToPin(Pin, InContext, -1);
 	}
 
-	/// Initialize an array for each of the output pins/vars
-	TMap<FName, FTG_VariantArray> ArrayOutputs;
+	/// Transfer the descriptors of the outputs to the output vars ... one index at a time (only if they've been modified)
 	for (auto& OutputVar : InContext->Outputs.VarArguments)
 	{
-		FTG_VariantArray ArrayOutput;
-		ArrayOutput.SetNum(MaxCount);
-		ArrayOutputs.Add(OutputVar.Key, std::move(ArrayOutput));
+		const FTG_VariantArray* ArrayOutput = ArrayOutputs.Find(OutputVar.Key);
+		if (!ArrayOutput->GetArray().IsEmpty())
+		{
+			const FTG_Variant& ArrayValue = ArrayOutput->Get(0);
+
+			/// We only need to do this for textures
+			if (ArrayValue.IsTexture() && !ArrayValue.GetTexture().Descriptor.IsDefault())
+			{
+				const FTG_Texture& ArrayTexture = ArrayValue.GetTexture();
+				FString ToType = OutputVar.Value.Argument.GetCPPTypeName().ToString();
+
+				if (ToType.Find("FTG_Variant") >= 0)
+					OutputVar.Value.Var->EditAs<FTG_Variant>().EditTexture().Descriptor = ArrayTexture.Descriptor;
+				else if (ToType == TEXT("FTG_Texture"))
+					OutputVar.Value.Var->EditAs<FTG_Texture>().Descriptor = ArrayTexture.Descriptor;
+			}
+		}
 	}
 
 	for (int VarIndex = 0; VarIndex < MaxCount; VarIndex++)
@@ -706,6 +738,7 @@ void FTG_Evaluation::EvaluateNodeArray(UTG_Node* InNode, const TArray<TObjectPtr
 			FTG_Evaluation::TransferVarToPin(ArrayPin, InContext, VarArrayIndex);
 		}
 
+		/// Now we evaluate the expression
 		UTG_Expression* Expression = InNode->GetExpression();
 		/// Trigger the Evaluation
 		Expression->SetupAndEvaluate(InContext);
@@ -715,7 +748,6 @@ void FTG_Evaluation::EvaluateNodeArray(UTG_Node* InNode, const TArray<TObjectPtr
 			FTG_VariantArray* ArrayOutput = ArrayOutputs.Find(OutputVar.Key);
 			check(ArrayOutput && ArrayOutput->Num() > VarIndex);
 
-			FName ConverterKey = FName();
 			FString FromType = OutputVar.Value.Argument.GetCPPTypeName().ToString();
 			FString ToType = TEXT("FTG_Variant");
 
diff --git a/Engine/Plugins/Experimental/TextureGraph/Source/TextureGraph/Private/TG_Pin.cpp b/Engine/Plugins/Experimental/TextureGraph/Source/TextureGraph/Private/TG_Pin.cpp
--- a/Engine/Plugins/Experimental/TextureGraph/Source/TextureGraph/Private/TG_Pin.cpp
+++ b/Engine/Plugins/Experimental/TextureGraph/Source/TextureGraph/Private/TG_Pin.cpp
@@ -196,7 +196,8 @@ void UTG_Pin::NotifyPinSelfVarChanged(bool bIsTweaking /*= false*/)
 	// then we also transfer the new value to the Expression's own matching value.
 	if ((IsInput() && !IsConnected()) || (IsOutput()))
 	{
-		EditSelfVar()->CopyTo(GetNodePtr()->GetExpression(), Argument);
+		if (!IsArgTextureArray())
+			EditSelfVar()->CopyTo(GetNodePtr()->GetExpression(), Argument);
 	}
 
 #if WITH_EDITOR
@@ -590,6 +591,14 @@ bool UTG_Pin::SetValue(const FTG_TextureDescriptor& Value)
 		{
 			EditSelfVar()->EditAs<FTG_Variant>().EditTexture().Descriptor = Value;
 		}
+		else if (IsArgTextureArray())
+		{
+			FTG_VariantArray& VarArray = EditSelfVar()->EditAs<FTG_VariantArray>();
+			for (int Index = 0; Index < VarArray.Num(); Index++)
+			{
+				VarArray.Get(Index).GetTexture().Descriptor = Value;
+			}
+		}
 		else
 		{
 			EditSelfVar()->EditAs<FTG_Texture>().Descriptor = Value;
diff --git a/Engine/Plugins/Experimental/TextureGraph/Source/TextureGraph/Private/TG_Variant.cpp b/Engine/Plugins/Experimental/TextureGraph/Source/TextureGraph/Private/TG_Variant.cpp
--- a/Engine/Plugins/Experimental/TextureGraph/Source/TextureGraph/Private/TG_Variant.cpp
+++ b/Engine/Plugins/Experimental/TextureGraph/Source/TextureGraph/Private/TG_Variant.cpp
@@ -217,6 +217,13 @@ const FTG_Texture& FTG_Variant::GetTexture() const
 	check(IsTexture());
 	return Data.Get<FTG_Texture>();
 }
+
+FTG_Texture& FTG_Variant::GetTexture() 
+{
+	check(IsTexture());
+	return Data.Get<FTG_Texture>();
+}
+
 FTG_Texture& FTG_Variant::EditTexture()
 {
 	ResetTypeAs(EType::Texture);
diff --git a/Engine/Plugins/Experimental/TextureGraph/Source/TextureGraph/Public/TG_Pin.h b/Engine/Plugins/Experimental/TextureGraph/Source/TextureGraph/Public/TG_Pin.h
--- a/Engine/Plugins/Experimental/TextureGraph/Source/TextureGraph/Public/TG_Pin.h
+++ b/Engine/Plugins/Experimental/TextureGraph/Source/TextureGraph/Public/TG_Pin.h
@@ -171,6 +171,7 @@ public:
 	bool IsArgColor() const { return Argument.IsColor(); }
 	bool IsArgVector() const { return Argument.IsVector(); }
 	bool IsArgArray() const { return Argument.IsArray() || (GetSelfVar() && GetSelfVar()->IsArray()); }
+	bool IsArgTextureArray() const { return IsArgArray() && GetSelfVar() && GetSelfVar()->GetAs<FTG_VariantArray>().IsTexture(); }
 	bool IsArgTexture() const { return Argument.IsTexture(); }
 	bool IsArgVariant() const { return Argument.IsVariant(); }
 	bool IsArgObject(UClass* InClass) const { return Argument.IsObject(InClass); }
diff --git a/Engine/Plugins/Experimental/TextureGraph/Source/TextureGraph/Public/TG_Texture.h b/Engine/Plugins/Experimental/TextureGraph/Source/TextureGraph/Public/TG_Texture.h
--- a/Engine/Plugins/Experimental/TextureGraph/Source/TextureGraph/Public/TG_Texture.h
+++ b/Engine/Plugins/Experimental/TextureGraph/Source/TextureGraph/Public/TG_Texture.h
@@ -67,6 +67,10 @@ PRAGMA_ENABLE_DEPRECATION_WARNINGS
 	}
 
 	FORCEINLINE operator BufferDescriptor() const { return ToBufferDescriptor(); }
+	FORCEINLINE bool IsDefault() const
+	{
+		return Width == EResolution::Auto && Height == EResolution::Auto && TextureFormat == ETG_TextureFormat::Auto;
+	}
 };
 
 USTRUCT()
diff --git a/Engine/Plugins/Experimental/TextureGraph/Source/TextureGraph/Public/TG_Variant.h b/Engine/Plugins/Experimental/TextureGraph/Source/TextureGraph/Public/TG_Variant.h
--- a/Engine/Plugins/Experimental/TextureGraph/Source/TextureGraph/Public/TG_Variant.h
+++ b/Engine/Plugins/Experimental/TextureGraph/Source/TextureGraph/Public/TG_Variant.h
@@ -156,6 +156,7 @@ public:
 	const FLinearColor& GetColor() const;
 	const FVector4f& GetVector() const;
 	const FTG_Texture& GetTexture() const;
+	FTG_Texture& GetTexture();
 
 	FLinearColor GetColor(FLinearColor Default = FLinearColor::Black);
 	FVector4f GetVector(FVector4f Default = FVector4f { 0, 0, 0, 0 });
@@ -193,12 +194,18 @@ private:
 	TArray<FTG_Variant> Array;			/// An array of textures
 	int ViewIndex = 0;
 
+
 public:
 	void SetNum(int Count);
 	void Set(int Index, const FTG_Variant& Var);
 	void Set(int Index, const FTG_Texture& Tex);
 	void SetViewIndex(int Index);
 
+	FORCEINLINE bool IsTexture()
+	{
+		return !Array.IsEmpty() && Array[0].IsTexture();
+	}
+
 	FORCEINLINE int Num() const
 	{
 		return Array.Num();
@@ -210,11 +217,25 @@ public:
 		return Array[Index];
 	}
 
+	FORCEINLINE const FTG_Variant& Get(int Index) const
+	{
+		check(Index >= 0 && Index < Array.Num());
+		return Array[Index];
+	}
+
 	FORCEINLINE const TArray<FTG_Variant>& GetArray() const
 	{
 		return Array;
 	}
 
+	FORCEINLINE FTG_Texture* GetViewTexturePtr() 
+	{
+		if (ViewIndex >= 0 && ViewIndex < Array.Num() && Array[ViewIndex].IsTexture())
+			return &(Array[ViewIndex].GetTexture());
+
+		return nullptr;
+	}
+
 	FORCEINLINE FTG_Texture GetViewTexture() const
 	{
 		if (ViewIndex >= 0 && ViewIndex < Array.Num() && Array[ViewIndex].IsTexture())
