Making blocking (non-async) texture graph execution.

=

diff --git a/Engine/Plugins/Experimental/TextureGraph/Source/TextureGraph/Private/Blueprint/TG_AsyncExportTask.cpp b/Engine/Plugins/Experimental/TextureGraph/Source/TextureGraph/Private/Blueprint/TG_AsyncExportTask.cpp
--- a/Engine/Plugins/Experimental/TextureGraph/Source/TextureGraph/Private/Blueprint/TG_AsyncExportTask.cpp
+++ b/Engine/Plugins/Experimental/TextureGraph/Source/TextureGraph/Private/Blueprint/TG_AsyncExportTask.cpp
@@ -11,7 +11,6 @@
 UTG_AsyncExportTask::UTG_AsyncExportTask(const FObjectInitializer& ObjectInitializer)
 	: Super(ObjectInitializer)
 {
-
 }
 
 UTG_AsyncExportTask* UTG_AsyncExportTask::TG_AsyncExportTask( UTextureGraphBase* InTextureGraph, const bool OverwriteTextures, const bool bSave, const bool bExportAll)
@@ -24,8 +23,8 @@ UTG_AsyncExportTask* UTG_AsyncExportTask::TG_AsyncExportTask( UTextureGraphBase*
 
 	if (InTextureGraph != nullptr)
 	{
-		Task->OrignalTextureGraphPtr = InTextureGraph;
-		Task->TextureGraphPtr = Cast<UTextureGraphBase>(StaticDuplicateObject(Task->OrignalTextureGraphPtr, GetTransientPackage(), NAME_None, RF_Standalone));
+		Task->OriginalTextureGraphPtr = InTextureGraph;
+		Task->TextureGraphPtr = Cast<UTextureGraphBase>(StaticDuplicateObject(Task->OriginalTextureGraphPtr, GetTransientPackage(), NAME_None, RF_Standalone));
 		Task->TextureGraphPtr->Initialize();
 		FTG_HelperFunctions::InitTargets(Task->TextureGraphPtr);
 		Task->RegisterWithTGAsyncTaskManger();
@@ -34,6 +33,31 @@ UTG_AsyncExportTask* UTG_AsyncExportTask::TG_AsyncExportTask( UTextureGraphBase*
 	return Task;
 }
 
+const TArray<UTextureRenderTarget2D*>& UTG_AsyncExportTask::ActivateBlocking(JobBatchPtr Batch)
+{
+	TargetExportSettings = FExportSettings();
+	TargetExportSettings.OnDone.BindUFunction(this, "OnExportDone");
+
+	if (!Batch)
+		Batch = FTG_HelperFunctions::InitExportBatch(TextureGraphPtr, "", "", TargetExportSettings, false, bOverwriteTextures, bExportAll, bSave);
+	
+	Super::ActivateBlocking(Batch);
+
+	bool bExportDone = false;
+	TextureExporter::ExportAsUAsset(TextureGraphPtr, TargetExportSettings, "").then([&bExportDone]()
+		{
+			bExportDone = true;
+		});
+
+	while (!bExportDone)
+	{
+		TextureGraphEngine::GetInstance()->Update(0);
+		FTaskGraphInterface::Get().ProcessThreadUntilIdle(ENamedThreads::GameThread);
+	}
+
+	return OutputRts;
+}
+
 void UTG_AsyncExportTask::Activate()
 {
 	// Start the async task on a new thread
@@ -56,8 +80,10 @@ void UTG_AsyncExportTask::OnExportDone()
 {
 	TargetExportSettings.ExportPreset.clear();
 	TargetExportSettings.OnDone.Unbind();
-	OnDone.Broadcast();
-	OnStaticDone.ExecuteIfBound(0);
+
+	if (OnDone.IsBound())
+		OnDone.Broadcast();
+
 	TextureGraphPtr->FlushInvalidations();
 	ClearFlags(RF_Standalone);
 	SetReadyToDestroy();
@@ -70,7 +96,7 @@ void UTG_AsyncExportTask::FinishDestroy()
 		TextureGraphPtr->GetSettings()->FreeTargets();
 		TextureGraphPtr->ClearFlags(RF_Standalone);
 		TextureGraphPtr = nullptr;
-		OrignalTextureGraphPtr = nullptr;
+		OriginalTextureGraphPtr = nullptr;
 	}
 	Super::FinishDestroy();
 }
diff --git a/Engine/Plugins/Experimental/TextureGraph/Source/TextureGraph/Private/Blueprint/TG_AsyncRenderTask.cpp b/Engine/Plugins/Experimental/TextureGraph/Source/TextureGraph/Private/Blueprint/TG_AsyncRenderTask.cpp
--- a/Engine/Plugins/Experimental/TextureGraph/Source/TextureGraph/Private/Blueprint/TG_AsyncRenderTask.cpp
+++ b/Engine/Plugins/Experimental/TextureGraph/Source/TextureGraph/Private/Blueprint/TG_AsyncRenderTask.cpp
@@ -12,51 +12,79 @@
 #include "Device/FX/DeviceBuffer_FX.h"
 #include "Model/Mix/MixSettings.h"
 #include "GameDelegates.h"
+#include "Job/Scheduler.h"
 
-UTG_AsyncRenderTask::UTG_AsyncRenderTask(const FObjectInitializer& ObjectInitializer)
+UTG_AsyncRenderTaskBase::UTG_AsyncRenderTaskBase(const FObjectInitializer& ObjectInitializer)
 	: Super(ObjectInitializer)
 {
-	
 }
 
-UTG_AsyncRenderTask* UTG_AsyncRenderTask::TG_AsyncRenderTask( UTextureGraphBase* InTextureGraph)
+JobBatchPtr UTG_AsyncRenderTaskBase::PrepareActivate(JobBatchPtr Batch, bool bIsAsync)
 {
-	UTG_AsyncRenderTask* Task = NewObject<UTG_AsyncRenderTask>();
-	Task->SetFlags(RF_Standalone);
+	UE_LOG(LogTextureGraph, Log, TEXT("UTG_AsyncRenderTaskBase::PrepareActivate"));
 
-	if (InTextureGraph != nullptr)
+	if (TextureGraphPtr == nullptr)
 	{
-		Task->OriginalTextureGraphPtr = InTextureGraph;
-		Task->TextureGraphPtr = Cast<UTextureGraphBase>(StaticDuplicateObject(Task->OriginalTextureGraphPtr, GetTransientPackage(), NAME_None, RF_Standalone));
-		Task->TextureGraphPtr->Initialize();
-		FTG_HelperFunctions::InitTargets(Task->TextureGraphPtr);
-		Task->RegisterWithTGAsyncTaskManger();
-		FWorldDelegates::OnWorldCleanup.AddUObject(Task, &UTG_AsyncRenderTask::OnWorldCleanup);
+		UE_LOG(LogTextureGraph, Warning, TEXT("UTG_AsyncRenderTaskBase::Cannot render Texture Graph not selected"));
+		return nullptr;
 	}
-	
-	return Task;
+
+	OutputBlobs.Empty();
+	OutputRts.Empty();
+	TextureGraphPtr->FlushInvalidations();
+
+	if (!Batch)
+		Batch = FTG_HelperFunctions::InitRenderBatch(TextureGraphPtr, nullptr);
+
+	Batch->SetAsync(bIsAsync);
+	Batch->SetNoCache(true);
+
+	bool bIsDisableIdle = TextureGraphEngine::GetScheduler()->IsDisableIdle();
+	TextureGraphEngine::GetScheduler()->SetDisableIdle(true);
+
+	Batch->OnDone([=](JobBatch*) mutable
+	{
+		TextureGraphEngine::GetScheduler()->SetDisableIdle(bIsDisableIdle);
+	});
+
+	TextureGraphEngine::GetScheduler()->AddBatch(Batch);
+
+	return Batch;
 }
 
-void UTG_AsyncRenderTask::Activate()
+const TArray<UTextureRenderTarget2D*>& UTG_AsyncRenderTaskBase::ActivateBlocking(JobBatchPtr Batch)
 {
-	// Start the async task on a new thread
-	Super::Activate();
-	UE_LOG(LogTextureGraph, Log, TEXT("UTG_AsyncRenderTask:: Activate"));
+	UE_LOG(LogTextureGraph, Log, TEXT("UTG_AsyncRenderTaskBase::ActivateBlocking"));
 
-	if (TextureGraphPtr == nullptr)
+	Batch = PrepareActivate(Batch, false);
+
+	if (!Batch)
+		return OutputRts;
+
+	while (!Batch->IsFinished())
 	{
-		UE_LOG(LogTextureGraph, Warning, TEXT("UTG_AsyncRenderTask::Cannot render Texture Graph not selected"));
-		return;
+		TextureGraphEngine::GetInstance()->Update(0);
+		FTaskGraphInterface::Get().ProcessThreadUntilIdle(ENamedThreads::GameThread);
 	}
 
-	OutputBlobs.Empty();
-	OutputRts.Empty();
-	TextureGraphPtr->FlushInvalidations();
+	GatherAllOutputBlobs();
+	FinalizeAllOutputBlobs();
 
-	FTG_HelperFunctions::RenderAsync(TextureGraphPtr)
-	.then([this](bool bRenderResult) mutable
+	bRenderComplete = false;
+	GetRenderTextures();
+
+	while (!bRenderComplete)
 	{
-		TextureGraphPtr->Graph()->ForEachNodes([this](const UTG_Node* node, uint32 index)
+		TextureGraphEngine::GetInstance()->Update(0);
+		FTaskGraphInterface::Get().ProcessThreadUntilIdle(ENamedThreads::GameThread);
+	}
+
+	return OutputRts;
+}
+
+void UTG_AsyncRenderTaskBase::GatherAllOutputBlobs()
+{
+	TextureGraphPtr->Graph()->ForEachNodes([this](const UTG_Node* node, uint32 index)
 		{
 			if (node && node->GetExpression()->IsA(UTG_Expression_Output::StaticClass()))
 			{
@@ -68,30 +96,9 @@ void UTG_AsyncRenderTask::Activate()
 				}
 			}
 		});
-
-		return FinalizeAllOutputBlobs();
-	})
-	.then([this](bool bFinalized)
-	{
-		return GetRenderTextures();
-	})
-	.then([this](bool bRtResult)
-	{
-		OnDone.Broadcast(OutputRts);
-		OnStaticDone.ExecuteIfBound(OutputRts);
-		bRenderComplete = true;
-
-		UE_LOG(LogTextureGraph, Log, TEXT("UTG_AsyncRenderTask:: OnDone : bShouldDestroyOnRenderComplete %i"), bShouldDestroyOnRenderComplete ? 1 : 0);
-		if (bShouldDestroyOnRenderComplete)
-		{
-			SetReadyToDestroy();
-		}
-
-		return bRtResult;
-	});
 }
 
-AsyncBool UTG_AsyncRenderTask::FinalizeAllOutputBlobs()
+AsyncBool UTG_AsyncRenderTaskBase::FinalizeAllOutputBlobs()
 {
 	std::vector<AsyncBlobResultPtr> promises;
 	for (const auto& Blob : OutputBlobs)
@@ -106,7 +113,26 @@ AsyncBool UTG_AsyncRenderTask::FinalizeAllOutputBlobs()
 	});
 }
 
-AsyncBool UTG_AsyncRenderTask::GetRenderTextures()
+void UTG_AsyncRenderTaskBase::GatherAllRenderTargets()
+{
+	for (const auto& Blob : OutputBlobs)
+	{
+		auto TiledOutput = std::static_pointer_cast<TiledBlob>(Blob);
+		auto FXBuffer = std::static_pointer_cast<DeviceBuffer_FX>(TiledOutput->GetBufferRef().GetPtr());
+		auto Rt = FXBuffer->GetTexture()->GetRenderTarget();
+		OutputRts.Add(Rt);
+	}
+
+	bRenderComplete = true;
+
+	UE_LOG(LogTextureGraph, Log, TEXT("UTG_AsyncRenderTaskBase:: OnDone : bShouldDestroyOnRenderComplete %i"), bShouldDestroyOnRenderComplete ? 1 : 0);
+	if (bShouldDestroyOnRenderComplete)
+	{
+		SetReadyToDestroy();
+	}
+}
+
+AsyncBool UTG_AsyncRenderTaskBase::GetRenderTextures()
 {
 	std::vector<AsyncBufferResultPtr> promises;
 	for (const auto& Blob : OutputBlobs)
@@ -117,18 +143,12 @@ AsyncBool UTG_AsyncRenderTask::GetRenderTextures()
 
 	return cti::when_all(promises.begin(), promises.end()).then([this](std::vector<BufferResultPtr> results) mutable
 	{
-		for (const auto& Blob : OutputBlobs)
-		{
-			auto TiledOutput = std::static_pointer_cast<TiledBlob>(Blob);
-			auto FXBuffer = std::static_pointer_cast<DeviceBuffer_FX>(TiledOutput->GetBufferRef().GetPtr());
-			auto Rt = FXBuffer->GetTexture()->GetRenderTarget();
-			OutputRts.Add(Rt);
-		}
+		GatherAllRenderTargets();
 		return true;
 	});
 }
 
-void UTG_AsyncRenderTask::OnWorldCleanup(UWorld* World, bool bSessionEnded, bool CleanupResources)
+void UTG_AsyncRenderTaskBase::OnWorldCleanup(UWorld* World, bool bSessionEnded, bool CleanupResources)
 {
 	FWorldDelegates::OnWorldCleanup.RemoveAll(this);
 	bShouldDestroyOnRenderComplete = true;
@@ -139,20 +159,20 @@ void UTG_AsyncRenderTask::OnWorldCleanup(UWorld* World, bool bSessionEnded, bool
 		SetReadyToDestroy();
 	}
 
-	UE_LOG(LogTextureGraph, Log, TEXT("UTG_AsyncRenderTask:: OnWorldCleanup"));
+	UE_LOG(LogTextureGraph, Log, TEXT("UTG_AsyncRenderTaskBase:: OnWorldCleanup"));
 }
 
-void UTG_AsyncRenderTask::SetReadyToDestroy()
+void UTG_AsyncRenderTaskBase::SetReadyToDestroy()
 {
-	UE_LOG(LogTextureGraph, Log, TEXT("UTG_AsyncRenderTask:: SetReadyToDestroy"));
+	UE_LOG(LogTextureGraph, Log, TEXT("UTG_AsyncRenderTaskBase:: SetReadyToDestroy"));
 	TextureGraphPtr->FlushInvalidations();
 	ClearFlags(RF_Standalone);
 	UTG_AsyncTask::SetReadyToDestroy();
 }
 
-void UTG_AsyncRenderTask::FinishDestroy()
+void UTG_AsyncRenderTaskBase::FinishDestroy()
 {
-	UE_LOG(LogTextureGraph, Log, TEXT("UTG_AsyncRenderTask:: FinishDestroy"));
+	UE_LOG(LogTextureGraph, Log, TEXT("UTG_AsyncRenderTaskBase:: FinishDestroy"));
 	if (TextureGraphPtr != nullptr)
 	{
 		TextureGraphPtr->GetSettings()->FreeTargets();
@@ -161,4 +181,58 @@ void UTG_AsyncRenderTask::FinishDestroy()
 		OriginalTextureGraphPtr = nullptr;
 	}
 	Super::FinishDestroy();
-}
\ No newline at end of file
+}
+
+/////////////////////////////////////////////////////////////////////
+UTG_AsyncRenderTask::UTG_AsyncRenderTask(const FObjectInitializer& ObjectInitializer)
+	: Super(ObjectInitializer)
+{
+}
+
+UTG_AsyncRenderTask* UTG_AsyncRenderTask::TG_AsyncRenderTask(UTextureGraphBase* InTextureGraph)
+{
+	UTG_AsyncRenderTask* Task = NewObject<UTG_AsyncRenderTask>();
+	Task->SetFlags(RF_Standalone);
+
+	if (InTextureGraph != nullptr)
+	{
+		Task->OriginalTextureGraphPtr = InTextureGraph;
+		Task->TextureGraphPtr = Cast<UTextureGraphBase>(StaticDuplicateObject(Task->OriginalTextureGraphPtr, GetTransientPackage(), NAME_None, RF_Standalone));
+		Task->TextureGraphPtr->Initialize();
+		FTG_HelperFunctions::InitTargets(Task->TextureGraphPtr);
+		Task->RegisterWithTGAsyncTaskManger();
+		FWorldDelegates::OnWorldCleanup.AddUObject(Task, &UTG_AsyncRenderTaskBase::OnWorldCleanup);
+	}
+
+	return Task;
+}
+
+
+void UTG_AsyncRenderTask::Activate()
+{
+	// Start the async task on a new thread
+	Super::Activate();
+	UE_LOG(LogTextureGraph, Log, TEXT("UTG_AsyncRenderTask::Activate"));
+
+	JobBatchPtr Batch = PrepareActivate(nullptr, true);
+
+	if (!Batch)
+		return;
+
+	FTG_HelperFunctions::RenderAsync(TextureGraphPtr)
+		.then([this](bool bRenderResult) mutable
+			{
+				GatherAllOutputBlobs();
+				return FinalizeAllOutputBlobs();
+			})
+		.then([this](bool bFinalized)
+			{
+				return GetRenderTextures();
+			})
+		.then([this](bool bRtResult)
+			{
+				OnDone.Broadcast(OutputRts);
+				return bRtResult;
+			});
+}
+
diff --git a/Engine/Plugins/Experimental/TextureGraph/Source/TextureGraph/Private/Blueprint/TG_BlueprintFunctionLibrary.cpp b/Engine/Plugins/Experimental/TextureGraph/Source/TextureGraph/Private/Blueprint/TG_BlueprintFunctionLibrary.cpp
--- a/Engine/Plugins/Experimental/TextureGraph/Source/TextureGraph/Private/Blueprint/TG_BlueprintFunctionLibrary.cpp
+++ b/Engine/Plugins/Experimental/TextureGraph/Source/TextureGraph/Private/Blueprint/TG_BlueprintFunctionLibrary.cpp
@@ -267,18 +267,16 @@ FTG_OutputSettings UTG_BlueprintFunctionLibrary::GetSettingsParameterValue(UObje
 	return ParameterValue;
 }
 
-void UTG_BlueprintFunctionLibrary::RenderTextureGraph(UObject* WorldContextObject, UTextureGraphBase* InTextureGraph, FOnTextureGraphRenderDone CompletionDelegate)
+const TArray<UTextureRenderTarget2D*>& UTG_BlueprintFunctionLibrary::RenderTextureGraph(UObject* WorldContextObject, UTextureGraphBase* InTextureGraph)
 {
 	UTG_AsyncRenderTask* Task = UTG_AsyncRenderTask::TG_AsyncRenderTask(InTextureGraph);
-	Task->OnStaticDone = CompletionDelegate;
-	Task->Activate();
+	return Task->ActivateBlocking(nullptr);
 }
 
-void UTG_BlueprintFunctionLibrary::ExportTextureGraph(UObject* WorldContextObject, UTextureGraphBase* InTextureGraph, FOnTextureGraphExportDone CompletionDelegate)
+void UTG_BlueprintFunctionLibrary::ExportTextureGraph(UObject* WorldContextObject, UTextureGraphBase* InTextureGraph)
 {
 	UTG_AsyncExportTask* Task = UTG_AsyncExportTask::TG_AsyncExportTask(InTextureGraph, true, true, true);
-	Task->OnStaticDone = CompletionDelegate;
-	Task->Activate();
+	Task->ActivateBlocking(nullptr);
 }
 
 void UTG_BlueprintFunctionLibrary::AddParamWarning(FName ParamName, UObject* ObjectPtr, FString FunctionName)
diff --git a/Engine/Plugins/Experimental/TextureGraph/Source/TextureGraph/Public/Blueprint/TG_AsyncExportTask.h b/Engine/Plugins/Experimental/TextureGraph/Source/TextureGraph/Public/Blueprint/TG_AsyncExportTask.h
--- a/Engine/Plugins/Experimental/TextureGraph/Source/TextureGraph/Public/Blueprint/TG_AsyncExportTask.h
+++ b/Engine/Plugins/Experimental/TextureGraph/Source/TextureGraph/Public/Blueprint/TG_AsyncExportTask.h
@@ -10,13 +10,13 @@
 #include "TG_HelperFunctions.h"
 #include "Export/TextureExporter.h"
 #include "TG_AsyncTask.h"
+#include "TG_AsyncRenderTask.h"
 #include "TG_AsyncExportTask.generated.h"
 
 DECLARE_DYNAMIC_MULTICAST_DELEGATE(FTSExportTaskDelegate);
-DECLARE_DYNAMIC_DELEGATE_OneParam(FOnTextureGraphExportDone, int, Result);
 
 UCLASS()
-class TEXTUREGRAPH_API UTG_AsyncExportTask : public UTG_AsyncTask
+class TEXTUREGRAPH_API UTG_AsyncExportTask : public UTG_AsyncRenderTaskBase
 {
 	GENERATED_UCLASS_BODY()
 
@@ -25,28 +25,20 @@ public:
 	static UTG_AsyncExportTask* TG_AsyncExportTask(UTextureGraphBase* TextureGraph, const bool OverwriteTextures, const bool bSave, const bool bExportAll);
 
 	virtual void Activate() override;
+	virtual const TArray<UTextureRenderTarget2D*>& ActivateBlocking(JobBatchPtr Batch);
 
 	virtual void FinishDestroy() override;
 
 	UPROPERTY(BlueprintAssignable, Category = "TextureGraph")
 	FTSExportTaskDelegate OnDone;
-	FOnTextureGraphExportDone OnStaticDone;
 
 private:
 
 	UFUNCTION()
 	void OnExportDone();
-	
 	bool bOverwriteTextures = true;
-
 	bool bSave = false;
-	
 	bool bExportAll = false;
-
-	UTextureGraphBase* OrignalTextureGraphPtr;
-
-	UTextureGraphBase* TextureGraphPtr;
-
 	FExportSettings TargetExportSettings;
 };
 
diff --git a/Engine/Plugins/Experimental/TextureGraph/Source/TextureGraph/Public/Blueprint/TG_AsyncRenderTask.h b/Engine/Plugins/Experimental/TextureGraph/Source/TextureGraph/Public/Blueprint/TG_AsyncRenderTask.h
--- a/Engine/Plugins/Experimental/TextureGraph/Source/TextureGraph/Public/Blueprint/TG_AsyncRenderTask.h
+++ b/Engine/Plugins/Experimental/TextureGraph/Source/TextureGraph/Public/Blueprint/TG_AsyncRenderTask.h
@@ -15,44 +15,47 @@
 #include "TG_AsyncRenderTask.generated.h"
 
 DECLARE_DYNAMIC_MULTICAST_DELEGATE_OneParam(FTSRenderTaskDelegate, const TArray<UTextureRenderTarget2D*>&, OutputRts);
-DECLARE_DYNAMIC_DELEGATE_OneParam(FOnTextureGraphRenderDone, const TArray<UTextureRenderTarget2D*>&, OutputRts);
 
 UCLASS()
-class TEXTUREGRAPH_API UTG_AsyncRenderTask : public UTG_AsyncTask
+class TEXTUREGRAPH_API UTG_AsyncRenderTaskBase : public UTG_AsyncTask
 {
 	GENERATED_UCLASS_BODY()
 
 public:
-	UFUNCTION(BlueprintCallable, Category = "TextureGraph" , meta=(DisplayName="Texture Graph Render (Async)" , BlueprintInternalUseOnly = "true"))
-	static UTG_AsyncRenderTask* TG_AsyncRenderTask(UTextureGraphBase* InTextureGraph);
-
-	virtual void Activate() override;
-
+	virtual const TArray<UTextureRenderTarget2D*>& ActivateBlocking(JobBatchPtr Batch);
 	virtual void FinishDestroy() override;
-
 	virtual void SetReadyToDestroy() override;
 
-	UPROPERTY(BlueprintAssignable, Category = "TextureGraph")
-	FTSRenderTaskDelegate OnDone;
-	FOnTextureGraphRenderDone OnStaticDone;
-
-private:
+protected:
 	AsyncBool FinalizeAllOutputBlobs();
-
+	JobBatchPtr PrepareActivate(JobBatchPtr Batch, bool bIsAsync);
+	void GatherAllOutputBlobs();
+	void GatherAllRenderTargets();
 	AsyncBool GetRenderTextures();
 
 	void OnWorldCleanup(UWorld* World, bool bSessionEnded, bool CleanupResources);
 
 	TArray<UTextureRenderTarget2D*> OutputRts;
-
 	TArray<BlobPtr> OutputBlobs;
-
 	UTextureGraphBase* OriginalTextureGraphPtr;
-
 	UTextureGraphBase* TextureGraphPtr;
-
 	bool bShouldDestroyOnRenderComplete = false;
-
 	bool bRenderComplete = false;
 };
 
+/////////////////////////////////////////////////////////////////////
+UCLASS()
+class TEXTUREGRAPH_API UTG_AsyncRenderTask : public UTG_AsyncRenderTaskBase
+{
+	GENERATED_UCLASS_BODY()
+
+public:
+	UFUNCTION(BlueprintCallable, Category = "TextureGraph" , meta=(DisplayName="Texture Graph Render (Async)" , BlueprintInternalUseOnly = "true"))
+	static UTG_AsyncRenderTask* TG_AsyncRenderTask(UTextureGraphBase* InTextureGraph);
+
+	virtual void Activate() override;
+
+	UPROPERTY(BlueprintAssignable, Category = "TextureGraph")
+	FTSRenderTaskDelegate OnDone;
+};
+
diff --git a/Engine/Plugins/Experimental/TextureGraph/Source/TextureGraph/Public/Blueprint/TG_BlueprintFunctionLibrary.h b/Engine/Plugins/Experimental/TextureGraph/Source/TextureGraph/Public/Blueprint/TG_BlueprintFunctionLibrary.h
--- a/Engine/Plugins/Experimental/TextureGraph/Source/TextureGraph/Public/Blueprint/TG_BlueprintFunctionLibrary.h
+++ b/Engine/Plugins/Experimental/TextureGraph/Source/TextureGraph/Public/Blueprint/TG_BlueprintFunctionLibrary.h
@@ -82,11 +82,11 @@ public:
 
 	/** Render the texture graph and return an array of texture render targets. */
 	UFUNCTION(BlueprintCallable, Category = "TextureGraph", meta = (Keywords = "RenderTextureGraph", WorldContext = "WorldContextObject"))
-	static void RenderTextureGraph(UObject* WorldContextObject, UTextureGraphBase* InTextureGraph, FOnTextureGraphRenderDone CompletionDelegate);
+	static const TArray<UTextureRenderTarget2D*>& RenderTextureGraph(UObject* WorldContextObject, UTextureGraphBase* InTextureGraph);
 
 	/** Export texture graph. */
 	UFUNCTION(BlueprintCallable, Category = "TextureGraph", meta = (Keywords = "ExportTextureGraph", WorldContext = "WorldContextObject"))
-	static void ExportTextureGraph(UObject* WorldContextObject, UTextureGraphBase* InTextureGraph, FOnTextureGraphExportDone CompletionDelegate);
+	static void ExportTextureGraph(UObject* WorldContextObject, UTextureGraphBase* InTextureGraph);
 
 	static void AddParamWarning(FName ParamName, UObject* ObjectPtr, FString FunctionName);
 	static void AddError(UObject* ObjectPtr, FString FunctionName, FString Error);
diff --git a/Engine/Plugins/Experimental/TextureGraph/Source/TextureGraphEngine/Device/Device.cpp b/Engine/Plugins/Experimental/TextureGraph/Source/TextureGraphEngine/Device/Device.cpp
--- a/Engine/Plugins/Experimental/TextureGraph/Source/TextureGraphEngine/Device/Device.cpp
+++ b/Engine/Plugins/Experimental/TextureGraph/Source/TextureGraphEngine/Device/Device.cpp
@@ -537,7 +537,7 @@ void Device::Update(float dt)
 		PostExecThreadFunc();
 	}
 
-	UpdatePreExec();
+	UpdateNonAsync();
 }
 
 bool Device::ShouldCollect(double TimeNow, DeviceBuffer* Buffer)
@@ -788,6 +788,26 @@ void Device::MarkTaskFinished(DeviceNativeTaskPtr& Task, bool bRemoveFromOwnerQu
 	}
 }
 
+void Device::UpdateNonAsyncPostExec()
+{
+	check(IsInGameThread());
+	DeviceNativeTaskPtrWVec Tasks = NonAsyncPostExecQueue.to_vector_and_clear();
+
+	for (DeviceNativeTaskPtrW& TaskW : Tasks)
+	{
+		DeviceNativeTaskPtr Task = TaskW.lock();
+		check(Task);
+		Task->PostExec();
+		MarkTaskFinished(Task, true);
+	}
+}
+
+void Device::UpdateNonAsync()
+{
+	UpdatePreExec();
+	UpdateNonAsyncPostExec();
+}
+
 void Device::UpdatePreExec()
 {
 	DeviceNativeTaskPtrWVec Tasks = NativePreExec.to_vector_and_clear();
@@ -877,7 +897,12 @@ void Device::ExecThreadFunc()
 
 							/// Only add if its not an exec only Task
 							if (!Task->IsExecOnly() && !Task->IsDone())
-								NativePostExec.add(Task);
+							{
+								if (Task->IsAsync())
+									NativePostExec.add(Task);
+								else
+									NonAsyncPostExecQueue.add(Task);
+							}
 							else
 								MarkTaskFinished(Task, true);
 						});
@@ -890,7 +915,12 @@ void Device::ExecThreadFunc()
 							{
 								UE_LOG(LogDevice, VeryVerbose, TEXT("[%u] END - Device::ExecAsync Task: %s [SUCCESS]"), Task->GetTaskId(), *Task->GetName());
 								if (!Task->IsDone())
-									NativePostExec.add(Task);
+								{
+									if (Task->IsAsync())
+										NativePostExec.add(Task);
+									else
+										NonAsyncPostExecQueue.add(Task);
+								}
 								else
 									MarkTaskFinished(Task, true);
 							})
@@ -898,7 +928,12 @@ void Device::ExecThreadFunc()
 							{
 								UE_LOG(LogDevice, VeryVerbose, TEXT("[%u] END - Device::ExecAsync Task: %s [FAIL]"), Task->GetTaskId(), *Task->GetName());
 								if (!Task->IsDone())
-									NativePostExec.add(Task);
+								{
+									if (Task->IsAsync())
+										NativePostExec.add(Task);
+									else
+										NonAsyncPostExecQueue.add(Task);
+								}
 								else
 									MarkTaskFinished(Task, true);
 							});
@@ -958,32 +993,39 @@ void Device::AddNativeTask(DeviceNativeTaskPtr Task)
 	/// Keep a strong ref here
 	NativeTasks.add(Task);
 
-	/// and add to the pre-exec queue
-	if (!Task->IsExecOnly())
-	{
-		/// Debug check
+	//if (Task->IsAsync())
+	//{
+		/// and add to the pre-exec queue
+		if (!Task->IsExecOnly())
+		{
+			/// Debug check
 #if (DEVICE_NATIVE_TASKS_CHECK_DEPENDENCY_PRESENT == 1)
-		const DeviceNativeTaskPtrVec& prev = Task->Prev();
-		NativePreExecWait.lock(); 
-		auto& v = NativePreExecWait.accessInnerVector_Unsafe();
+			const DeviceNativeTaskPtrVec& prev = Task->Prev();
+			NativePreExecWait.lock();
+			auto& v = NativePreExecWait.accessInnerVector_Unsafe();
 
-		for (DeviceNativeTaskPtr p : prev)
-		{
-			bool found = false;
-			for (auto iter = v.begin(); iter != v.end() && !found; iter++)
+			for (DeviceNativeTaskPtr p : prev)
 			{
-				if ((*iter).lock() == p)
-					found = true;
-			}
+				bool found = false;
+				for (auto iter = v.begin(); iter != v.end() && !found; iter++)
+				{
+					if ((*iter).lock() == p)
+						found = true;
+				}
 
-			check(found);
-		}
-		NativePreExecWait.unlock();
+				check(found);
+			}
+			NativePreExecWait.unlock();
 #endif 
-		NativePreExecWait.add(Task);
-	}
-	else
-		NativeExec.add(Task);
+			NativePreExecWait.add(Task);
+		}
+		else
+			NativeExec.add(Task);
+	//}
+	//else
+	//{
+	//	NativePreExec.add(Task);
+	//}
 }
 
 void Device::PrintStats()
diff --git a/Engine/Plugins/Experimental/TextureGraph/Source/TextureGraphEngine/Device/Device.h b/Engine/Plugins/Experimental/TextureGraph/Source/TextureGraphEngine/Device/Device.h
--- a/Engine/Plugins/Experimental/TextureGraph/Source/TextureGraphEngine/Device/Device.h
+++ b/Engine/Plugins/Experimental/TextureGraph/Source/TextureGraphEngine/Device/Device.h
@@ -128,12 +128,15 @@ protected:
 	std::unique_ptr<std::thread>	PostExecThread = nullptr;		/// Thread running post-exec
 	size_t							PostExecThreadId = 0;			/// The thread id of the post-exec thread (used mostly for debugging and error checking)
 
+	DeviceNativeTask_PriorityQueue	NonAsyncPostExecQueue;			/// Non asynchronous post-exec tasks are queued here
 	DeviceObserverSourcePtr			ObserverSource;					/// Device observer source that sends of different events happening within the device
 
 	mutable std::mutex				AwaitingCollectionMutex;		/// Mutex for the list below
 	DeviceBufferPList				AwaitingCollection;				/// Device buffers that are awaiting collection
 
 	virtual void					UpdateAwaitingCollection();
+	virtual void					UpdateNonAsyncPostExec();
+	virtual void					UpdateNonAsync();
 
 									Device(DeviceType InType, DeviceBuffer* InBufferFactory);
 	virtual							~Device();
diff --git a/Engine/Plugins/Experimental/TextureGraph/Source/TextureGraphEngine/Device/DeviceNativeTask.cpp b/Engine/Plugins/Experimental/TextureGraph/Source/TextureGraphEngine/Device/DeviceNativeTask.cpp
--- a/Engine/Plugins/Experimental/TextureGraph/Source/TextureGraphEngine/Device/DeviceNativeTask.cpp
+++ b/Engine/Plugins/Experimental/TextureGraph/Source/TextureGraphEngine/Device/DeviceNativeTask.cpp
@@ -471,11 +471,6 @@ void DeviceNativeTask::PreExec()
 		Wait();
 }
 
-bool DeviceNativeTask::IsAsync() const
-{
-	return false;
-}
-
 void DeviceNativeTask::SetPromise(int32 Value)
 {
 	ErrorCode = Value;
@@ -522,17 +517,17 @@ DeviceNativeTask_Generic::DeviceNativeTask_Generic(Device* TargetDevice, FString
 	: DeviceNativeTask(TaskPriority, TaskName)
 	, Dev(TargetDevice)
 	, ThreadId(TaskThreadId)
-	, bIsAsync(IsTaskAsync)
 {
+	DeviceNativeTask::bIsAsync = IsTaskAsync;
 }
 
 DeviceNativeTask_Generic::DeviceNativeTask_Generic(DeviceNativeTaskPtrVec& PrevTasks, Device* TargetDevice, FString TaskName, ENamedThreads::Type TaskThreadId /* = ENamedThreads::AnyThread */, 
-	bool IsTAskAsync /* = false */, int32 TaskPriority /* = JobPriority::kNormal */)
+	bool IsTaskAsync /* = false */, int32 TaskPriority /* = JobPriority::kNormal */)
 	: DeviceNativeTask(PrevTasks, TaskPriority, TaskName)
 	, Dev(TargetDevice)
 	, ThreadId(TaskThreadId)
-	, bIsAsync(IsTAskAsync)
 {
+	DeviceNativeTask::bIsAsync = IsTaskAsync;
 }
 
 //////////////////////////////////////////////////////////////////////////
diff --git a/Engine/Plugins/Experimental/TextureGraph/Source/TextureGraphEngine/Device/DeviceNativeTask.h b/Engine/Plugins/Experimental/TextureGraph/Source/TextureGraphEngine/Device/DeviceNativeTask.h
--- a/Engine/Plugins/Experimental/TextureGraph/Source/TextureGraphEngine/Device/DeviceNativeTask.h
+++ b/Engine/Plugins/Experimental/TextureGraph/Source/TextureGraphEngine/Device/DeviceNativeTask.h
@@ -44,6 +44,7 @@ protected:
 	bool							bExecOnly = false;		/// Whether its an exec only Task with no pre and post exec. This is to avoid waiting for system Tasks
 	bool							bIsCulled = false;		/// Whether the Task is culled or not
 	bool							bIsDone = false;		/// Whether the Task is done or not
+	bool							bIsAsync = true;		/// Whether the native task is async or not
 
 	int64_t							BatchId = 0;			/// What's the batch id of this native Task
 
@@ -86,7 +87,6 @@ public:
 	/// Waits for the _prev Task to finish before executing. Its the responsibility of the system
 	/// to ensure that both the _prev
 	virtual ENamedThreads::Type		GetExecutionThread() const;
-	virtual bool					IsAsync() const;
 	virtual FString					GetName() const;
 	virtual FString					GetDebugName() const;
 	virtual void					Terminate();
@@ -99,6 +99,7 @@ public:
 
 	void							CheckCyclicalDependency(std::unordered_set<DeviceNativeTask*>& TransferChain);
 	virtual bool					DebugCompleteCheck();
+	virtual bool					IsAsync() const { return bIsAsync; }
 
 	//////////////////////////////////////////////////////////////////////////
 	/// Inline functions
@@ -113,6 +114,7 @@ public:
 	FORCEINLINE bool				IsExecOnly() const { return bExecOnly; }
 	FORCEINLINE bool				IsCulled() const { return bIsCulled; }
 	FORCEINLINE bool				IsDone() const { return bIsDone || IsCulled(); }
+	FORCEINLINE void				SetIsAsync(bool bIsAsyncIn) { bIsAsync = bIsAsyncIn; }
 	FORCEINLINE const DeviceNativeTaskPtrVec& GetPrev() const { return Prev; }
 };
 
@@ -130,7 +132,6 @@ class TEXTUREGRAPHENGINE_API DeviceNativeTask_Generic : public DeviceNativeTask
 protected:
 	Device*							Dev = nullptr;			/// The target device that this is associated with
 	ENamedThreads::Type				ThreadId = ENamedThreads::AnyThread; /// What thread is this going to run on
-	bool							bIsAsync = false;		/// Is this asynchronous or not
 
 public:
 									DeviceNativeTask_Generic(Device* TargetDevice, FString TaskName, ENamedThreads::Type TaskThreadId = ENamedThreads::AnyThread, 
@@ -143,13 +144,11 @@ public:
 	/// Waits for the _prev Task to finish before executing. Its the responsibility of the system
 	/// to ensure that both the _prev
 	virtual ENamedThreads::Type		GetExecutionThread() const override { return ThreadId; }
-	virtual bool					IsAsync() const override { return bIsAsync; }
 
 	//////////////////////////////////////////////////////////////////////////
 	/// Inline functions
 	//////////////////////////////////////////////////////////////////////////
 	FORCEINLINE ENamedThreads::Type& GetThreadId() { return ThreadId; }
-	FORCEINLINE bool&				IsAsync() { return bIsAsync; }
 };
 
 //////////////////////////////////////////////////////////////////////////
diff --git a/Engine/Plugins/Experimental/TextureGraph/Source/TextureGraphEngine/Device/Mem/Device_Mem.cpp b/Engine/Plugins/Experimental/TextureGraph/Source/TextureGraphEngine/Device/Mem/Device_Mem.cpp
--- a/Engine/Plugins/Experimental/TextureGraph/Source/TextureGraphEngine/Device/Mem/Device_Mem.cpp
+++ b/Engine/Plugins/Experimental/TextureGraph/Source/TextureGraphEngine/Device/Mem/Device_Mem.cpp
@@ -38,7 +38,7 @@ void Device_Mem::Update(float Delta)
 void Device_Mem::AddNativeTask(DeviceNativeTaskPtr Task)
 {
 	/// Device_Mem tasks MUST always be async
-	check(Task->IsAsync());
+	//check(Task->IsAsync());
 
 	/// Cannot have game thread as the main execution thread because this can
 	/// potentially block the game thread, where all the tasks execute
diff --git a/Engine/Plugins/Experimental/TextureGraph/Source/TextureGraphEngine/Helper/Util.cpp b/Engine/Plugins/Experimental/TextureGraph/Source/TextureGraphEngine/Helper/Util.cpp
--- a/Engine/Plugins/Experimental/TextureGraph/Source/TextureGraphEngine/Helper/Util.cpp
+++ b/Engine/Plugins/Experimental/TextureGraph/Source/TextureGraphEngine/Helper/Util.cpp
@@ -21,6 +21,7 @@
 #include <Framework/Application/SlateApplication.h>
 #include <UObject/Package.h>
 #include <Materials/Material.h>
+#include "Device/FX/Device_FX.h"
 
 #include "FxMat/MaterialManager.h"
 #if WITH_EDITOR
@@ -87,6 +88,14 @@ bool Util::IsDirectoryEmpty(const FString& Path)
 void Util::OnGameThread(TUniqueFunction<void()> Callback) 
 { 
 	check(!TextureGraphEngine::IsDestroying());
+
+	/// If already in game thread then don't wait
+	if (IsInGameThread())
+	{
+		Callback();
+		return;
+	}
+
 	AsyncTask(ENamedThreads::GameThread, std::move(Callback)); 
 }
 
@@ -121,6 +130,12 @@ void Util::OnRenderingThread(TUniqueFunction<void(FRHICommandListImmediate&)> Ca
 {
 	check(!TextureGraphEngine::IsDestroying());
 
+	if (IsInRenderingThread())
+	{
+		Callback(Device_FX::Get()->RHI());
+		return;
+	}
+
 	ENQUEUE_RENDER_COMMAND(SceneDrawCompletion)([Callback = std::forward<decltype(Callback)>(Callback)](FRHICommandListImmediate& RHI) mutable
 	{
 		Callback(RHI);
diff --git a/Engine/Plugins/Experimental/TextureGraph/Source/TextureGraphEngine/Job/Job.cpp b/Engine/Plugins/Experimental/TextureGraph/Source/TextureGraphEngine/Job/Job.cpp
--- a/Engine/Plugins/Experimental/TextureGraph/Source/TextureGraphEngine/Job/Job.cpp
+++ b/Engine/Plugins/Experimental/TextureGraph/Source/TextureGraphEngine/Job/Job.cpp
@@ -1318,7 +1318,7 @@ ENamedThreads::Type Job::GetExecutionThread() const
 
 bool Job::IsAsync() const
 {
-	return Transform->IsAsync();
+	return Transform->IsAsync() && DeviceNativeTask::IsAsync();
 }
 
 void Job::ResetForReplay(bool noCache)
diff --git a/Engine/Plugins/Experimental/TextureGraph/Source/TextureGraphEngine/Job/JobBatch.cpp b/Engine/Plugins/Experimental/TextureGraph/Source/TextureGraphEngine/Job/JobBatch.cpp
--- a/Engine/Plugins/Experimental/TextureGraph/Source/TextureGraphEngine/Job/JobBatch.cpp
+++ b/Engine/Plugins/Experimental/TextureGraph/Source/TextureGraphEngine/Job/JobBatch.cpp
@@ -199,6 +199,7 @@ void JobBatch::AddNativeJob_Now(JobPtr JobObj, bool AddToAllJobs)
 	{
 		JobObj->FixPriorities();
 
+		JobObj->SetIsAsync(bIsAsync);
 		JobObj->GetRunInfo().Dev->AddNativeTask(JobObj);
 		JobsFinished.push_back(false);
 
diff --git a/Engine/Plugins/Experimental/TextureGraph/Source/TextureGraphEngine/Job/JobBatch.h b/Engine/Plugins/Experimental/TextureGraph/Source/TextureGraphEngine/Job/JobBatch.h
--- a/Engine/Plugins/Experimental/TextureGraph/Source/TextureGraphEngine/Job/JobBatch.h
+++ b/Engine/Plugins/Experimental/TextureGraph/Source/TextureGraphEngine/Job/JobBatch.h
@@ -54,6 +54,7 @@ protected:
 	bool							bNoCache = false;			/// Whether we should ignore the cache
 	bool							bIsIdle = false;			/// Is the batch generated by an idle service or not
 	bool							bIsFinished = false;		/// Whether the batch has been finished or not
+	bool							bIsAsync = true;			/// Whether the job batch is asynchronous or not
 
 	AsyncJobResultPtr				BeginQueue_Native(JobPriorityQueue& InQueue);
 
@@ -110,6 +111,8 @@ public:
 	FORCEINLINE bool				WasGeneratedFromIdleService() const { return bIsIdle; }
 	FORCEINLINE bool&				IsIdle() { return bIsIdle; }
 	FORCEINLINE bool				IsFinished() const { return bIsFinished; }
+	FORCEINLINE void				SetAsync(bool bIsAsyncIn) { bIsAsync = bIsAsyncIn; }
+	FORCEINLINE bool				IsAsync() const { return bIsAsync; }
 
 	//////////////////////////////////////////////////////////////////////////
 	/// Static functions
diff --git a/Engine/Plugins/Experimental/TextureGraph/Source/TextureGraphEngine/Job/Scheduler.cpp b/Engine/Plugins/Experimental/TextureGraph/Source/TextureGraphEngine/Job/Scheduler.cpp
--- a/Engine/Plugins/Experimental/TextureGraph/Source/TextureGraphEngine/Job/Scheduler.cpp
+++ b/Engine/Plugins/Experimental/TextureGraph/Source/TextureGraphEngine/Job/Scheduler.cpp
@@ -236,6 +236,9 @@ AsyncJobResultPtr Scheduler::UpdateIdleBatch(size_t Index)
 
 void Scheduler::UpdateIdle()
 {
+	if (bIsDisableIdle)
+		return;
+
 	{
 		FScopeLock Lock(&CurrentBatchMutex);
 
@@ -329,6 +332,7 @@ void Scheduler::AddBatch(JobBatchPtr Batch)
 
 	FScopeLock Lock(&BatchMutex);
 
+#if 0
 	if (Batches.size())
 	{
 		check(Batches.back()->GetFrameId() <= Batch->GetFrameId());
@@ -336,6 +340,7 @@ void Scheduler::AddBatch(JobBatchPtr Batch)
 
 	check(Batch->GetFrameId())
 		UE_LOG(LogBatch, Verbose, TEXT("Adding Batch: %llu [FrameId: %llu, Num batches: %llu]"), Batch->GetBatchId(), Batch->GetFrameId(), Batches.size());
+#endif 
 
 	Batches.push_back(Batch);
 	ObserverSource->BatchAdded(Batch); // notify observer
diff --git a/Engine/Plugins/Experimental/TextureGraph/Source/TextureGraphEngine/Job/Scheduler.h b/Engine/Plugins/Experimental/TextureGraph/Source/TextureGraphEngine/Job/Scheduler.h
--- a/Engine/Plugins/Experimental/TextureGraph/Source/TextureGraphEngine/Job/Scheduler.h
+++ b/Engine/Plugins/Experimental/TextureGraph/Source/TextureGraphEngine/Job/Scheduler.h
@@ -84,6 +84,7 @@ private:
 	double								TimeSinceIdle = 0;				/// time since scheduler has been idle
 
 	bool								bCaptureNextBatch = false;		/// set the next incoming batch to get captured by RenderDoc
+	bool								bIsDisableIdle = false;			/// Disable idle services or not
 
 	uint64								BatchIndex = 0;					/// Index of the batches
 
@@ -117,6 +118,8 @@ public:
 	FORCEINLINE MipMapServicePtrW		GetMipMapService() const { return MipMapServiceObj; }
 	FORCEINLINE MinMaxServicePtrW		GetMinMaxService() const { return MinMaxServiceObj; }
 	FORCEINLINE HistogramServicePtrW	GetHistogramService() const { return HistogramServiceObj; }
+	FORCEINLINE void					SetDisableIdle(bool bIsDisableIdleIn) { bIsDisableIdle = bIsDisableIdleIn; }
+	FORCEINLINE bool					IsDisableIdle() const { return bIsDisableIdle; }
 
 	FORCEINLINE	size_t					NumBatches() const
 	{
diff --git a/Engine/Plugins/Experimental/TextureGraph/Source/TextureGraphEngine/Transform/Utility/T_LoadStaticResource.h b/Engine/Plugins/Experimental/TextureGraph/Source/TextureGraphEngine/Transform/Utility/T_LoadStaticResource.h
--- a/Engine/Plugins/Experimental/TextureGraph/Source/TextureGraphEngine/Transform/Utility/T_LoadStaticResource.h
+++ b/Engine/Plugins/Experimental/TextureGraph/Source/TextureGraphEngine/Transform/Utility/T_LoadStaticResource.h
@@ -33,6 +33,7 @@ protected:
 
 	virtual cti::continuable<int32>	PreExecAsync(ENamedThreads::Type execThread, ENamedThreads::Type returnThread) override;
 	virtual cti::continuable<int32>	ExecAsync(ENamedThreads::Type execThread, ENamedThreads::Type returnThread) override;
+	virtual bool					IsAsync() const override { return true; }
 };
 
 //////////////////////////////////////////////////////////////////////////
