Fix issue with array slice/splice

diff --git a/Engine/Plugins/Experimental/TextureGraph/Source/TextureGraph/Private/Expressions/Arrays/TG_Expression_ArrayOps.cpp b/Engine/Plugins/Experimental/TextureGraph/Source/TextureGraph/Private/Expressions/Arrays/TG_Expression_ArrayOps.cpp
--- a/Engine/Plugins/Experimental/TextureGraph/Source/TextureGraph/Private/Expressions/Arrays/TG_Expression_ArrayOps.cpp
+++ b/Engine/Plugins/Experimental/TextureGraph/Source/TextureGraph/Private/Expressions/Arrays/TG_Expression_ArrayOps.cpp
@@ -53,7 +53,7 @@ void UTG_Expression_ArraySplit::Evaluate(FTG_EvaluationContext* InContext)
 	{
 		Sliced.SetNum(SlicedCount);
 		for (int i = 0; i < SlicedCount; i++)
-			Sliced.Set(i, Input.GetArray()[StartIndex + i]);
+			Sliced.Set(i + StartIndex, Input.GetArray()[StartIndex + i]);
 	}
 	else
 	{
@@ -67,8 +67,9 @@ void UTG_Expression_ArraySplit::Evaluate(FTG_EvaluationContext* InContext)
 
 		for (int i = 0; i < StartIndex; i++)
 			Spliced.Set(i, Input.GetArray()[i]);
+
 		for (int i = EndIndex; i < Input.Num(); i++)
-			Spliced.Set(StartIndex + i, Input.GetArray()[i]);
+			Spliced.Set(i - EndIndex + StartIndex, Input.GetArray()[i]);
 	}
 	else
 	{
