Fix issue with transform node not maintaining inpu resolution.

diff --git a/Engine/Plugins/Experimental/TextureGraph/Source/TextureGraph/Private/Expressions/Procedural/TG_Expression_Transform.cpp b/Engine/Plugins/Experimental/TextureGraph/Source/TextureGraph/Private/Expressions/Procedural/TG_Expression_Transform.cpp
--- a/Engine/Plugins/Experimental/TextureGraph/Source/TextureGraph/Private/Expressions/Procedural/TG_Expression_Transform.cpp
+++ b/Engine/Plugins/Experimental/TextureGraph/Source/TextureGraph/Private/Expressions/Procedural/TG_Expression_Transform.cpp
@@ -9,24 +9,8 @@ void UTG_Expression_Transform::Evaluate(FTG_EvaluationContext* InContext)
 
 
 	auto DesiredDescriptor = Output.Descriptor;
-	FVector2f spacing = {0, 0};
+	FVector2f Spacing = {0, 0};
 
-	if (!Input)
-	{
-		DesiredDescriptor.Width = DesiredDescriptor.Height = EResolution::Resolution1024;
-	}
-	else
-	{
-		if (DesiredDescriptor.Width == EResolution::Auto)
-		{
-			DesiredDescriptor.Width = EResolution::Resolution1024; 			
-		}
-		if (DesiredDescriptor.Height == EResolution::Auto)
-		{
-			DesiredDescriptor.Height = EResolution::Resolution1024;
-		}
-	}
-	FVector2f OutResolution = { float(DesiredDescriptor.Width), float( DesiredDescriptor.Height) };
 	FVector2f OutCoverage = Coverage;
 	FVector2f OutOffset = Offset;
 
@@ -48,7 +32,7 @@ void UTG_Expression_Transform::Evaluate(FTG_EvaluationContext* InContext)
 	T_Transform::CellParameter CellParam{
 		.Zoom = Zoom,
 		.StretchToFit = StretchToFit,
-		.Spacing = spacing,
+		.Spacing = Spacing,
 		.Stagger = {(StaggerOffset * StaggerHorizontally), (StaggerOffset * (1 - StaggerHorizontally))},
 		.Stride = Stride,
 	};
@@ -61,7 +45,6 @@ void UTG_Expression_Transform::Evaluate(FTG_EvaluationContext* InContext)
 		.ShowDebugGrid = ShowDebugGridValue
 	};
 	
-	Output = T_Transform::Create(InContext->Cycle, DesiredDescriptor, Input,
-		XformParam, CellParam, ColorParam, InContext->TargetId);
+	Output = T_Transform::Create(InContext->Cycle, DesiredDescriptor, Input, XformParam, CellParam, ColorParam, InContext->TargetId);
 }
 
