Fix issue with width or height settings getting lost for array inputs.

diff --git a/Engine/Plugins/Experimental/TextureGraph/Source/TextureGraph/Private/TG_GraphEvaluation.cpp b/Engine/Plugins/Experimental/TextureGraph/Source/TextureGraph/Private/TG_GraphEvaluation.cpp
--- a/Engine/Plugins/Experimental/TextureGraph/Source/TextureGraph/Private/TG_GraphEvaluation.cpp
+++ b/Engine/Plugins/Experimental/TextureGraph/Source/TextureGraph/Private/TG_GraphEvaluation.cpp
@@ -698,7 +698,7 @@ void FTG_Evaluation::EvaluateNodeArray(UTG_Node* InNode, const TArray<TObjectPtr
 			const FTG_Variant& ArrayValue = ArrayOutput->Get(0);
 
 			/// We only need to do this for textures
-			if (ArrayValue.IsTexture() && !ArrayValue.GetTexture().Descriptor.IsDefault())
+			if (ArrayValue.IsTexture())
 			{
 				const FTG_Texture& ArrayTexture = ArrayValue.GetTexture();
 				FString ToType = OutputVar.Value.Argument.GetCPPTypeName().ToString();
