Introduce additional bool options in sync versions of BP graph and bring them in line with the async versions.

diff --git a/Engine/Plugins/Experimental/TextureGraph/Source/TextureGraph/Private/Blueprint/TG_BlueprintFunctionLibrary.cpp b/Engine/Plugins/Experimental/TextureGraph/Source/TextureGraph/Private/Blueprint/TG_BlueprintFunctionLibrary.cpp
--- a/Engine/Plugins/Experimental/TextureGraph/Source/TextureGraph/Private/Blueprint/TG_BlueprintFunctionLibrary.cpp
+++ b/Engine/Plugins/Experimental/TextureGraph/Source/TextureGraph/Private/Blueprint/TG_BlueprintFunctionLibrary.cpp
@@ -273,9 +273,9 @@ const TArray<UTextureRenderTarget2D*>& UTG_BlueprintFunctionLibrary::RenderTextu
 	return Task->ActivateBlocking(nullptr);
 }
 
-void UTG_BlueprintFunctionLibrary::ExportTextureGraph(UObject* WorldContextObject, UTextureGraphBase* InTextureGraph)
+void UTG_BlueprintFunctionLibrary::ExportTextureGraph(UObject* WorldContextObject, UTextureGraphBase* InTextureGraph, bool bOverwriteTextures, bool bSave, bool bExportAll)
 {
-	UTG_AsyncExportTask* Task = UTG_AsyncExportTask::TG_AsyncExportTask(InTextureGraph, true, true, true);
+	UTG_AsyncExportTask* Task = UTG_AsyncExportTask::TG_AsyncExportTask(InTextureGraph, bOverwriteTextures, bSave, bExportAll);
 	Task->ActivateBlocking(nullptr);
 }
 
diff --git a/Engine/Plugins/Experimental/TextureGraph/Source/TextureGraph/Public/Blueprint/TG_BlueprintFunctionLibrary.h b/Engine/Plugins/Experimental/TextureGraph/Source/TextureGraph/Public/Blueprint/TG_BlueprintFunctionLibrary.h
--- a/Engine/Plugins/Experimental/TextureGraph/Source/TextureGraph/Public/Blueprint/TG_BlueprintFunctionLibrary.h
+++ b/Engine/Plugins/Experimental/TextureGraph/Source/TextureGraph/Public/Blueprint/TG_BlueprintFunctionLibrary.h
@@ -86,7 +86,7 @@ public:
 
 	/** Export texture graph. */
 	UFUNCTION(BlueprintCallable, Category = "TextureGraph", meta = (Keywords = "ExportTextureGraph", WorldContext = "WorldContextObject"))
-	static void ExportTextureGraph(UObject* WorldContextObject, UTextureGraphBase* InTextureGraph);
+	static void ExportTextureGraph(UObject* WorldContextObject, UTextureGraphBase* InTextureGraph, bool bOverwriteTextures = true, bool bSave = false, bool bExportAll = false);
 
 	static void AddParamWarning(FName ParamName, UObject* ObjectPtr, FString FunctionName);
 	static void AddError(UObject* ObjectPtr, FString FunctionName, FString Error);
