Texture Graph, Levels Expression issues:
- UE-216885  Fixing the direction of correction of the Mid value: the ratio was inverted
- UE-218476 Levels applied to color image is working on the Grayscaled value and then scale back the whole color of the input image

diff --git a/Engine/Plugins/Experimental/TextureGraph/Shaders/Expressions/Expression_Levels.usf b/Engine/Plugins/Experimental/TextureGraph/Shaders/Expressions/Expression_Levels.usf
--- a/Engine/Plugins/Experimental/TextureGraph/Shaders/Expressions/Expression_Levels.usf
+++ b/Engine/Plugins/Experimental/TextureGraph/Shaders/Expressions/Expression_Levels.usf
@@ -31,21 +31,17 @@ float MidPercentage;
 Texture2D Histogram;
 
 
-float LevelsControlInputRange(float Value, float LowInput, float HighInput)
-{
-	return min(max(Value - LowInput, float(0)) / (float(HighInput) - LowInput), float(1));
-}
-
 float LevelsControlInputCurve(float Value, float LowInput, float HighInput, float MidInput)
 {
-	float ValueRanged = LevelsControlInputRange(Value, LowInput, HighInput);
+	// Map incoming value to input range
+	float ValueRanged = saturate((Value - LowInput) / (HighInput - LowInput));
 
-	// 0.5 = EvalRange(Mid) ^ Exponent
+	// Map mid to iput range within an epsilon
+	float MidRanged = clamp((MidInput - LowInput) / (HighInput - LowInput), 0.001, 0.9999);
+	//And evaluate the MidExponent:
+	// 0.5 = EvalRange(Mid) ^ MidExponent
 	// thus
-	float MidRanged = saturate((MidInput - LowInput) / (HighInput - LowInput));
-	MidRanged = max(0.001, min(0.9999, MidRanged));
-
-	float InvGammaValue = log(MidRanged) / log(0.5);
+	float MidExponent = log(0.5) / log(MidRanged);
 	
 #if OUT_LEVELS
 	float4 HistoMin			= Histogram.Load(int3(2, 1, 0));
@@ -59,9 +55,7 @@ float LevelsControlInputCurve(float Value, float LowInput, float HighInput, floa
 	ValueRanged = OutLow + ((ValueRanged - LowInput) / InputRange) * OutputRange;
 #endif /// OUT_LEVELS
 
-
-
-	return pow(ValueRanged, InvGammaValue);
+	return pow(ValueRanged, MidExponent);
 }
 
 float4 FSH_Levels(float2 uv : TEXCOORD0) : SV_Target0
@@ -79,7 +73,6 @@ float4 FSH_Levels(float2 uv : TEXCOORD0) : SV_Target0
 	// Standard path, Low and High come from specified uniforms
 	float LowInput = LowValue;
 	float HighInput = HighValue;
-
 	float MidInput = MidValue;
 
 #if AUTO_LEVELS
@@ -93,8 +86,16 @@ float4 FSH_Levels(float2 uv : TEXCOORD0) : SV_Target0
 	HighInput = HistoMax.w;
 #endif /// AUTO_LEVELS
 
-	float3 LeveledColor = LevelsControlInputCurve(SourceValue, LowInput, HighInput, MidValue);
-	
+	float LeveledValue = LevelsControlInputCurve(SourceValue, LowInput, HighInput, MidValue);
+
+#if CONVERT_TO_GRAYSCALE
+	// The new  Grayscale level is used to scale the color darker or brighter
+	float Ratio = (SourceValue == 0 ? 0.0 : LeveledValue / SourceValue); 	
+	float3 LeveledColor = Color * Ratio;
+#else
+	float3 LeveledColor = LeveledValue;
+#endif
+
 	return float4(LeveledColor.rgb, 1.0);
 }
 
