Fixed modify array issues

diff --git a/Engine/Plugins/TextureGraph/Source/TextureGraph/Private/Expressions/Arrays/TG_Expression_ArrayOps.cpp b/Engine/Plugins/TextureGraph/Source/TextureGraph/Private/Expressions/Arrays/TG_Expression_ArrayOps.cpp
--- a/Engine/Plugins/TextureGraph/Source/TextureGraph/Private/Expressions/Arrays/TG_Expression_ArrayOps.cpp
+++ b/Engine/Plugins/TextureGraph/Source/TextureGraph/Private/Expressions/Arrays/TG_Expression_ArrayOps.cpp
@@ -53,7 +53,7 @@ void UTG_Expression_ArraySplit::Evaluate(FTG_EvaluationContext* InContext)
 	{
 		Sliced.SetNum(SlicedCount);
 		for (int i = 0; i < SlicedCount; i++)
-			Sliced.Set(i + StartIndex, Input.GetArray()[StartIndex + i]);
+			Sliced.Set(i, Input.GetArray()[StartIndex + i]);
 	}
 	else
 	{
diff --git a/Engine/Plugins/TextureGraph/Source/TextureGraphEditor/Private/EdGraph/TG_EdGraphNode.cpp b/Engine/Plugins/TextureGraph/Source/TextureGraphEditor/Private/EdGraph/TG_EdGraphNode.cpp
--- a/Engine/Plugins/TextureGraph/Source/TextureGraphEditor/Private/EdGraph/TG_EdGraphNode.cpp
+++ b/Engine/Plugins/TextureGraph/Source/TextureGraphEditor/Private/EdGraph/TG_EdGraphNode.cpp
@@ -412,14 +412,24 @@ FString UTG_EdGraphNode::GetTitleDetail()
 		{
 			if (Texture.RasterBlob)
 			{
+				FString ArrayInfo;
+
+				if (Pin->IsArgTextureArray())
+				{
+					FTG_VariantArray VarArray;
+					Pin->GetValue(VarArray);
+					ArrayInfo = FString::Printf(TEXT("[%d]"), VarArray.Num());
+				}
+
 				const auto Desc = Texture.RasterBlob->GetDescriptor();
 				FStringFormatNamedArguments Args;
 				Args.Add(TEXT("Channels"), TextureHelper::GetChannelsTextFromItemsPerPoint(Desc.ItemsPerPoint));
 				Args.Add(TEXT("Format"), Desc.FormatToString(Desc.Format));
 				Args.Add(TEXT("IsSRGB"), Desc.bIsSRGB ? "sRGB" : "Linear");
+				Args.Add(TEXT("ArrayInfo"), ArrayInfo);
 				Args.Add(TEXT("Width"), FString::FromInt(Desc.Width));
 				Args.Add(TEXT("Height"), FString::FromInt(Desc.Height));
-				Details = FString::Format(TEXT("{Channels}_{Format}, {IsSRGB}\n{Width}x{Height}"), Args);
+				Details = FString::Format(TEXT("{Channels}_{Format}, {IsSRGB} {ArrayInfo}\n{Width}x{Height}"), Args);
 				break;
 			}
 		}
