Angle seed issue with pattern node

diff --git a/Engine/Plugins/TextureGraph/Shaders/Mask/PatternMask.usf b/Engine/Plugins/TextureGraph/Shaders/Mask/PatternMask.usf
--- a/Engine/Plugins/TextureGraph/Shaders/Mask/PatternMask.usf
+++ b/Engine/Plugins/TextureGraph/Shaders/Mask/PatternMask.usf
@@ -86,7 +86,7 @@ float GetBrightnessJitter(float2 jitterUVs)
 
 float GetAngleTiltJitter(float2 jitterUVs, float2 pattern_cell_uv)
 {
-	float jaVariance = SFNoise(jitterUVs, JT_Seed + 0.2);
+	float jaVariance = SFNoise(jitterUVs, JA_Seed + 0.2);
 	jaVariance = (lerp(JA_Amount.x, JA_Amount.y, jaVariance) - 0.5f) * 2 * PI;
 				//jaVariance = lerp(((_JA_Amount.x + _JA_Amount.y) / 2.0 - 0.5f) * 2 * PI, jaVariance, jaThreshold);
 
@@ -96,8 +96,8 @@ float GetAngleTiltJitter(float2 jitterUVs, float2 pattern_cell_uv)
 	angledUV = (angledUV + 1.0) / 2.0;
 
 				// Tilt Jitter
-	float jtValue = SFNoise(jitterUVs, JT_Seed);
-	float jtThreshold = 1 - SFNoise(jitterUVs, JT_Seed + 0.1);
+	float jtValue = SFNoise(jitterUVs, JA_Seed);
+	float jtThreshold = 1 - SFNoise(jitterUVs, JA_Seed + 0.1);
 	jtThreshold = step(jtThreshold, JT_Threshold);
 	jtValue = Remap(jtValue, 0, 1, JT_Amount.x, JT_Amount.y) * jtThreshold * angledUV.x;
 
@@ -235,7 +235,7 @@ float4 FSH_PatternMask(float2 uv : TEXCOORD0) : SV_Target0
 	//float jaThreshold = SFNoise(jitterUVs, _JA_Seed + 0.1);
 	//jaThreshold = step(jaThreshold, _JA_Threshold);
 
-	float jaVariance = SFNoise(jitterUVs, JT_Seed + 0.2);
+	float jaVariance = SFNoise(jitterUVs, JA_Seed + 0.2);
 	jaVariance = (lerp(JA_Amount.x, JA_Amount.y, jaVariance) - 0.5f) * 2 * PI;
 	//jaVariance = lerp(((_JA_Start + _JA_End) / 2.0 - 0.5f) * 2 * PI, jaVariance, jaThreshold);
 
