Disallow connecting array nodes directly to output node

diff --git a/Engine/Plugins/TextureGraph/Source/TextureGraph/Private/TG_Graph.cpp b/Engine/Plugins/TextureGraph/Source/TextureGraph/Private/TG_Graph.cpp
--- a/Engine/Plugins/TextureGraph/Source/TextureGraph/Private/TG_Graph.cpp
+++ b/Engine/Plugins/TextureGraph/Source/TextureGraph/Private/TG_Graph.cpp
@@ -540,12 +540,15 @@ bool UTG_Graph::ArePinsCompatible(const UTG_Pin* PinFrom, const UTG_Pin* PinTo,
 {
 	if (PinFrom && PinTo)
 	{
-		if (PinFrom->GetNodePtr() != PinTo->GetNodePtr())
+		UTG_Node* NodeFrom = PinFrom->GetNodePtr();
+		UTG_Node* NodeTo = PinTo->GetNodePtr();
+
+		/// We don't allow connecting an array parameter to an output node. Too many complications with subgraphs etc
+		if (NodeFrom != NodeTo && !(PinFrom->IsArgArray() && Cast<UTG_Expression_Output>(NodeTo->GetExpression())))
 		{
-			if (PinFrom->GetNodePtr()->GetGraph() == PinTo->GetNodePtr()->GetGraph())
+			if (NodeFrom->GetGraph() == NodeTo->GetGraph())
 			{
-				return FTG_Evaluation::AreArgumentsCompatible(PinFrom->GetArgument(), PinTo->GetArgument(),
-				                                              ConverterKey);
+				return FTG_Evaluation::AreArgumentsCompatible(PinFrom->GetArgument(), PinTo->GetArgument(), ConverterKey);
 			}
 		}
 	}
