Fixed array placement issues.

diff --git a/Engine/Plugins/TextureGraph/Source/TextureGraphEngine/Device/FX/Device_FX.cpp b/Engine/Plugins/TextureGraph/Source/TextureGraphEngine/Device/FX/Device_FX.cpp
--- a/Engine/Plugins/TextureGraph/Source/TextureGraphEngine/Device/FX/Device_FX.cpp
+++ b/Engine/Plugins/TextureGraph/Source/TextureGraphEngine/Device/FX/Device_FX.cpp
@@ -329,6 +329,8 @@ AsyncDeviceBufferRef Device_FX::DrawTilesToBuffer_Deferred(DeviceBufferRef Buffe
 
 				for (int32 RowId = 0; RowId < Tiles.Rows(); RowId++)
 				{
+					uint32 XOffset = 0;
+
 					for (int32 ColId = 0; ColId < Tiles.Cols(); ColId++)
 					{
 						DeviceBuffer_FX* TileBuffer = static_cast<DeviceBuffer_FX*>(Tiles(RowId, ColId).get());
@@ -344,11 +346,14 @@ AsyncDeviceBufferRef Device_FX::DrawTilesToBuffer_Deferred(DeviceBufferRef Buffe
 						const float SrcTextureWidth = TileTex->GetWidth();
 						const float SrcTextureHeight = TileTex->GetHeight();
 
-				//// This was here but shouldn't be needed if default state is SRV 
-				///		DevRHI.Transition(FRHITransitionInfo(TileResource, ERHIAccess::Unknown, ERHIAccess::SRVGraphics));
+						uint32 X1 = XOffset; ///ColId * SrcTextureWidth;
+						uint32 Y1 = 0; ///RowId * SrcTextureHeight;
 
-						int32 X1 = ColId * SrcTextureWidth;
-						int32 Y1 = RowId * SrcTextureHeight;
+						for (int32 TAIndex = 0; TAIndex < RowId; TAIndex++)
+						{
+							DeviceBuffer_FX* TileAbove = static_cast<DeviceBuffer_FX*>(Tiles(TAIndex, ColId).get());
+							Y1 += TileAbove->Descriptor().Height;
+						}
 
 						DevRHI.SetViewport(X1, Y1, 0, X1 + ViewportWidth, Y1 + ViewportHeight, 0);
 
@@ -365,6 +370,7 @@ AsyncDeviceBufferRef Device_FX::DrawTilesToBuffer_Deferred(DeviceBufferRef Buffe
 							VertexShader,
 							EDRF_Default);
 
+						XOffset += SrcTextureWidth;
 					}
 				}
 			}
diff --git a/Engine/Plugins/TextureGraph/Source/TextureGraphEngine/Transform/Expressions/T_ArrayGrid.cpp b/Engine/Plugins/TextureGraph/Source/TextureGraphEngine/Transform/Expressions/T_ArrayGrid.cpp
--- a/Engine/Plugins/TextureGraph/Source/TextureGraphEngine/Transform/Expressions/T_ArrayGrid.cpp
+++ b/Engine/Plugins/TextureGraph/Source/TextureGraphEngine/Transform/Expressions/T_ArrayGrid.cpp
@@ -41,6 +41,7 @@ public:
 
 	static BlobPtrTiles				GetTiles(TArray<TiledBlobPtr>& Inputs, int32 InNumRows, int32 InNumCols, 
 		int32 NumRows, int32 NumCols, uint32* OutMaxRowHeight = nullptr, uint32* OutMaxColWidth = nullptr, 
+		uint32* OutWidth = nullptr, uint32 *OutHeight = nullptr,
 		bool* bIsWellDefinedOut = nullptr, CHashPtrVec* OutHashes = nullptr)
 	{
 		BlobPtrTiles ResultTiles(InNumRows * NumRows, InNumCols * NumCols); 
@@ -48,12 +49,15 @@ public:
 		bool bIsWellDefined = true;
 		uint32 MaxRowHeight = 0;
 		uint32 MaxColWidth = 0;
+		uint32 Width = 0;
+		uint32 Height = 0;
 		size_t HashIndex = 0;
 
 		/// First of all we combine all the tiles into a single result
 		for (int32 ArrRowIter = 0; ArrRowIter < NumRows; ArrRowIter++)
 		{
 			int32 RowIndex = ArrRowIter * InNumCols;
+			uint32 RowWidth = 0;
 			for (int32 ArrColIter = 0; ArrColIter < NumCols; ArrColIter++)
 			{
 				int32 TileIndex = ArrRowIter * NumCols + ArrColIter;
@@ -63,6 +67,7 @@ public:
 
 				MaxRowHeight = std::max(MaxRowHeight, Input->GetHeight());
 				MaxColWidth = std::max(MaxColWidth, Input->GetWidth());
+				RowWidth += Input->GetWidth();
 
 				check(HashIndex < InputHashes.size());
 				check(Input->Hash());
@@ -89,8 +94,15 @@ public:
 					}
 				}
 			}
+
+			Width = std::max(Width, RowWidth);
+			Height += MaxRowHeight;
 		}
 
+		if (OutWidth)
+			*OutWidth = Width;
+		if (OutHeight)
+			*OutHeight = Height;
 		if (OutMaxRowHeight)
 			*OutMaxRowHeight = MaxRowHeight;
 		if (OutMaxColWidth)
@@ -108,7 +120,8 @@ protected:
 	{
 		int32 InNumRows = Inputs[0]->Rows();
 		int32 InNumCols = Inputs[0]->Cols();
-		BlobPtrTiles ResultTiles = GetTiles(Inputs, InNumRows, InNumCols, NumRows, NumCols, &MaxRowHeight, &MaxColWidth);
+		uint32 Width = 0, Height = 0;
+		BlobPtrTiles ResultTiles = GetTiles(Inputs, InNumRows, InNumCols, NumRows, NumCols, &MaxRowHeight, &MaxColWidth, &Width, &Height);
 
 		/// The we combine those tiles into a single result
 		const BufferDescriptor& InputDesc = Inputs[0]->GetDescriptor();
@@ -117,8 +130,8 @@ protected:
 		if (DesiredOutputDesc.IsLateBound())
 			DesiredOutputDesc.Format = InputDesc.Format;
 
-		DesiredOutputDesc.Width = MaxColWidth * NumCols;
-		DesiredOutputDesc.Height = MaxRowHeight * NumRows;
+		DesiredOutputDesc.Width = Width;
+		DesiredOutputDesc.Height = Height;
 
 		TiledBlobPtr CombinedResultTiled = TiledBlob::InitFromTiles(DesiredOutputDesc, ResultTiles);
 
@@ -153,8 +166,12 @@ TiledBlobPtr T_ArrayGrid::Create(MixUpdateCyclePtr Cycle, BufferDescriptor& InDe
 	bool bIsWellDefined = true;
 	uint32 MaxRowHeight = 0;
 	uint32 MaxColWidth = 0;
+	uint32 Width = 0;
+	uint32 Height = 0;
 
-	BlobPtrTiles ResultTiles = Job_ArrayGrid::GetTiles(Inputs, InNumRows, InNumCols, NumRows, NumCols, &MaxColWidth, &MaxRowHeight, &bIsWellDefined, &InputHashes);
+	BlobPtrTiles ResultTiles = Job_ArrayGrid::GetTiles(Inputs, InNumRows, InNumCols, NumRows, NumCols, 
+		&MaxColWidth, &MaxRowHeight, &Width, &Height, 
+		&bIsWellDefined, &InputHashes);
 	TiledBlobPtr FinalResult;
 
 	/// If the result is late bound then we have no option but to resize the result render target through a custom job
@@ -176,8 +193,8 @@ TiledBlobPtr T_ArrayGrid::Create(MixUpdateCyclePtr Cycle, BufferDescriptor& InDe
 	else
 	{
 		BufferDescriptor DesiredOutputDesc = BufferDescriptor::Combine(Inputs[0]->GetDescriptor(), InDesiredOutputDesc);
-		DesiredOutputDesc.Width = MaxColWidth * NumCols;
-		DesiredOutputDesc.Height = MaxRowHeight * NumRows;
+		DesiredOutputDesc.Width = Width;
+		DesiredOutputDesc.Height = Height;
 
 		/// The we combine those tiles into a single result
 		TiledBlobPtr CombinedResultTiled = TiledBlob::InitFromTiles(DesiredOutputDesc, ResultTiles);
