Prevent input texture and texture path nodes from reloading the textures at every update

diff --git a/Engine/Plugins/TextureGraph/Source/TextureGraph/Private/Expressions/Input/TG_Expression_Texture.cpp b/Engine/Plugins/TextureGraph/Source/TextureGraph/Private/Expressions/Input/TG_Expression_Texture.cpp
--- a/Engine/Plugins/TextureGraph/Source/TextureGraph/Private/Expressions/Input/TG_Expression_Texture.cpp
+++ b/Engine/Plugins/TextureGraph/Source/TextureGraph/Private/Expressions/Input/TG_Expression_Texture.cpp
@@ -17,15 +17,20 @@ void UTG_Expression_Texture::Evaluate(FTG_EvaluationContext* InContext)
 	else if (Source)
 	{
 		const FString Path = Source->GetPathName();
-		UStaticImageResource* StaticImageResource = UStaticImageResource::CreateNew<UStaticImageResource>();
-		StaticImageResource->SetAssetUUID(Path);
-
-		//Until we have srgb value exposed in the UI we need to set the Srgb of the Output Descriptor here from the source
-		//This gets updated for the late bond case but since we do not have the UI to specify the override in other nodes 
-		// the override value will always be set to false while combining the buffers
-		auto DesiredDesc = Output.GetBufferDescriptor();
-		DesiredDesc.bIsSRGB = Source->SRGB;
-		Output = StaticImageResource->GetBlob(InContext->Cycle, &DesiredDesc, 0);
+
+		if (!Output || OutputPath != Path)
+		{
+			UStaticImageResource* StaticImageResource = UStaticImageResource::CreateNew<UStaticImageResource>();
+			StaticImageResource->SetAssetUUID(Path);
+
+			//Until we have srgb value exposed in the UI we need to set the Srgb of the Output Descriptor here from the source
+			//This gets updated for the late bond case but since we do not have the UI to specify the override in other nodes 
+			// the override value will always be set to false while combining the buffers
+			auto DesiredDesc = Output.GetBufferDescriptor();
+			DesiredDesc.bIsSRGB = Source->SRGB;
+			Output = StaticImageResource->GetBlob(InContext->Cycle, &DesiredDesc, 0);
+			OutputPath = Path;
+		}
 	}
 	else
 	{
diff --git a/Engine/Plugins/TextureGraph/Source/TextureGraph/Private/Expressions/Input/TG_Expression_TexturePath.cpp b/Engine/Plugins/TextureGraph/Source/TextureGraph/Private/Expressions/Input/TG_Expression_TexturePath.cpp
--- a/Engine/Plugins/TextureGraph/Source/TextureGraph/Private/Expressions/Input/TG_Expression_TexturePath.cpp
+++ b/Engine/Plugins/TextureGraph/Source/TextureGraph/Private/Expressions/Input/TG_Expression_TexturePath.cpp
@@ -77,35 +77,41 @@ void UTG_Expression_TexturePath::Evaluate(FTG_EvaluationContext* InContext)
 	FString LocalPath;
 	bool bValidation = ValidateInputPath(LocalPath);
 
-	if (!Path.IsEmpty() && bValidation)
+	if (Output.Num() == 0 || LocalPath != OutputPath)
 	{
-		FPaths::MakePlatformFilename(LocalPath);
-		if (FPaths::FileExists(LocalPath))
+		if (!Path.IsEmpty() && bValidation)
 		{
-			Output.SetNum(1);
-			FTG_Texture Item = LoadStaticImage(InContext, LocalPath, nullptr);
-			Output.Set(0, Item);
-		}
-		else if (FPaths::DirectoryExists(LocalPath))
-		{
-			FFileManagerGeneric FileManager;
-			TArray<FString> Files;
-			FileManager.FindFiles(Files, *LocalPath);
-
-			if (!Files.IsEmpty())
+			FPaths::MakePlatformFilename(LocalPath);
+			if (FPaths::FileExists(LocalPath))
 			{
-				Output.SetNum(Files.Num());
+				Output.SetNum(1);
+				FTG_Texture Item = LoadStaticImage(InContext, LocalPath, nullptr);
+				Output.Set(0, Item);
+				OutputPath = LocalPath;
+			}
+			else if (FPaths::DirectoryExists(LocalPath))
+			{
+				FFileManagerGeneric FileManager;
+				TArray<FString> Files;
+				FileManager.FindFiles(Files, *LocalPath);
 
-				/// Now load a texture against all the files
-				for (int32 Index = 0; Index < Files.Num(); Index++)
+				if (!Files.IsEmpty())
 				{
-					const FString& File = Files[Index];
-					FString Filename = FPaths::Combine(LocalPath, File);
-					FTG_Texture Item = LoadStaticImage(InContext, Filename, nullptr);
-					Output.Set(Index, Item);
+					Output.SetNum(Files.Num());
+
+					/// Now load a texture against all the files
+					for (int32 Index = 0; Index < Files.Num(); Index++)
+					{
+						const FString& File = Files[Index];
+						FString Filename = FPaths::Combine(LocalPath, File);
+						FTG_Texture Item = LoadStaticImage(InContext, Filename, nullptr);
+						Output.Set(Index, Item);
+					}
 				}
 			}
 		}
+
+		OutputPath = LocalPath;
 	}
 
 	//For the connected pin we will report error here in eveluate because it does not have the updated value during validation.
diff --git a/Engine/Plugins/TextureGraph/Source/TextureGraph/Public/Expressions/Input/TG_Expression_Texture.h b/Engine/Plugins/TextureGraph/Source/TextureGraph/Public/Expressions/Input/TG_Expression_Texture.h
--- a/Engine/Plugins/TextureGraph/Source/TextureGraph/Public/Expressions/Input/TG_Expression_Texture.h
+++ b/Engine/Plugins/TextureGraph/Source/TextureGraph/Public/Expressions/Input/TG_Expression_Texture.h
@@ -26,6 +26,7 @@ public:
 	// The output of the node, which is the loaded texture asset
 	UPROPERTY(meta = (TGType = "TG_Output", PinDisplayName = "", HideInnerPropertiesInNode))
 	FTG_Texture Output;
+	FString OutputPath;
 
 	// The source asset to be used to generate the Output
 	UPROPERTY(EditAnywhere, Setter, BlueprintReadWrite, Category = NoCategory, meta = (TGType = "TG_Input", TGPinNotConnectable) )
diff --git a/Engine/Plugins/TextureGraph/Source/TextureGraph/Public/Expressions/Input/TG_Expression_TexturePath.h b/Engine/Plugins/TextureGraph/Source/TextureGraph/Public/Expressions/Input/TG_Expression_TexturePath.h
--- a/Engine/Plugins/TextureGraph/Source/TextureGraph/Public/Expressions/Input/TG_Expression_TexturePath.h
+++ b/Engine/Plugins/TextureGraph/Source/TextureGraph/Public/Expressions/Input/TG_Expression_TexturePath.h
@@ -35,6 +35,7 @@ public:
 	// All the output textures from this node. Note that it outputs an array of textures.
 	UPROPERTY(meta = (TGType = "TG_Output", PinDisplayName = "", HideInnerPropertiesInNode))
 	FTG_VariantArray Output;
+	FString OutputPath;
 
 	// Input file path of the texture. If the path is a directory then all files from that directory are loaded
 	UPROPERTY(EditAnywhere, BlueprintReadWrite, Category = NoCategory, meta = (TGType = "TG_InputParam", NoResetToDefault) )
