Fix flickering of 2D viewport and update it more frequently.

diff --git a/Engine/Plugins/TextureGraph/Source/TextureGraphEditor/Private/STG_NodePreview.cpp b/Engine/Plugins/TextureGraph/Source/TextureGraphEditor/Private/STG_NodePreview.cpp
--- a/Engine/Plugins/TextureGraph/Source/TextureGraphEditor/Private/STG_NodePreview.cpp
+++ b/Engine/Plugins/TextureGraph/Source/TextureGraphEditor/Private/STG_NodePreview.cpp
@@ -340,21 +340,40 @@ void STG_NodePreviewWidget::NodeDeleted(const UTG_Node* Node)
 
 void STG_NodePreviewWidget::Update() const
 {
-	// Lambda for actually updating the preview.
-	auto UpdatePreview = [this]
+	FTG_Variant Variant;
+	BlobPtr Blob;
+	bool bValidVariant = GetOutputVariantFromNode(Variant);
+
+	if (Variant)
 	{
-		FTG_Variant Variant;
-		const bool bValidVariant = GetOutputVariantFromNode(Variant);
+		FIntPoint ImageSize = FIntPoint::ZeroValue;
 
-		if(Variant.IsTexture())
+		if (Variant.IsTexture() && Variant.GetTexture() && Variant.GetTexture().RasterBlob)
 		{
-			if(Variant.GetTexture() && Variant.GetTexture()->IsValid())
+			Blob = Variant.GetTexture().RasterBlob;
+			ImageSize = FIntPoint {(int32)Blob->GetWidth(), (int32)Blob->GetHeight()};
+
+			if (Blob->IsTiled())
 			{
-				NodeViewer->SetTexture(Variant.GetTexture().RasterBlob);
+				Blob->OnFinalise()
+					.then([Blob]()
+						{
+							UE_LOG(LogTemp, Warning, TEXT("!!!Finalised. Combining!!!"));
+							const TiledBlobPtr BlobTiled = std::static_pointer_cast<TiledBlob>(Blob);
+							return BlobTiled->CombineTiles(false, false);
+						})
+					.then([this, Blob]()
+						{
+							UE_LOG(LogTemp, Warning, TEXT("!!!SETTING TEXTURE!!!"));
+							NodeViewer->SetTexture(Blob);
+						});
 			}
 			else
 			{
-				NodeViewer->SetTexture(nullptr);
+				Blob->OnFinalise().then([this, Blob]()
+					{
+						NodeViewer->SetTexture(Blob);
+					});
 			}
 		}
 		else if(Variant.IsColor())
@@ -365,49 +384,16 @@ void STG_NodePreviewWidget::Update() const
 		{
 			NodeViewer->SetTexture(nullptr);
 		}
-		
-		NodeViewer->SetLabelText(GetLabelText(Variant, bValidVariant));
 
-		Viewport->ResetZoom(NodeViewer->GetCurrentImageInfo().Size);
-	};
-
-	FTG_Variant Variant;
-
-	GetOutputVariantFromNode(Variant);
-	BlobPtr Blob;
-	
-	if(Variant && Variant.IsTexture() && Variant.GetTexture())
-	{
-		Blob = Variant.GetTexture().RasterBlob;
+		NodeViewer->SetLabelText(GetLabelText(Variant, bValidVariant));
+		Viewport->ResetZoom(ImageSize);
 
-		if(Blob)
-		{
-			if (Blob->IsTiled())
-			{
-				Blob->OnFinalise()
-					.then([Blob]()
-					{
-						const TiledBlobPtr BlobTiled = std::static_pointer_cast<TiledBlob>(Blob);
-						return BlobTiled->CombineTiles(false, false);
-					})
-					.then(UpdatePreview);
-			}
-			else
-			{
-				Blob->OnFinalise()
-					.then(UpdatePreview);
-			}
-		}
-		else
-		{
-			UpdatePreview();
-		}
-	}
-	else
-	{
-		UpdatePreview();
+		return;
 	}
 
+	/// If we get to here then we need to set it to null
+	NodeViewer->SetTexture(nullptr);
+
 	// Update preview blob and trigger related external updates.
 	[[maybe_unused]] bool Result = OnNodeBlobChanged.ExecuteIfBound(Blob);
 }
diff --git a/Engine/Plugins/TextureGraph/Source/TextureGraphEditor/Private/TG_Editor.cpp b/Engine/Plugins/TextureGraph/Source/TextureGraphEditor/Private/TG_Editor.cpp
--- a/Engine/Plugins/TextureGraph/Source/TextureGraphEditor/Private/TG_Editor.cpp
+++ b/Engine/Plugins/TextureGraph/Source/TextureGraphEditor/Private/TG_Editor.cpp
@@ -190,6 +190,7 @@ void FTG_Editor::InitEditor(const EToolkitMode::Type Mode, const TSharedPtr< cla
 
 	//Editor gets notified when rendering is done
 	EditedTextureGraph->OnRenderDone.BindRaw(this, &FTG_Editor::OnRenderingDone);
+	EditedTextureGraph->OnBatchQueued.BindRaw(this, &FTG_Editor::OnBatchQueued);
 
 	FViewportSettings& ViewportSettings = EditedTextureGraph->GetSettings()->GetViewportSettings();
 
@@ -370,6 +371,7 @@ FTG_Editor::~FTG_Editor()
 {
 	// cleanup events
 	EditedTextureGraph->OnRenderDone.Unbind();
+	EditedTextureGraph->OnBatchQueued.Unbind();
 	//EditedTextureGraph->Graph()->OnGraphChangedDelegate.Remove();
 
 	EditedTextureGraph = nullptr;
@@ -634,6 +636,13 @@ void FTG_Editor::OnRunGraph_Clicked()
 	}
 }
 void FTG_Editor::OnRenderingDone(UMixInterface* TextureGraph, const FInvalidationDetails* Details)
+{
+	if (TextureGraph != nullptr)
+	{
+	}
+}
+
+void FTG_Editor::OnBatchQueued(UMixInterface* TextureGraph, const FInvalidationDetails* Details)
 {
 	if (TextureGraph != nullptr)
 	{
diff --git a/Engine/Plugins/TextureGraph/Source/TextureGraphEditor/Private/TG_Editor.h b/Engine/Plugins/TextureGraph/Source/TextureGraphEditor/Private/TG_Editor.h
--- a/Engine/Plugins/TextureGraph/Source/TextureGraphEditor/Private/TG_Editor.h
+++ b/Engine/Plugins/TextureGraph/Source/TextureGraphEditor/Private/TG_Editor.h
@@ -169,6 +169,7 @@ private:
 	bool											GetTabSelected(const FName TabID);
 
 	void											OnRenderingDone(UMixInterface* TextureGraph, const FInvalidationDetails* Details);
+	void											OnBatchQueued(UMixInterface* TextureGraph, const FInvalidationDetails* Details);
 	void											OnPreviewMeshChangedEvent();
 	void											OnViewportMaterialChanged();
 	void											OnMaterialMappingChanged();
diff --git a/Engine/Plugins/TextureGraph/Source/TextureGraphEngine/Job/Scheduler.cpp b/Engine/Plugins/TextureGraph/Source/TextureGraphEngine/Job/Scheduler.cpp
--- a/Engine/Plugins/TextureGraph/Source/TextureGraphEngine/Job/Scheduler.cpp
+++ b/Engine/Plugins/TextureGraph/Source/TextureGraphEngine/Job/Scheduler.cpp
@@ -88,51 +88,48 @@ void Scheduler::Update(float dt)
 		FScopeLock Lock(&BatchMutex);
 		if (Batches.size())
 		{
-			/// TODO: Enable this at some point
-			if (Batches.size() == 1)
-			{
-				BatchToRun = Batches.front();
-				Batches.pop_front();
-			}
-			else
-			{
-				check(!TextureGraphEngine::IsTestMode());
+			BatchToRun = Batches.front();
+			Batches.pop_front();
 
+			/// We merge batches if we're lagging behind ...
+			if (!TextureGraphEngine::IsTestMode() && Batches.size() >= 8)
+			{
 				/// OK here, we're going to merge batches belonging to the same mix together
 				
 				// figure out the prioritized mix from the Batch
 				BatchToRun = Batches.front();
-				UMixInterface* currentBatchMix = BatchToRun->GetCycle()->GetMix();
+				Batches.pop_front();
+				UMixInterface* CurrentBatchMix = BatchToRun->GetCycle()->GetMix();
 
 				// find the latest Batch of required mix
 				auto Iter = Batches.begin();
-				JobBatchPtr lastMixbatch = BatchToRun;
+				JobBatchPtr LastMixBatch = BatchToRun;
 
 				while (Iter != Batches.end())
 				{
-					JobBatchPtr futureBatch = *Iter;
+					JobBatchPtr FutureBatch = *Iter;
 
-					if (futureBatch->GetCycle()->GetMix() == currentBatchMix)
+					if (FutureBatch->GetCycle()->GetMix() == CurrentBatchMix)
 					{
-						lastMixbatch = *Iter;
+						LastMixBatch = *Iter;
 					}
 
 					Iter++;
 				}
 
 				// we get latest Batch to run for the mix
-				BatchToRun = lastMixbatch;
+				BatchToRun = LastMixBatch;
 
 				// Merge invalidation details of the mix from multiple batches.
 				// Remove the batches from the list after merging.
 				Iter = Batches.begin();
 				while (Iter != Batches.end())
 				{
-					JobBatchPtr futureBatch = *Iter;
+					JobBatchPtr FutureBatch = *Iter;
 
-					if (futureBatch->GetCycle()->GetMix() == currentBatchMix)
+					if (FutureBatch->GetCycle()->GetMix() == CurrentBatchMix)
 					{
-						BatchToRun->GetCycle()->MergeDetails(futureBatch->GetCycle()->GetDetails());
+						BatchToRun->GetCycle()->MergeDetails(FutureBatch->GetCycle()->GetDetails());
 						Iter = Batches.erase(Iter);
 					}
 					else
@@ -175,6 +172,7 @@ void Scheduler::Update(float dt)
 
 				ObserverSource->BatchJobsDone(PreviousBatch);
 
+				UE_LOG(LogBatch, Warning, TEXT("!!!BATCH done!!!"));
 				/// Mark the current batch as done
 				PreviousBatch->GetCycle()->GetDetails().BroadcastOnDone();
 
@@ -345,6 +343,7 @@ void Scheduler::AddBatch(JobBatchPtr Batch)
 	Batches.push_back(Batch);
 	ObserverSource->BatchAdded(Batch); // notify observer
 	TimeSinceIdle = Util::Time();
+	Batch->GetCycle()->GetDetails().BroadcastOnQueued();
 }
 
 void Scheduler::Start()
diff --git a/Engine/Plugins/TextureGraph/Source/TextureGraphEngine/Model/Mix/MixInterface.cpp b/Engine/Plugins/TextureGraph/Source/TextureGraphEngine/Model/Mix/MixInterface.cpp
--- a/Engine/Plugins/TextureGraph/Source/TextureGraphEngine/Model/Mix/MixInterface.cpp
+++ b/Engine/Plugins/TextureGraph/Source/TextureGraphEngine/Model/Mix/MixInterface.cpp
@@ -101,6 +101,11 @@ void UMixInterface::BroadcastOnRenderingDone(const FInvalidationDetails* Details
 	OnRenderDone.ExecuteIfBound(this, Details);
 }
 
+void UMixInterface::BroadcastOnBatchQueued(const FInvalidationDetails* Details)
+{
+	OnBatchQueued.ExecuteIfBound(this, Details);
+}
+
 UMixSettings* UMixInterface::GetSettings() const 
 { 
 	check(Settings);
diff --git a/Engine/Plugins/TextureGraph/Source/TextureGraphEngine/Model/Mix/MixInterface.h b/Engine/Plugins/TextureGraph/Source/TextureGraphEngine/Model/Mix/MixInterface.h
--- a/Engine/Plugins/TextureGraph/Source/TextureGraphEngine/Model/Mix/MixInterface.h
+++ b/Engine/Plugins/TextureGraph/Source/TextureGraphEngine/Model/Mix/MixInterface.h
@@ -61,6 +61,9 @@ public:
 
 
 
+	DECLARE_DELEGATE_TwoParams(FOnBatchQueued, UMixInterface*, const FInvalidationDetails*);
+	FOnBatchQueued						OnBatchQueued;
+
 	DECLARE_DELEGATE_TwoParams(FOnRenderDone, UMixInterface*, const FInvalidationDetails*);
 	FOnRenderDone						OnRenderDone;
 	virtual void						PostMeshLoad();
@@ -92,6 +95,7 @@ public:
 	virtual void						InvalidateAll();
 
 	virtual void						BroadcastOnRenderingDone(const FInvalidationDetails* Details);
+	virtual void						BroadcastOnBatchQueued(const FInvalidationDetails* Details);
 	
 	//////////////////////////////////////////////////////////////////////////
 	/// Inline functions
diff --git a/Engine/Plugins/TextureGraph/Source/TextureGraphEngine/Model/ModelObject.cpp b/Engine/Plugins/TextureGraph/Source/TextureGraphEngine/Model/ModelObject.cpp
--- a/Engine/Plugins/TextureGraph/Source/TextureGraphEngine/Model/ModelObject.cpp
+++ b/Engine/Plugins/TextureGraph/Source/TextureGraphEngine/Model/ModelObject.cpp
@@ -86,6 +86,12 @@ void FInvalidationDetails::BroadcastOnDone() const
 	Mix->BroadcastOnRenderingDone(this);
 }
 
+void FInvalidationDetails::BroadcastOnQueued() const
+{
+	OnQueued.ExecuteIfBound(this);
+	Mix->BroadcastOnBatchQueued(this);
+}
+
 bool FInvalidationDetails::IsDiscard() const
 {
 	return bTweaking;
diff --git a/Engine/Plugins/TextureGraph/Source/TextureGraphEngine/Model/ModelObject.h b/Engine/Plugins/TextureGraph/Source/TextureGraphEngine/Model/ModelObject.h
--- a/Engine/Plugins/TextureGraph/Source/TextureGraphEngine/Model/ModelObject.h
+++ b/Engine/Plugins/TextureGraph/Source/TextureGraphEngine/Model/ModelObject.h
@@ -24,6 +24,9 @@ struct TEXTUREGRAPHENGINE_API FInvalidationDetails
 	DECLARE_DELEGATE_OneParam(FOnDone, const FInvalidationDetails*);
 	FOnDone								OnDone;										/// Delegate allowing to provide a user callback once the invalidation is executed
 
+	DECLARE_DELEGATE_OneParam(FOnQueued, const FInvalidationDetails*);
+	FOnDone								OnQueued;									/// Delegate allowing to provide a user callback once the invalidation is executed
+
 	bool								bReload = false;							/// Whether to reload the sources or not
 	bool								bSelective = true;							/// Selective invalidation [tile based]
 	bool								bTweaking = false;							/// Whether being tweaked or not
@@ -47,6 +50,8 @@ struct TEXTUREGRAPHENGINE_API FInvalidationDetails
 	void								BroadcastOnDone() const;	/// Trigger the broadcast of the OnDone delegate AND all 
 																	/// the other accumulated delegates in _internal_onDoneMerged
 
+	void								BroadcastOnQueued() const;
+
 	FInvalidationDetails() : Mix(nullptr) {}
 	explicit FInvalidationDetails(const TWeakObjectPtr<UMixInterface> InMixObj);
 
