Fix warning in tests

diff --git a/Engine/Plugins/TextureGraph/Source/TextureGraphEngine/Job/Scheduler.cpp b/Engine/Plugins/TextureGraph/Source/TextureGraphEngine/Job/Scheduler.cpp
--- a/Engine/Plugins/TextureGraph/Source/TextureGraphEngine/Job/Scheduler.cpp
+++ b/Engine/Plugins/TextureGraph/Source/TextureGraphEngine/Job/Scheduler.cpp
@@ -172,7 +172,6 @@ void Scheduler::Update(float dt)
 
 				ObserverSource->BatchJobsDone(PreviousBatch);
 
-				UE_LOG(LogBatch, Warning, TEXT("!!!BATCH done!!!"));
 				/// Mark the current batch as done
 				PreviousBatch->GetCycle()->GetDetails().BroadcastOnDone();
 
