Resize node work

diff --git a/Engine/Plugins/TextureGraph/Source/TextureGraphEngine/Device/FX/Device_FX.cpp b/Engine/Plugins/TextureGraph/Source/TextureGraphEngine/Device/FX/Device_FX.cpp
--- a/Engine/Plugins/TextureGraph/Source/TextureGraphEngine/Device/FX/Device_FX.cpp
+++ b/Engine/Plugins/TextureGraph/Source/TextureGraphEngine/Device/FX/Device_FX.cpp
@@ -18,6 +18,7 @@
 #include <ScreenRendering.h>
 #include <TextureResource.h>
 #include <Engine/TextureRenderTarget2D.h>
+#include "ClearQuad.h"
 
 DECLARE_GPU_STAT_NAMED(Stat_CombineToTiles, TEXT("CombineToTilesProfile"));
 DECLARE_GPU_STAT_NAMED(Stat_SplitFromTiles, TEXT("SplitFromTilesProfile"));
@@ -283,9 +284,9 @@ AsyncDeviceBufferRef Device_FX::FillTextureArray_Deferred(DeviceBufferRef Buffer
 	return cti::make_ready_continuable(Buffer);
 }
 
-AsyncDeviceBufferRef Device_FX::DrawTilesToBuffer_Deferred(DeviceBufferRef Buffer, const T_Tiles<DeviceBufferRef>& Tiles)
+AsyncDeviceBufferRef Device_FX::DrawTilesToBuffer_Deferred(DeviceBufferRef Buffer, const T_Tiles<DeviceBufferRef>& Tiles, const DrawTilesSettings* DrawSettings, FLinearColor* ClearColor)
 {
-	DeviceNativeTaskPtr Task = DeviceNativeTask_Lambda::Create(this, (int32)E_Priority::kSystem, TEXT("DrawTilesToBuffer"), [this, Buffer, Tiles]() mutable -> int32
+	DeviceNativeTaskPtr Task = DeviceNativeTask_Lambda::Create(this, (int32)E_Priority::kSystem, TEXT("DrawTilesToBuffer"), [this, Buffer, Tiles, DrawSettings, ClearColor]() mutable -> int32
 		{
 			SCOPE_CYCLE_COUNTER(STAT_Device_FX_DrawTilesToBuffer);
 			check(IsInRenderingThread());
@@ -310,6 +311,10 @@ AsyncDeviceBufferRef Device_FX::DrawTilesToBuffer_Deferred(DeviceBufferRef Buffe
 
 			FRHIRenderPassInfo RenderPassInfo(RTResDest, ERenderTargetActions::Load_Store);
 			DevRHI.BeginRenderPass(RenderPassInfo, TEXT("CopyTexture"));
+
+			if (ClearColor)
+				DrawClearQuad(DevRHI, *ClearColor);
+
 			{
 				FGraphicsPipelineStateInitializer GraphicsPSOInit;
 				DevRHI.ApplyCachedRenderTargets(GraphicsPSOInit);
@@ -324,6 +329,11 @@ AsyncDeviceBufferRef Device_FX::DrawTilesToBuffer_Deferred(DeviceBufferRef Buffe
 				GraphicsPSOInit.BoundShaderState.VertexDeclarationRHI = GFilterVertexDeclaration.VertexDeclarationRHI;
 
 				FRHISamplerState* PixelSampler = TStaticSamplerState<SF_Bilinear>::GetRHI();
+				bool bUseInputPositions = DrawSettings != nullptr ? true : false;
+
+				/// Verify that there are either no DrawSettings or they are setup correctly to match the input array dimensions
+				check(!DrawSettings || (DrawSettings->Position.Rows() == Tiles.Rows() && DrawSettings->Position.Cols() == Tiles.Cols()));
+				check(!DrawSettings || (DrawSettings->Size.Rows() == Tiles.Rows() && DrawSettings->Size.Cols() == Tiles.Cols()));
 
 				SetGraphicsPipelineState(DevRHI, GraphicsPSOInit, 0);
 
@@ -343,16 +353,29 @@ AsyncDeviceBufferRef Device_FX::DrawTilesToBuffer_Deferred(DeviceBufferRef Buffe
 						UTextureRenderTarget2D* TileRT = TileTex->GetRenderTarget();
 						FRHITexture* TileResource = TileTex->GetRHITexture();
 
-						const float SrcTextureWidth = TileTex->GetWidth();
-						const float SrcTextureHeight = TileTex->GetHeight();
+						float SrcTextureWidth = TileTex->GetWidth();
+						float SrcTextureHeight = TileTex->GetHeight();
 
 						uint32 X1 = XOffset; ///ColId * SrcTextureWidth;
 						uint32 Y1 = 0; ///RowId * SrcTextureHeight;
 
-						for (int32 TAIndex = 0; TAIndex < RowId; TAIndex++)
+						if (!bUseInputPositions)
 						{
-							DeviceBuffer_FX* TileAbove = static_cast<DeviceBuffer_FX*>(Tiles(TAIndex, ColId).get());
-							Y1 += TileAbove->Descriptor().Height;
+							for (int32 TAIndex = 0; TAIndex < RowId; TAIndex++)
+							{
+								DeviceBuffer_FX* TileAbove = static_cast<DeviceBuffer_FX*>(Tiles(TAIndex, ColId).get());
+								Y1 += TileAbove->Descriptor().Height;
+							}
+						}
+						else
+						{
+							const FIntPoint& TilePos = DrawSettings->Position(RowId, ColId);
+							X1 = (uint32)TilePos.X;
+							Y1 = (uint32)TilePos.Y;
+
+							const FIntPoint& TileSize = DrawSettings->Size(RowId, ColId);
+							SrcTextureWidth = TileSize.X;
+							SrcTextureHeight = TileSize.Y;
 						}
 
 						DevRHI.SetViewport(X1, Y1, 0, X1 + ViewportWidth, Y1 + ViewportHeight, 0);
@@ -364,7 +387,7 @@ AsyncDeviceBufferRef Device_FX::DrawTilesToBuffer_Deferred(DeviceBufferRef Buffe
 							0, 0,
 							SrcTextureWidth, SrcTextureHeight,
 							0.0f, 0.0f,
-							1.0f, 1.0f,
+							1, 1,
 							TargetSize,
 							FIntPoint(1, 1),
 							VertexShader,
diff --git a/Engine/Plugins/TextureGraph/Source/TextureGraphEngine/Device/FX/Device_FX.h b/Engine/Plugins/TextureGraph/Source/TextureGraphEngine/Device/FX/Device_FX.h
--- a/Engine/Plugins/TextureGraph/Source/TextureGraphEngine/Device/FX/Device_FX.h
+++ b/Engine/Plugins/TextureGraph/Source/TextureGraphEngine/Device/FX/Device_FX.h
@@ -17,6 +17,12 @@ typedef std::shared_ptr<class TexArray>	TexArrayPtr;
 class UTextureRenderTarget;
 class IRendererModule;
 
+struct TEXTUREGRAPHENGINE_API DrawTilesSettings
+{
+	T_Tiles<FIntPoint>				Position;						/// The position where to put the tiles
+	T_Tiles<FIntPoint>				Size;							/// The size of the tiles
+};
+
 class TEXTUREGRAPHENGINE_API Device_FX : public Device,  public FGCObject
 {
 private: 
@@ -77,7 +83,7 @@ public:
 
 	virtual AsyncDeviceBufferRef	CombineFromTiles(const CombineSplitArgs& CombineArgs) override;
 
-	AsyncDeviceBufferRef			DrawTilesToBuffer_Deferred(DeviceBufferRef buffer, const T_Tiles<DeviceBufferRef>& Tiles);
+	AsyncDeviceBufferRef			DrawTilesToBuffer_Deferred(DeviceBufferRef buffer, const T_Tiles<DeviceBufferRef>& Tiles, const DrawTilesSettings* DrawSettings = nullptr, FLinearColor* ClearColor = nullptr);
 	AsyncDeviceBufferRef			FillTextureArray_Deferred(DeviceBufferRef buffer, const T_Tiles<DeviceBufferRef>& Tiles);
 
 	virtual AsyncDeviceBufferRef	SplitToTiles(const CombineSplitArgs& splitArgs) override;
diff --git a/Engine/Plugins/TextureGraph/Source/TextureGraphEngine/Helper/DataUtil.h b/Engine/Plugins/TextureGraph/Source/TextureGraphEngine/Helper/DataUtil.h
--- a/Engine/Plugins/TextureGraph/Source/TextureGraphEngine/Helper/DataUtil.h
+++ b/Engine/Plugins/TextureGraph/Source/TextureGraphEngine/Helper/DataUtil.h
@@ -222,6 +222,7 @@ protected:
 	size_t							NumCols = 0;		/// Number of columns
 
 public:
+	T_Tiles() = default;
 	T_Tiles(size_t InNumRows, size_t InNumCols) : NumRows(InNumRows), NumCols(InNumCols)
 	{
 		if (NumRows > 0 && NumCols > 0)
diff --git a/Engine/Plugins/TextureGraph/Source/TextureGraphEngine/Transform/Utility/T_CombineTiledBlob.cpp b/Engine/Plugins/TextureGraph/Source/TextureGraphEngine/Transform/Utility/T_CombineTiledBlob.cpp
--- a/Engine/Plugins/TextureGraph/Source/TextureGraphEngine/Transform/Utility/T_CombineTiledBlob.cpp
+++ b/Engine/Plugins/TextureGraph/Source/TextureGraphEngine/Transform/Utility/T_CombineTiledBlob.cpp
@@ -4,10 +4,10 @@
 #include "FxMat/MaterialManager.h"
 #include "Job/JobBatch.h"
 #include "Device/FX/DeviceBuffer_FX.h"
-#include "Device/FX/Device_FX.h"
 #include "Model/Mix/MixInterface.h"
 
-TiledBlobPtr T_CombineTiledBlob::Create(MixUpdateCyclePtr Cycle, BufferDescriptor DesiredOutputDesc, int32 TargetId, TiledBlobPtr SourceTex, JobUPtr JobToUse)
+TiledBlobPtr T_CombineTiledBlob::Create(MixUpdateCyclePtr Cycle, BufferDescriptor DesiredOutputDesc, int32 TargetId, 
+	TiledBlobPtr SourceTex, JobUPtr JobToUse, const CombineSettings* InSettings)
 {
 	if (SourceTex && SourceTex->IsTiled())
 	{
@@ -16,7 +16,7 @@ TiledBlobPtr T_CombineTiledBlob::Create(MixUpdateCyclePtr Cycle, BufferDescripto
 		if (!JobToUse)
 		{
 			FString Name = TEXT("T_CombineTiledBlob"); //FString::Printf(TEXT("[%s].[%d].[%llu] CombineTiledBlob"), *SourceTex->Name(), TargetId, Cycle->GetBatch()->GetFrameId());
-			std::shared_ptr<CombineTiledBlob_Transform> CombinedTiledBlobMat = std::make_shared< CombineTiledBlob_Transform>(Name, SourceTex);
+			std::shared_ptr<CombineTiledBlob_Transform> CombinedTiledBlobMat = std::make_shared<CombineTiledBlob_Transform>(Name, SourceTex, InSettings);
 			JobObj = std::make_unique<Job>(Cycle->GetMix(), TargetId, std::static_pointer_cast<BlobTransform>(CombinedTiledBlobMat));
 		}
 		else
@@ -24,6 +24,12 @@ TiledBlobPtr T_CombineTiledBlob::Create(MixUpdateCyclePtr Cycle, BufferDescripto
 
 		JobObj->AddArg(WithUnbounded(ARG_BLOB(SourceTex, "Source")));
 
+		if (InSettings)
+		{
+			JobObj->AddArg(WithUnbounded(ARG_LINEAR_COLOR(InSettings->BackgroundColor, "BackgroundColor")));
+			JobObj->AddArg(WithUnbounded(ARG_BOOL(InSettings->bMaintainAspectRatio, "MaintainAspectRatio")));
+		}
+
 		// Express the dependency of the new job on the job delivering the SourceTex
 		auto PrevJob = SourceTex->Job();
 		if (!PrevJob.expired())
@@ -56,10 +62,12 @@ TiledBlobPtr T_CombineTiledBlob::Create(MixUpdateCyclePtr Cycle, BufferDescripto
 
 
 //////////////////////////////////////////////////////////////////////////
-CombineTiledBlob_Transform::CombineTiledBlob_Transform(FString InName, TiledBlobPtr InSource)
+CombineTiledBlob_Transform::CombineTiledBlob_Transform(FString InName, TiledBlobPtr InSource, const CombineSettings* InSettings)
 	: BlobTransform(InName)
 	, Source(InSource)
 {
+	if (InSettings)
+		Settings = *InSettings;
 }
 
 Device* CombineTiledBlob_Transform::TargetDevice(size_t DevIndex) const
@@ -76,18 +84,106 @@ bool CombineTiledBlob_Transform::CanHandleTiles() const
 	return false;
 }
 
-AsyncTransformResultPtr	CombineTiledBlob_Transform::Exec(const TransformArgs& Args)
+void CombineTiledBlob_Transform::InitDrawSettingsWithoutAspectRatio(BlobPtr Target, T_Tiles<DeviceBufferRef>& Tiles)
 {
-	BlobPtr Target = Args.Target.lock();
-	check(Target);
+	uint32 TargetWidth = Target->GetWidth();
+	uint32 TargetHeight = Target->GetHeight();
+	uint32 SourceWidth = Source->GetWidth();
+	uint32 SourceHeight = Source->GetHeight();
+	float WidthRatio = TargetWidth / SourceWidth;
+	float HeightRatio = TargetHeight / SourceHeight;
 
-	DeviceBuffer_FX* DevBuffer = dynamic_cast<DeviceBuffer_FX*>(Target->GetBufferRef().get());
-	check(DevBuffer);
+	DrawSettings.Position.Resize(Source->Rows(), Source->Cols());
+	DrawSettings.Size.Resize(Source->Rows(), Source->Cols());
 
-	T_Tiles<DeviceBufferRef> Tiles(Source->Rows(), Source->Cols());
+	uint32 YOffset = 0;
+	for (int32 RowId = 0; RowId < Source->Rows(); RowId++)
+	{
+		uint32 XOffset = 0;
+		uint32 MaxRowHeight = 0;
+
+		for (int32 ColId = 0; ColId < Source->Cols(); ColId++)
+		{
+			check(Source->GetTile(RowId, ColId));
+
+			BlobPtr Tile = Source->GetTile(RowId, ColId);
+			check(Tile->GetBufferRef());
+
+			Tiles[RowId][ColId] = Tile->GetBufferRef();
+
+			uint32 Width = (uint32)(float)Tile->GetWidth() * WidthRatio;
+			uint32 Height = (uint32)(float)Tile->GetHeight() * HeightRatio;
+
+			/// account for rounding errors in the aspect ratio adjustments
+			if (ColId == Source->Cols() - 1)
+				Width = TargetWidth - XOffset;
+			if (RowId == Source->Rows() - 1)
+				Height = TargetHeight - YOffset;
+
+			MaxRowHeight = FMath::Max(Height, MaxRowHeight);
+
+			DrawSettings.Position[RowId][ColId] = FIntPoint(XOffset, YOffset);
+			DrawSettings.Size[RowId][ColId] = FIntPoint(Width, Height);
+
+			XOffset += Width;
+		}
+
+		YOffset += MaxRowHeight;
+	}
+}
+
+void CombineTiledBlob_Transform::InitDrawSettingsWithAspectRatio(BlobPtr Target, T_Tiles<DeviceBufferRef>& Tiles)
+{
+	uint32 TargetWidth = Target->GetWidth();
+	uint32 TargetHeight = Target->GetHeight();
+	uint32 SourceWidth = Source->GetWidth();
+	uint32 SourceHeight = Source->GetHeight();
+	float TargetAspectRatio = (float)TargetWidth / (float)TargetHeight;
+	float SourceAspectRatio = (float)SourceWidth / (float)SourceHeight;
+	float Scale = TargetAspectRatio / SourceAspectRatio;
+	//float TargetAspectRatio = TargetWidth / TargetHeight;
+	//float SourceAspectRatio = SourceWidth / SourceHeight;
+	uint32 NewTargetWidth = 0;
+	uint32 NewTargetHeight = 0;
+
+	if (TargetAspectRatio > SourceAspectRatio)
+	{
+		NewTargetWidth = TargetWidth;
+		NewTargetHeight = (uint32)(TargetWidth / SourceAspectRatio);
+	}
+	else
+	{
+		NewTargetWidth = TargetHeight * SourceAspectRatio;
+		NewTargetHeight = TargetHeight;
+	}
+
+	if (NewTargetHeight > TargetHeight)
+	{
+		NewTargetHeight = TargetHeight;
+		NewTargetWidth = NewTargetHeight * SourceAspectRatio;
+	}
+	else if (NewTargetWidth > TargetWidth)
+	{
+		NewTargetWidth = TargetWidth;
+		NewTargetHeight = (uint32)(TargetWidth / SourceAspectRatio);
+	}
+
+	float WidthRatio = (float)NewTargetWidth / SourceWidth;
+	float HeightRatio = (float)NewTargetHeight / SourceHeight;
 
+	check(NewTargetWidth <= TargetWidth && NewTargetHeight <= TargetHeight);
+	uint32 XMargin = (TargetWidth - NewTargetWidth) >> 1;
+	uint32 YMargin = (TargetHeight - NewTargetHeight) >> 1;
+
+	DrawSettings.Position.Resize(Source->Rows(), Source->Cols());
+	DrawSettings.Size.Resize(Source->Rows(), Source->Cols());
+
+	uint32 YOffset = YMargin;
 	for (int32 RowId = 0; RowId < Source->Rows(); RowId++)
 	{
+		uint32 XOffset = XMargin;
+		uint32 MaxRowHeight = 0;
+
 		for (int32 ColId = 0; ColId < Source->Cols(); ColId++)
 		{
 			check(Source->GetTile(RowId, ColId));
@@ -96,22 +192,94 @@ AsyncTransformResultPtr	CombineTiledBlob_Transform::Exec(const TransformArgs& Ar
 			check(Tile->GetBufferRef());
 
 			Tiles[RowId][ColId] = Tile->GetBufferRef();
+
+			uint32 Width = (uint32)(float)Tile->GetWidth() * WidthRatio;
+			uint32 Height = (uint32)(float)Tile->GetHeight() * HeightRatio;
+
+			MaxRowHeight = FMath::Max(Height, MaxRowHeight);
+
+			DrawSettings.Position[RowId][ColId] = FIntPoint(XOffset, YOffset);
+			DrawSettings.Size[RowId][ColId] = FIntPoint(Width, Height);
+
+			XOffset += Width;
 		}
+
+		YOffset += MaxRowHeight;
 	}
+}
 
-	Device_FX* FXDevice = Device_FX::Get();
+void CombineTiledBlob_Transform::InitDrawSettings(BlobPtr Target, T_Tiles<DeviceBufferRef>& Tiles)
+{
+	uint32 TargetWidth = Target->GetWidth();
+	uint32 TargetHeight = Target->GetHeight();
+	uint32 SourceWidth = Source->GetWidth();
+	uint32 SourceHeight = Source->GetHeight();
+	bool bUseDefault = true;
 
-	return FXDevice->DrawTilesToBuffer_Deferred(Target->GetBufferRef(), Tiles).then([this, Target](DeviceBufferRef)
+	if (TargetWidth != SourceWidth || TargetHeight != SourceHeight)
 	{
-			TransformResultPtr Result = std::make_shared<TransformResult>();
-			Result->Target = Target;
-			return Result;
-	});
+		if (!Settings.bMaintainAspectRatio)
+		{
+			InitDrawSettingsWithoutAspectRatio(Target, Tiles);
+			bUseDefault = false;
+		}
+		else if (Settings.bMaintainAspectRatio)
+		{
+			InitDrawSettingsWithAspectRatio(Target, Tiles);
+			bUseDefault = false;
+		}
+	}
+
+	if (bUseDefault)
+	{
+		for (int32 RowId = 0; RowId < Source->Rows(); RowId++)
+		{
+			for (int32 ColId = 0; ColId < Source->Cols(); ColId++)
+			{
+				check(Source->GetTile(RowId, ColId));
+
+				BlobPtr Tile = Source->GetTile(RowId, ColId);
+				check(Tile->GetBufferRef());
+
+				Tiles[RowId][ColId] = Tile->GetBufferRef();
+			}
+		}
+	}
+}
+
+AsyncTransformResultPtr	CombineTiledBlob_Transform::Exec(const TransformArgs& Args)
+{
+	BlobPtr Target = Args.Target.lock();
+	check(Target);
+
+	DeviceBuffer_FX* DevBuffer = dynamic_cast<DeviceBuffer_FX*>(Target->GetBufferRef().get());
+	check(DevBuffer);
+
+	T_Tiles<DeviceBufferRef> Tiles(Source->Rows(), Source->Cols());
+	DrawTilesSettings* DrawSettingsToUse = nullptr;
+	InitDrawSettings(Target, Tiles);
+
+	FLinearColor* ClearColor = nullptr;
+
+	if (DrawSettings.Position.Rows() == Tiles.Rows())
+	{
+		DrawSettingsToUse = &DrawSettings;
+		ClearColor = &Settings.BackgroundColor;
+	}
+
+	Device_FX* FXDevice = Device_FX::Get();
+	return FXDevice->DrawTilesToBuffer_Deferred(Target->GetBufferRef(), Tiles, DrawSettingsToUse, ClearColor)
+		.then([this, Target](DeviceBufferRef)
+		{
+				TransformResultPtr Result = std::make_shared<TransformResult>();
+				Result->Target = Target;
+				return Result;
+		});
 }
 
 std::shared_ptr<BlobTransform> CombineTiledBlob_Transform::DuplicateInstance(FString NewName)
 {
-	return std::make_shared<CombineTiledBlob_Transform>(NewName, Source);
+	return std::make_shared<CombineTiledBlob_Transform>(NewName, Source, &Settings);
 }
 
 AsyncBufferResultPtr CombineTiledBlob_Transform::Bind(BlobPtr BlobObj, const ResourceBindInfo& BindInfo)
diff --git a/Engine/Plugins/TextureGraph/Source/TextureGraphEngine/Transform/Utility/T_CombineTiledBlob.h b/Engine/Plugins/TextureGraph/Source/TextureGraphEngine/Transform/Utility/T_CombineTiledBlob.h
--- a/Engine/Plugins/TextureGraph/Source/TextureGraphEngine/Transform/Utility/T_CombineTiledBlob.h
+++ b/Engine/Plugins/TextureGraph/Source/TextureGraphEngine/Transform/Utility/T_CombineTiledBlob.h
@@ -7,7 +7,22 @@
 #include "FxMat/FxMaterial.h"
 #include "Model/Mix/MixUpdateCycle.h"
 #include <DataDrivenShaderPlatformInfo.h>
+#include "Device/FX/Device_FX.h"
 
+//UENUM(BlueprintType)
+//enum class ResizeSettings : uint8
+//{
+//	Stretch = 0			UMETA(DisplayName = "Stretch (Default)"),
+//	 = 1		UMETA(DisplayName = "Directional"),
+//	Radial = 2			UMETA(DisplayName = "Radial")
+//};
+
+struct TEXTUREGRAPHENGINE_API CombineSettings
+{
+	bool bFixed = false;
+	bool bMaintainAspectRatio = false;
+	FLinearColor BackgroundColor = FLinearColor::Transparent;
+};
 
 /**
  * A dedicated BlobTransform creating a Combined version of a TiledBlob
@@ -15,10 +30,16 @@
 class TEXTUREGRAPHENGINE_API CombineTiledBlob_Transform : public BlobTransform
 {
 private:
-	TiledBlobPtr					Source;
+	TiledBlobPtr					Source;				/// The source blob that we are combining
+	CombineSettings					Settings;			/// Combine settings
+	DrawTilesSettings				DrawSettings;		/// FX device specific draw settings
+
+	void							InitDrawSettings(BlobPtr Target, T_Tiles<DeviceBufferRef>& Tiles);
+	void							InitDrawSettingsWithoutAspectRatio(BlobPtr Target, T_Tiles<DeviceBufferRef>& Tiles);
+	void							InitDrawSettingsWithAspectRatio(BlobPtr Target, T_Tiles<DeviceBufferRef>& Tiles);
 
 public:
-	CombineTiledBlob_Transform(FString InName, TiledBlobPtr InSource);
+	CombineTiledBlob_Transform(FString InName, TiledBlobPtr InSource, const CombineSettings* InSettings = nullptr);
 	virtual AsyncBufferResultPtr	Bind(BlobPtr Value, const ResourceBindInfo& BindInfo) override;
 	virtual AsyncBufferResultPtr	Unbind(BlobPtr BlobObj, const ResourceBindInfo& BindInfo) override;
 	virtual Device*					TargetDevice(size_t DevIndex) const override;
@@ -40,5 +61,6 @@ public:
 	/// Static functions
 	//////////////////////////////////////////////////////////////////////////
 	
-	static TiledBlobPtr				Create(MixUpdateCyclePtr Cycle, BufferDescriptor DesiredOutputDesc, int32 TargetId, TiledBlobPtr SourceTex, JobUPtr JobToUse = nullptr);
+	static TiledBlobPtr				Create(MixUpdateCyclePtr Cycle, BufferDescriptor DesiredOutputDesc, int32 TargetId, 
+									TiledBlobPtr SourceTex, JobUPtr JobToUse = nullptr, const CombineSettings* InSettings = nullptr);
 };
