Resize node

diff --git a/Engine/Plugins/Experimental/TextureGraph/Source/TextureGraph/Private/Expressions/Utilities/TG_Expression_Utils.cpp b/Engine/Plugins/Experimental/TextureGraph/Source/TextureGraph/Private/Expressions/Utilities/TG_Expression_Utils.cpp
new file mode 100644
--- /dev/null
+++ b/Engine/Plugins/Experimental/TextureGraph/Source/TextureGraph/Private/Expressions/Utilities/TG_Expression_Utils.cpp
@@ -0,0 +1,54 @@
+// Copyright Epic Games, Inc. All Rights Reserved.
+
+#include "Expressions/Utilities/TG_Expression_Utils.h"
+#include "Expressions/Utilities/TG_Expression_MaterialID.h"
+
+#include "TG_Texture.h"
+#include "Transform/Expressions/T_ExtractMaterialIds.h"
+#include "Transform/Expressions/T_MaterialIDMask.h"
+
+void UTG_Expression_Utils_GetWidth::Evaluate(FTG_EvaluationContext* InContext)
+{
+	Super::Evaluate(InContext);
+
+	Width = 0;
+
+	if (Input)
+	{
+		if (!Input->IsLateBound())
+		{
+			Width = Input->GetWidth();
+		}
+		else
+		{
+			Input->OnFinalise().then([=, this](const Blob* Result)
+			{
+				Width = Result->GetWidth();
+			});
+		}
+	}
+}
+
+//////////////////////////////////////////////////////////
+void UTG_Expression_Utils_GetHeight::Evaluate(FTG_EvaluationContext* InContext)
+{
+	Super::Evaluate(InContext);
+
+	Height = 0;
+
+	if (Input)
+	{
+		if (!Input->IsLateBound())
+		{
+			Height = Input->GetHeight();
+		}
+		else
+		{
+			Input->OnFinalise().then([=, this](const Blob* Result)
+			{
+				Height = Result->GetHeight();
+			});
+		}
+	}
+}
+
diff --git a/Engine/Plugins/Experimental/TextureGraph/Source/TextureGraph/Public/Expressions/Utilities/TG_Expression_Utils.h b/Engine/Plugins/Experimental/TextureGraph/Source/TextureGraph/Public/Expressions/Utilities/TG_Expression_Utils.h
new file mode 100644
--- /dev/null
+++ b/Engine/Plugins/Experimental/TextureGraph/Source/TextureGraph/Public/Expressions/Utilities/TG_Expression_Utils.h
@@ -0,0 +1,51 @@
+// Copyright Epic Games, Inc. All Rights Reserved.
+#pragma once
+
+#include "CoreMinimal.h"
+
+#include "Expressions/TG_Expression.h"
+#include "TG_Texture.h"
+#include "Transform/Expressions/T_ExtractMaterialIds.h"
+
+#include "TG_Expression_Utils.generated.h"
+
+//////////////////////////////////////////////////////////
+UCLASS()
+class TEXTUREGRAPH_API UTG_Expression_Utils_GetWidth : public UTG_Expression
+{
+	GENERATED_BODY()
+public:
+	TG_DECLARE_EXPRESSION(TG_Category::Utilities)
+	virtual void Evaluate(FTG_EvaluationContext* InContext) override;
+
+	UPROPERTY(meta = (TGType = "TG_Input", PinDisplayName = ""))
+	FTG_Texture Input;
+
+	/// The width of the texture
+	UPROPERTY(meta = (TGType = "TG_Output", HideInnerPropertiesInNode))
+	float Width;
+public:
+	virtual FTG_Name GetDefaultName() const override { return TEXT("GetWidth"); }
+	virtual FText GetTooltipText() const override { return FText::FromString(TEXT("Extracts the width from an input texture. Useful for doing maths on image dimensions.")); }
+};
+
+//////////////////////////////////////////////////////////
+UCLASS()
+class TEXTUREGRAPH_API UTG_Expression_Utils_GetHeight : public UTG_Expression
+{
+	GENERATED_BODY()
+public:
+	TG_DECLARE_EXPRESSION(TG_Category::Utilities)
+	virtual void Evaluate(FTG_EvaluationContext* InContext) override;
+
+	UPROPERTY(meta = (TGType = "TG_Input", PinDisplayName = ""))
+	FTG_Texture Input;
+
+	/// The height of the texture
+	UPROPERTY(meta = (TGType = "TG_Output", HideInnerPropertiesInNode))
+	float Height;
+public:
+	virtual FTG_Name GetDefaultName() const override { return TEXT("GetHeight"); }
+	virtual FText GetTooltipText() const override { return FText::FromString(TEXT("Extracts the height from an input texture. Useful for doing maths on image dimensions.")); }
+};
+
