Fix crash on undo

diff --git a/Engine/Plugins/TextureGraph/Source/TextureGraph/Private/TG_Node.cpp b/Engine/Plugins/TextureGraph/Source/TextureGraph/Private/TG_Node.cpp
--- a/Engine/Plugins/TextureGraph/Source/TextureGraph/Private/TG_Node.cpp
+++ b/Engine/Plugins/TextureGraph/Source/TextureGraph/Private/TG_Node.cpp
@@ -211,16 +211,20 @@ void UTG_Node::OnPinConnectionChanged(FTG_Id InPinId, FTG_Id OldPinId, FTG_Id Ne
 void UTG_Node::OnPinConnectionUndo(FTG_Id InPinId)
 {
 	UTG_Pin* PinTo = GetPin(InPinId);
-	// Reset the Pin to Default in this case
-	if (PinTo->IsArgVariant())
-	{
-		GetExpression()->NotifyCommonInputVariantTypeChanged(FTG_Variant::EType::Scalar); // notify and update the common type and the signature accordingly
-	}
-	else if (PinTo->IsArgTexture())
+
+	if (PinTo)
 	{
-		//Reset the FTGTexture
-		PinTo->EditSelfVar()->EditAs<FTG_Texture>() = nullptr;
-		GetExpression()->CopyVarToExpressionArgument(PinTo->GetArgument(), PinTo->EditSelfVar());
+		// Reset the Pin to Default in this case
+		if (PinTo->IsArgVariant())
+		{
+			GetExpression()->NotifyCommonInputVariantTypeChanged(FTG_Variant::EType::Scalar); // notify and update the common type and the signature accordingly
+		}
+		else if (PinTo->IsArgTexture())
+		{
+			//Reset the FTGTexture
+			PinTo->EditSelfVar()->EditAs<FTG_Texture>() = nullptr;
+			GetExpression()->CopyVarToExpressionArgument(PinTo->GetArgument(), PinTo->EditSelfVar());
+		}
 	}
 }
 
