Some resize node fixes

diff --git a/Engine/Plugins/TextureGraph/Source/TextureGraph/Private/Expressions/Utilities/TG_Expression_Utils.cpp b/Engine/Plugins/TextureGraph/Source/TextureGraph/Private/Expressions/Utilities/TG_Expression_Utils.cpp
--- a/Engine/Plugins/TextureGraph/Source/TextureGraph/Private/Expressions/Utilities/TG_Expression_Utils.cpp
+++ b/Engine/Plugins/TextureGraph/Source/TextureGraph/Private/Expressions/Utilities/TG_Expression_Utils.cpp
@@ -67,7 +67,7 @@ void UTG_Expression_Utils_Resize::Evaluate(FTG_EvaluationContext* InContext)
 	CombineSettings Settings
 	{
 		.bFixed = false,
-		.bMaintainAspectRatio = false,
+		.bMaintainAspectRatio = bMaintainAspectRatio,
 		.BackgroundColor = BackgroundColor
 	};
 
diff --git a/Engine/Plugins/TextureGraph/Source/TextureGraphEditor/Private/STG_NodePreview.cpp b/Engine/Plugins/TextureGraph/Source/TextureGraphEditor/Private/STG_NodePreview.cpp
--- a/Engine/Plugins/TextureGraph/Source/TextureGraphEditor/Private/STG_NodePreview.cpp
+++ b/Engine/Plugins/TextureGraph/Source/TextureGraphEditor/Private/STG_NodePreview.cpp
@@ -358,13 +358,13 @@ void STG_NodePreviewWidget::Update() const
 				Blob->OnFinalise()
 					.then([Blob]()
 						{
-							UE_LOG(LogTemp, Warning, TEXT("!!!Finalised. Combining!!!"));
+							UE_LOG(LogTemp, Log, TEXT("!!!Finalised. Combining!!!"));
 							const TiledBlobPtr BlobTiled = std::static_pointer_cast<TiledBlob>(Blob);
 							return BlobTiled->CombineTiles(false, false);
 						})
 					.then([this, Blob]()
 						{
-							UE_LOG(LogTemp, Warning, TEXT("!!!SETTING TEXTURE!!!"));
+							UE_LOG(LogTemp, Log, TEXT("!!!SETTING TEXTURE!!!"));
 							NodeViewer->SetTexture(Blob);
 						});
 			}
diff --git a/Engine/Plugins/TextureGraph/Source/TextureGraphEngine/Transform/Utility/T_CombineTiledBlob.cpp b/Engine/Plugins/TextureGraph/Source/TextureGraphEngine/Transform/Utility/T_CombineTiledBlob.cpp
--- a/Engine/Plugins/TextureGraph/Source/TextureGraphEngine/Transform/Utility/T_CombineTiledBlob.cpp
+++ b/Engine/Plugins/TextureGraph/Source/TextureGraphEngine/Transform/Utility/T_CombineTiledBlob.cpp
@@ -90,8 +90,8 @@ void CombineTiledBlob_Transform::InitDrawSettingsWithoutAspectRatio(BlobPtr Targ
 	uint32 TargetHeight = Target->GetHeight();
 	uint32 SourceWidth = Source->GetWidth();
 	uint32 SourceHeight = Source->GetHeight();
-	float WidthRatio = TargetWidth / SourceWidth;
-	float HeightRatio = TargetHeight / SourceHeight;
+	float WidthRatio = (float)TargetWidth / (float)SourceWidth;
+	float HeightRatio = (float)TargetHeight / (float)SourceHeight;
 
 	DrawSettings.Position.Resize(Source->Rows(), Source->Cols());
 	DrawSettings.Size.Resize(Source->Rows(), Source->Cols());
@@ -140,36 +140,33 @@ void CombineTiledBlob_Transform::InitDrawSettingsWithAspectRatio(BlobPtr Target,
 	uint32 SourceHeight = Source->GetHeight();
 	float TargetAspectRatio = (float)TargetWidth / (float)TargetHeight;
 	float SourceAspectRatio = (float)SourceWidth / (float)SourceHeight;
-	float Scale = TargetAspectRatio / SourceAspectRatio;
-	//float TargetAspectRatio = TargetWidth / TargetHeight;
-	//float SourceAspectRatio = SourceWidth / SourceHeight;
 	uint32 NewTargetWidth = 0;
 	uint32 NewTargetHeight = 0;
 
 	if (TargetAspectRatio > SourceAspectRatio)
 	{
 		NewTargetWidth = TargetWidth;
-		NewTargetHeight = (uint32)(TargetWidth / SourceAspectRatio);
+		NewTargetHeight = (uint32)((float)TargetWidth / SourceAspectRatio);
 	}
 	else
 	{
-		NewTargetWidth = TargetHeight * SourceAspectRatio;
+		NewTargetWidth = (uint32)((float)TargetHeight * SourceAspectRatio);
 		NewTargetHeight = TargetHeight;
 	}
 
 	if (NewTargetHeight > TargetHeight)
 	{
 		NewTargetHeight = TargetHeight;
-		NewTargetWidth = NewTargetHeight * SourceAspectRatio;
+		NewTargetWidth = (uint32)((float)NewTargetHeight * SourceAspectRatio);
 	}
 	else if (NewTargetWidth > TargetWidth)
 	{
 		NewTargetWidth = TargetWidth;
-		NewTargetHeight = (uint32)(TargetWidth / SourceAspectRatio);
+		NewTargetHeight = (uint32)((float)TargetWidth / SourceAspectRatio);
 	}
 
-	float WidthRatio = (float)NewTargetWidth / SourceWidth;
-	float HeightRatio = (float)NewTargetHeight / SourceHeight;
+	float WidthRatio = (float)NewTargetWidth / (float)SourceWidth;
+	float HeightRatio = (float)NewTargetHeight / (float)SourceHeight;
 
 	check(NewTargetWidth <= TargetWidth && NewTargetHeight <= TargetHeight);
 	uint32 XMargin = (TargetWidth - NewTargetWidth) >> 1;
