Clang build fix.

diff --git a/Engine/Plugins/TextureGraph/Source/TextureGraph/Private/TG_OutputSettings.cpp b/Engine/Plugins/TextureGraph/Source/TextureGraph/Private/TG_OutputSettings.cpp
--- a/Engine/Plugins/TextureGraph/Source/TextureGraph/Private/TG_OutputSettings.cpp
+++ b/Engine/Plugins/TextureGraph/Source/TextureGraph/Private/TG_OutputSettings.cpp
@@ -73,6 +73,11 @@ void FTG_OutputSettings::OnSetTexturePresetType(ETG_TexturePresetType Type)
 	}
 }
 
+bool FTG_OutputSettings::Validate(FString& OutErrors) const
+{
+	return TextureExporter::IsFilePathValid(BaseName, FolderPath, OutErrors);
+}
+
 void FTG_OutputSettings::Initialize(FString PathName,FName InName /*= "Output"*/)
 {
 	OutputName = InName;
diff --git a/Engine/Plugins/TextureGraph/Source/TextureGraph/Public/TG_Graph.h b/Engine/Plugins/TextureGraph/Source/TextureGraph/Public/TG_Graph.h
--- a/Engine/Plugins/TextureGraph/Source/TextureGraph/Public/TG_Graph.h
+++ b/Engine/Plugins/TextureGraph/Source/TextureGraph/Public/TG_Graph.h
@@ -3,16 +3,13 @@
 #pragma once
 
 #include "TG_Pin.h"
-
 #include "TG_Var.h"
-
 #include "TG_Node.h"
-
 #include <functional>
-
-#include "TG_OutputSettings.h"
 #include "TG_Graph.generated.h"
 
+class UTG_ExportSettings;
+
 // Graph Traversal
 // A Graph internal struct populated at runtime used to Traverse it.
 // This data is built by the Graph itself.
diff --git a/Engine/Plugins/TextureGraph/Source/TextureGraph/Public/TG_HelperFunctions.h b/Engine/Plugins/TextureGraph/Source/TextureGraph/Public/TG_HelperFunctions.h
--- a/Engine/Plugins/TextureGraph/Source/TextureGraph/Public/TG_HelperFunctions.h
+++ b/Engine/Plugins/TextureGraph/Source/TextureGraph/Public/TG_HelperFunctions.h
@@ -7,7 +7,6 @@
 #include "Export/TextureExporter.h"
 #include "Data/Blob.h"
 #include "TG_Graph.h"
-#include "TG_Node.h"
 #include "2D/TextureHelper.h"
 
 class UTextureGraphBase;
@@ -26,6 +25,7 @@ public:
 
 	static JobBatchPtr InitRenderBatch(UTextureGraphBase* InTextureGraph, JobBatchPtr ExistingBatch = nullptr);
 	static AsyncBool RenderAsync(UTextureGraphBase* InTextureGraph, JobBatchPtr ExistingBatch = nullptr);
+
 	template <typename T_Type>
 	static TArray<T_Type> GetOutputsOfType(const UTG_Node* Node)
 	{
diff --git a/Engine/Plugins/TextureGraph/Source/TextureGraph/Public/TG_OutputSettings.h b/Engine/Plugins/TextureGraph/Source/TextureGraph/Public/TG_OutputSettings.h
--- a/Engine/Plugins/TextureGraph/Source/TextureGraph/Public/TG_OutputSettings.h
+++ b/Engine/Plugins/TextureGraph/Source/TextureGraph/Public/TG_OutputSettings.h
@@ -5,14 +5,12 @@
 #include "TG_SystemTypes.h"
 #include "Model/Mix/MixSettings.h"
 #include "Misc/OutputDeviceNull.h"
-#include "TG_HelperFunctions.h"
 #include "TG_OutputSettings.generated.h"
 
 class UTG_Expression_Output;
 class UTG_Node;
 class UTextureGraph;
 
-
 USTRUCT()
 struct TEXTUREGRAPH_API FOutputExpressionInfo
 {
@@ -114,9 +112,5 @@ struct TEXTUREGRAPH_API FTG_OutputSettings
 		TextureCompressionSettings InCompression = TextureCompressionSettings::TC_Default, TextureGroup InLodGroup = TextureGroup::TEXTUREGROUP_World, bool InbSRGB = false);
 
 	void OnSetTexturePresetType(ETG_TexturePresetType Type);
-
-	bool Validate(FString& OutErrors)
-	{
-		return TextureExporter::IsFilePathValid(BaseName, FolderPath, OutErrors);
-	}
+	bool Validate(FString& OutErrors) const;
 };
