Resize node issues fixed

diff --git a/Engine/Plugins/TextureGraph/Source/TextureGraph/Private/Expressions/Utilities/TG_Expression_Utils.cpp b/Engine/Plugins/TextureGraph/Source/TextureGraph/Private/Expressions/Utilities/TG_Expression_Utils.cpp
--- a/Engine/Plugins/TextureGraph/Source/TextureGraph/Private/Expressions/Utilities/TG_Expression_Utils.cpp
+++ b/Engine/Plugins/TextureGraph/Source/TextureGraph/Private/Expressions/Utilities/TG_Expression_Utils.cpp
@@ -7,6 +7,7 @@
 #include "Transform/Expressions/T_ExtractMaterialIds.h"
 #include "Transform/Expressions/T_MaterialIDMask.h"
 #include "Transform/Utility/T_CombineTiledBlob.h"
+#include "Transform/Utility/T_SplitToTiles.h"
 
 #if 0
 void UTG_Expression_Utils_GetWidth::Evaluate(FTG_EvaluationContext* InContext)
@@ -71,5 +72,6 @@ void UTG_Expression_Utils_Resize::Evaluate(FTG_EvaluationContext* InContext)
 		.BackgroundColor = BackgroundColor
 	};
 
-	Output = T_CombineTiledBlob::Create(InContext->Cycle, Output.Descriptor, InContext->TargetId, Input, nullptr, &Settings);
+	TiledBlobPtr CombinedResult = T_CombineTiledBlob::Create(InContext->Cycle, Output.Descriptor, InContext->TargetId, Input, nullptr, &Settings);
+	Output = T_SplitToTiles::Create(InContext->Cycle, InContext->TargetId, CombinedResult);
 }
diff --git a/Engine/Plugins/TextureGraph/Source/TextureGraphEngine/Job/Job.cpp b/Engine/Plugins/TextureGraph/Source/TextureGraphEngine/Job/Job.cpp
--- a/Engine/Plugins/TextureGraph/Source/TextureGraphEngine/Job/Job.cpp
+++ b/Engine/Plugins/TextureGraph/Source/TextureGraphEngine/Job/Job.cpp
@@ -270,7 +270,9 @@ CHashPtr Job::CalcTileHash(int32 RowId, int32 ColId) const
 
 CHashPtr Job::TileHash(int32 RowId, int32 ColId) const
 {
-	return TextureGraphEngine::GetBlobber()->AddGloballyUniqueHash(CalcTileHash(RowId, ColId));
+	CHashPtr Hash = CalcTileHash(RowId, ColId);
+	check(Hash);
+	return TextureGraphEngine::GetBlobber()->AddGloballyUniqueHash(Hash);
 }
 
 BufferDescriptor Job::GetCombinedDesc(BufferDescriptor& ArgsDescCombined, size_t& Count) const
diff --git a/Engine/Plugins/TextureGraph/Source/TextureGraphEngine/Transform/Expressions/T_ArrayGrid.cpp b/Engine/Plugins/TextureGraph/Source/TextureGraphEngine/Transform/Expressions/T_ArrayGrid.cpp
--- a/Engine/Plugins/TextureGraph/Source/TextureGraphEngine/Transform/Expressions/T_ArrayGrid.cpp
+++ b/Engine/Plugins/TextureGraph/Source/TextureGraphEngine/Transform/Expressions/T_ArrayGrid.cpp
@@ -156,12 +156,6 @@ TiledBlobPtr T_ArrayGrid::Create(MixUpdateCyclePtr Cycle, BufferDescriptor& InDe
 	int32 InNumRows = Cycle->GetMix()->GetNumXTiles();
 	int32 InNumCols = Cycle->GetMix()->GetNumYTiles();
 
-	for (int32 InputIndex = 0; InputIndex < Inputs.Num(); InputIndex++)
-	{
-		InNumRows = std::min<int32>(InNumRows, (int32)Inputs[InputIndex]->Rows());
-		InNumCols = std::min<int32>(InNumCols, (int32)Inputs[InputIndex]->Cols());
-	}
-
 	CHashPtrVec InputHashes(Inputs.Num());
 	bool bIsWellDefined = true;
 	uint32 MaxRowHeight = 0;
@@ -174,6 +168,19 @@ TiledBlobPtr T_ArrayGrid::Create(MixUpdateCyclePtr Cycle, BufferDescriptor& InDe
 		&bIsWellDefined, &InputHashes);
 	TiledBlobPtr FinalResult;
 
+	CombineSettings Settings
+	{
+		.bFixed = false,
+		.bMaintainAspectRatio = true,
+		.BackgroundColor = FLinearColor::Transparent
+	};
+
+	BufferDescriptor InputDesc = Inputs[0]->GetDescriptor();
+	InputDesc.Width = Width;
+	InputDesc.Height = Height;
+
+	BufferDescriptor DesiredOutputDesc = BufferDescriptor::CombineWithPreference(&InputDesc, &InDesiredOutputDesc, nullptr);
+
 	/// If the result is late bound then we have no option but to resize the result render target through a custom job
 	if (!bIsWellDefined)
 	{
@@ -187,18 +194,19 @@ TiledBlobPtr T_ArrayGrid::Create(MixUpdateCyclePtr Cycle, BufferDescriptor& InDe
 		for (int32 InputIndex = 0; InputIndex < Inputs.Num(); InputIndex++)
 			JobObj->AddArg(WithUnbounded(ARG_BLOB(Inputs[InputIndex], "Input")));
 
-		TiledBlobPtr CombinedResultOne = T_CombineTiledBlob::Create(Cycle, InDesiredOutputDesc, TargetId, CombinedResultTiled, std::move(JobObj));
+		TiledBlobPtr CombinedResult = T_CombineTiledBlob::Create(Cycle, InputDesc, TargetId, CombinedResultTiled, std::move(JobObj), nullptr);
+		TiledBlobPtr CombinedResultOne = CombinedResult;
+
+		if (InDesiredOutputDesc.Width != 0 || InDesiredOutputDesc.Height != 0)
+			CombinedResultOne = T_CombineTiledBlob::Create(Cycle, DesiredOutputDesc, TargetId, CombinedResult, nullptr, &Settings);
+
 		FinalResult = T_SplitToTiles::Create(Cycle, TargetId, CombinedResultOne);
 	}
 	else
 	{
-		BufferDescriptor DesiredOutputDesc = BufferDescriptor::Combine(Inputs[0]->GetDescriptor(), InDesiredOutputDesc);
-		DesiredOutputDesc.Width = Width;
-		DesiredOutputDesc.Height = Height;
-
 		/// The we combine those tiles into a single result
-		TiledBlobPtr CombinedResultTiled = TiledBlob::InitFromTiles(DesiredOutputDesc, ResultTiles);
-		TiledBlobPtr CombinedResultOne = T_CombineTiledBlob::Create(Cycle, DesiredOutputDesc, TargetId, CombinedResultTiled);
+		TiledBlobPtr CombinedResultTiled = TiledBlob::InitFromTiles(InputDesc, ResultTiles);
+		TiledBlobPtr CombinedResultOne = T_CombineTiledBlob::Create(Cycle, DesiredOutputDesc, TargetId, CombinedResultTiled, nullptr, &Settings);
 
 		if (DesiredOutputDesc.Width > 1 && DesiredOutputDesc.Height > 1)
 			FinalResult = T_SplitToTiles::Create(Cycle, TargetId, CombinedResultOne);
diff --git a/Engine/Plugins/TextureGraph/Source/TextureGraphEngine/Transform/Layer/T_Thumbnail.cpp b/Engine/Plugins/TextureGraph/Source/TextureGraphEngine/Transform/Layer/T_Thumbnail.cpp
--- a/Engine/Plugins/TextureGraph/Source/TextureGraphEngine/Transform/Layer/T_Thumbnail.cpp
+++ b/Engine/Plugins/TextureGraph/Source/TextureGraphEngine/Transform/Layer/T_Thumbnail.cpp
@@ -25,6 +25,9 @@ JobBatchPtr AddThumbJobToCycle(TiledBlobPtr InBlobToBind, JobUPtr JobPtr, UObjec
 
 TiledBlobRef T_Thumbnail::Bind(UMixInterface* Mix, UObject* Model, TiledBlobPtr InBlobToBind, int32 InTargetId)
 {
+	if (InBlobToBind.get() == nullptr)
+		return TextureHelper::GetBlack();
+
 	if (InBlobToBind->IsTransient())
 		return TextureHelper::GetBlack();
 
