Fixed incorrect handling of the GTE and LTE conditions in IfThenElseNode

diff --git a/Engine/Plugins/TextureGraph/Source/TextureGraphEngine/Transform/Expressions/T_Maths_TwoInputs.cpp b/Engine/Plugins/TextureGraph/Source/TextureGraphEngine/Transform/Expressions/T_Maths_TwoInputs.cpp
--- a/Engine/Plugins/TextureGraph/Source/TextureGraphEngine/Transform/Expressions/T_Maths_TwoInputs.cpp
+++ b/Engine/Plugins/TextureGraph/Source/TextureGraphEngine/Transform/Expressions/T_Maths_TwoInputs.cpp
@@ -118,7 +118,15 @@ RenderMaterial_FXPtr GetIfThenElseMaterial(EIfThenElseOperator Operator, EIfThen
 			return TextureGraphEngine::GetMaterialManager()->CreateMaterial_FX<VSH_Simple, FSH_IfThenElse_GT_Grayscale>("FSH_IfThenElse_GT_Grayscale");
 		break;
 
-		
+	case EIfThenElseOperator::GTE:
+		if (Type == EIfThenElseType::IndividualComponent)
+			return TextureGraphEngine::GetMaterialManager()->CreateMaterial_FX<VSH_Simple, FSH_IfThenElse_GTE_Component>("FSH_IfThenElse_GTE_Component");
+		else if (Type == EIfThenElseType::AllComponents)
+			return TextureGraphEngine::GetMaterialManager()->CreateMaterial_FX<VSH_Simple, FSH_IfThenElse_GTE_All>("FSH_IfThenElse_GTE_All");
+		else if (Type == EIfThenElseType::Grayscale)
+			return TextureGraphEngine::GetMaterialManager()->CreateMaterial_FX<VSH_Simple, FSH_IfThenElse_GTE_Grayscale>("FSH_IfThenElse_GTE_Grayscale");
+		break;
+
 	case EIfThenElseOperator::LT:
 		if (Type == EIfThenElseType::IndividualComponent)
 			return TextureGraphEngine::GetMaterialManager()->CreateMaterial_FX<VSH_Simple, FSH_IfThenElse_LT_Component>("FSH_IfThenElse_LT_Component");
@@ -128,6 +136,15 @@ RenderMaterial_FXPtr GetIfThenElseMaterial(EIfThenElseOperator Operator, EIfThen
 			return TextureGraphEngine::GetMaterialManager()->CreateMaterial_FX<VSH_Simple, FSH_IfThenElse_LT_Grayscale>("FSH_IfThenElse_LT_Grayscale");
 		break;
 		
+	case EIfThenElseOperator::LTE:
+		if (Type == EIfThenElseType::IndividualComponent)
+			return TextureGraphEngine::GetMaterialManager()->CreateMaterial_FX<VSH_Simple, FSH_IfThenElse_LTE_Component>("FSH_IfThenElse_LTE_Component");
+		else if (Type == EIfThenElseType::AllComponents)
+			return TextureGraphEngine::GetMaterialManager()->CreateMaterial_FX<VSH_Simple, FSH_IfThenElse_LTE_All>("FSH_IfThenElse_LTE_All");
+		else if (Type == EIfThenElseType::Grayscale)
+			return TextureGraphEngine::GetMaterialManager()->CreateMaterial_FX<VSH_Simple, FSH_IfThenElse_LTE_Grayscale>("FSH_IfThenElse_LTE_Grayscale");
+		break;
+
 	case EIfThenElseOperator::EQ:
 		if (Type == EIfThenElseType::IndividualComponent)
 			return TextureGraphEngine::GetMaterialManager()->CreateMaterial_FX<VSH_Simple, FSH_IfThenElse_EQ_Component>("FSH_IfThenElse_EQ_Component");
