Fix Horde issue with TextureGraphInstance using TStrongObjectPtr for its RuntimeGraph

diff --git a/Engine/Plugins/TextureGraph/Source/TextureGraph/Private/TextureGraph.cpp b/Engine/Plugins/TextureGraph/Source/TextureGraph/Private/TextureGraph.cpp
--- a/Engine/Plugins/TextureGraph/Source/TextureGraph/Private/TextureGraph.cpp
+++ b/Engine/Plugins/TextureGraph/Source/TextureGraph/Private/TextureGraph.cpp
@@ -27,7 +27,7 @@
 void UTextureGraphInstance::Construct(FString InName)
 {
 	Super::Construct(InName);
-	RuntimeGraph = TStrongObjectPtr<UTG_Graph>(NewObject<UTG_Graph>(GetTransientPackage(), NAME_None,RF_Transactional));
+	RuntimeGraph = NewObject<UTG_Graph>(GetTransientPackage(), NAME_None,RF_Transactional);
 	RuntimeGraph->Construct(InName);
 	
 	const UTG_Node* OutputNode = RuntimeGraph->CreateExpressionNode(UTG_Expression_Output::StaticClass());
@@ -62,7 +62,7 @@ void UTextureGraphInstance::SetParent(const TObjectPtr<UTextureGraphBase>& Paren
 	}
 	else
 	{
-		RuntimeGraph = TStrongObjectPtr<UTG_Graph>(DuplicateObject(Parent->Graph(), GetTransientPackage()));
+		RuntimeGraph = DuplicateObject(Parent->Graph(), GetTransientPackage());
 		Settings = DuplicateObject(Parent->GetSettings(), this);
 
 		// copy InputParams and ExportSettings from Parent if it is an instance
@@ -117,7 +117,7 @@ UTG_Graph* UTextureGraphInstance::Graph()
 {
 	if (!RuntimeGraph && ParentTextureGraph)
 	{
-		RuntimeGraph = TStrongObjectPtr<UTG_Graph>(DuplicateObject(ParentTextureGraph->Graph(), GetTransientPackage()));
+		RuntimeGraph = DuplicateObject(ParentTextureGraph->Graph(), GetTransientPackage());
 	}
 	return RuntimeGraph.Get();
 }
@@ -126,7 +126,7 @@ void UTextureGraphInstance::Initialize()
 {
 	if (ParentTextureGraph)
 	{
-		RuntimeGraph = TStrongObjectPtr<UTG_Graph>(DuplicateObject(ParentTextureGraph->Graph(), GetTransientPackage()));
+		RuntimeGraph = DuplicateObject(ParentTextureGraph->Graph(), GetTransientPackage());
 		CopyParamsToRuntimeGraph();
 	}
 }
diff --git a/Engine/Plugins/TextureGraph/Source/TextureGraph/Public/TextureGraph.h b/Engine/Plugins/TextureGraph/Source/TextureGraph/Public/TextureGraph.h
--- a/Engine/Plugins/TextureGraph/Source/TextureGraph/Public/TextureGraph.h
+++ b/Engine/Plugins/TextureGraph/Source/TextureGraph/Public/TextureGraph.h
@@ -81,7 +81,8 @@ class TEXTUREGRAPH_API UTextureGraphInstance: public UTextureGraphBase
 	GENERATED_BODY()
 	
 private:
-	TStrongObjectPtr<UTG_Graph>	RuntimeGraph;		// The runtime graph instance to use for the sub-graph
+	UPROPERTY(transient)
+	TObjectPtr<UTG_Graph>	RuntimeGraph;		// The runtime graph instance to use for the sub-graph
 public:
 	
 	UPROPERTY()
