Resize node crashes

diff --git a/Engine/Plugins/TextureGraph/Source/TextureGraph/Private/Expressions/Utilities/TG_Expression_Utils.cpp b/Engine/Plugins/TextureGraph/Source/TextureGraph/Private/Expressions/Utilities/TG_Expression_Utils.cpp
--- a/Engine/Plugins/TextureGraph/Source/TextureGraph/Private/Expressions/Utilities/TG_Expression_Utils.cpp
+++ b/Engine/Plugins/TextureGraph/Source/TextureGraph/Private/Expressions/Utilities/TG_Expression_Utils.cpp
@@ -59,7 +59,7 @@ void UTG_Expression_Utils_GetHeight::Evaluate(FTG_EvaluationContext* InContext)
 //////////////////////////////////////////////////////////
 void UTG_Expression_Utils_Resize::Evaluate(FTG_EvaluationContext* InContext)
 {
-	if (!Input)
+	if (!Input || Input->GetWidth() < 8 || Input->GetHeight() < 8)
 	{
 		Output = FTG_Texture::GetBlack();
 		return;
diff --git a/Engine/Plugins/TextureGraph/Source/TextureGraph/Private/TG_GraphEvaluation.cpp b/Engine/Plugins/TextureGraph/Source/TextureGraph/Private/TG_GraphEvaluation.cpp
--- a/Engine/Plugins/TextureGraph/Source/TextureGraph/Private/TG_GraphEvaluation.cpp
+++ b/Engine/Plugins/TextureGraph/Source/TextureGraph/Private/TG_GraphEvaluation.cpp
@@ -49,9 +49,20 @@ void FTG_Evaluation::TransferVarToPin(UTG_Pin* InPin, FTG_EvaluationContext* Con
 
 	if (Arg.IsInput())
 	{
-		if (InPin->ConnectionNeedsConversion())
+		bool bNeedsConversion = InPin->ConnectionNeedsConversion();
+		FName ArrayConverterKey = InPin->GetInputVarConverterKey();
+		FTG_Var* OutVar = InPin->EditConvertedVar(); // Generate the converted version of the connected input into the CoonvertedVar of the Pin
+
+		if (!bNeedsConversion && (Var->IsArray() && !InPin->EditSelfVar()->IsArray()))
+		{
+			ArrayConverterKey = MakeConvertKey("FTG_VariantArray", Arg.CPPTypeName);
+			OutVar = InPin->EditSelfVar();
+			bNeedsConversion = true;
+		}
+
+		if (bNeedsConversion)
 		{
-			FTG_Evaluation::VarConverter* Converter = FTG_Evaluation::DefaultConverters.Find(InPin->GetInputVarConverterKey());
+			FTG_Evaluation::VarConverter* Converter = FTG_Evaluation::DefaultConverters.Find(ArrayConverterKey);
 			// If converter is null then the var is compatible, no need to convert.
 			// Else COnverter is valid and we need to run it 
 			if ((*Converter))
@@ -59,7 +70,7 @@ void FTG_Evaluation::TransferVarToPin(UTG_Pin* InPin, FTG_EvaluationContext* Con
 				FTG_Evaluation::VarConverterInfo Info;
 				Info.Context = Context;
 				Info.InVar = Var;
-				Info.OutVar = InPin->EditConvertedVar(); // Generate the converted version of the connected input into the CoonvertedVar of the Pin
+				Info.OutVar = OutVar;
 				Info.Index = Index;
 
 				(*Converter)(Info);
diff --git a/Engine/Plugins/TextureGraph/Source/TextureGraph/Public/Expressions/Utilities/TG_Expression_Utils.h b/Engine/Plugins/TextureGraph/Source/TextureGraph/Public/Expressions/Utilities/TG_Expression_Utils.h
--- a/Engine/Plugins/TextureGraph/Source/TextureGraph/Public/Expressions/Utilities/TG_Expression_Utils.h
+++ b/Engine/Plugins/TextureGraph/Source/TextureGraph/Public/Expressions/Utilities/TG_Expression_Utils.h
@@ -78,6 +78,6 @@ public:
 
 public:
 	virtual FTG_Name GetDefaultName() const override { return TEXT("Resize Texture"); }
-	virtual FText GetTooltipText() const override { return FText::FromString(TEXT("Extracts the height from an input texture. Useful for doing maths on image dimensions.")); }
+	virtual FText GetTooltipText() const override { return FText::FromString(TEXT("Resizes an input texture to the given dimensions.")); }
 };
 
diff --git a/Engine/Plugins/TextureGraph/Source/TextureGraphEngine/Transform/Utility/T_SplitToTiles.h b/Engine/Plugins/TextureGraph/Source/TextureGraphEngine/Transform/Utility/T_SplitToTiles.h
--- a/Engine/Plugins/TextureGraph/Source/TextureGraphEngine/Transform/Utility/T_SplitToTiles.h
+++ b/Engine/Plugins/TextureGraph/Source/TextureGraphEngine/Transform/Utility/T_SplitToTiles.h
@@ -51,8 +51,8 @@ public:
 	{
 		if (NumThreadsX == 0 || NumThreadsY == 0)
 		{
-			NumThreadsX = RefBlob->GetWidth() / NumTilesX;
-			NumThreadsY = RefBlob->GetHeight() / NumTilesY;
+			NumThreadsX = FMath::Max((int32)RefBlob->GetWidth() / NumTilesX, 1);
+			NumThreadsY = FMath::Max((int32)RefBlob->GetHeight() / NumTilesY, 1);
 		}
 
 		check(NumThreadsX > 0 && NumThreadsY > 0);
