Channel swizzle

diff --git a/Engine/Plugins/TextureGraph/Shaders/Expressions/Expression_ChannelSwizzle.usf b/Engine/Plugins/TextureGraph/Shaders/Expressions/Expression_ChannelSwizzle.usf
new file mode 100644
--- /dev/null
+++ b/Engine/Plugins/TextureGraph/Shaders/Expressions/Expression_ChannelSwizzle.usf
@@ -0,0 +1,73 @@
+// Copyright Epic Games, Inc. All Rights Reserved.
+#include "/Engine/Public/Platform.ush"
+#include "/Plugin/TextureGraph/SamplerStates.ush"
+
+Texture2D SourceTexture;
+
+#define LINEAR 0
+#define EXP 1
+#define LOG 1
+
+#ifndef DST_CHANNEL_RED
+#define DST_CHANNEL_RED 0
+#endif
+
+#ifndef DST_CHANNEL_GREEN
+#define DST_CHANNEL_GREEN 1
+#endif
+
+#ifndef DST_CHANNEL_BLUE
+#define DST_CHANNEL_BLUE 2
+#endif
+
+#ifndef DST_CHANNEL_ALPHA
+#define DST_CHANNEL_ALPHA 3
+#endif
+
+float4 FSH_ChannelSwizzle(float2 uv : TEXCOORD0) : SV_Target0
+{
+	float4 SrcClr = SourceTexture.Sample(SamplerStates_NoBorder, uv);
+	float4 DstClr = SrcClr;
+
+#if DST_CHANNEL_RED == 0
+	DstClr.r = SrcClr.r;
+#elif DST_CHANNEL_RED == 1
+	DstClr.r = SrcClr.g;
+#elif DST_CHANNEL_RED == 2
+	DstClr.r = SrcClr.b;
+#elif DST_CHANNEL_RED == 3
+	DstClr.r = SrcClr.a;
+#endif 
+	
+#if DST_CHANNEL_GREEN == 0
+	DstClr.g = SrcClr.r;
+#elif DST_CHANNEL_GREEN == 1
+	DstClr.g = SrcClr.g;
+#elif DST_CHANNEL_GREEN == 2
+	DstClr.g = SrcClr.b;
+#elif DST_CHANNEL_GREEN == 3
+	DstClr.g = SrcClr.a;
+#endif 
+	
+#if DST_CHANNEL_BLUE == 0
+	DstClr.b = SrcClr.r;
+#elif DST_CHANNEL_BLUE == 1
+	DstClr.b = SrcClr.g;
+#elif DST_CHANNEL_BLUE == 2
+	DstClr.b = SrcClr.b;
+#elif DST_CHANNEL_BLUE == 3
+	DstClr.b = SrcClr.a;
+#endif 
+	
+#if DST_CHANNEL_ALPHA == 0
+	DstClr.a = SrcClr.r;
+#elif DST_CHANNEL_ALPHA == 1
+	DstClr.a = SrcClr.g;
+#elif DST_CHANNEL_ALPHA == 2
+	DstClr.a = SrcClr.b;
+#elif DST_CHANNEL_ALPHA == 3
+	DstClr.a = SrcClr.a;
+#endif 
+	
+	return DstClr;
+}
diff --git a/Engine/Plugins/TextureGraph/Source/TextureGraph/Private/Expressions/Channel/TG_Expression_ChannelSwizzle.cpp b/Engine/Plugins/TextureGraph/Source/TextureGraph/Private/Expressions/Channel/TG_Expression_ChannelSwizzle.cpp
new file mode 100644
--- /dev/null
+++ b/Engine/Plugins/TextureGraph/Source/TextureGraph/Private/Expressions/Channel/TG_Expression_ChannelSwizzle.cpp
@@ -0,0 +1,42 @@
+// Copyright Epic Games, Inc. All Rights Reserved.
+
+#include "Expressions/Channel/TG_Expression_ChannelSwizzle.h"
+#include "FxMat/MaterialManager.h"
+#include "Job/JobBatch.h"
+#include <Transform/Expressions/T_Color.h>
+
+void UTG_Expression_ChannelSwizzle::Evaluate(FTG_EvaluationContext* InContext)
+{
+	Super::Evaluate(InContext);
+
+	if (!Input)
+	{
+		Output = FTG_Texture::GetBlack();
+		return;
+	}
+
+	FTG_Texture Default = FTG_Texture::GetBlack();
+	TiledBlobRef Source = InContext->Inputs.GetVar("Input")->GetAsWithDefault<FTG_Texture>(Default).RasterBlob;
+
+	FSH_ChannelSwizzle::FPermutationDomain PermutationVector;
+	PermutationVector.Set<FVar_SwizzleDstChannelRed>((int32)RedChannel);
+	PermutationVector.Set<FVar_SwizzleDstChannelGreen>((int32)GreenChannel);
+	PermutationVector.Set<FVar_SwizzleDstChannelBlue>((int32)BlueChannel);
+	PermutationVector.Set<FVar_SwizzleDstChannelAlpha>((int32)AlphaChannel);
+
+	const RenderMaterial_FXPtr RenderMaterial = TextureGraphEngine::GetMaterialManager()->CreateMaterial_FX<VSH_Simple, FSH_ChannelSwizzle>(TEXT("T_ChannelSwizzle"), PermutationVector);
+	check(RenderMaterial);
+
+	JobUPtr RenderJob = std::make_unique<Job>(InContext->Cycle->GetMix(), InContext->TargetId, std::static_pointer_cast<BlobTransform>(RenderMaterial));
+	RenderJob
+		->AddArg(ARG_BLOB(Source, "SourceTexture"))
+		->AddArg(WithUnbounded(ARG_INT((int32)RedChannel, "RedChannel")))
+		->AddArg(WithUnbounded(ARG_INT((int32)GreenChannel, "GreenChannel")))
+		->AddArg(WithUnbounded(ARG_INT((int32)BlueChannel, "BlueChannel")))
+		->AddArg(WithUnbounded(ARG_INT((int32)AlphaChannel, "AlphaChannel")))
+		;
+
+	BufferDescriptor Desc = Output.GetBufferDescriptor();
+	Output = RenderJob->InitResult("ChannelSwizzle", &Desc);
+	InContext->Cycle->AddJob(InContext->TargetId, std::move(RenderJob));
+}
diff --git a/Engine/Plugins/TextureGraph/Source/TextureGraph/Public/Expressions/Channel/TG_Expression_ChannelSplitter.h b/Engine/Plugins/TextureGraph/Source/TextureGraph/Public/Expressions/Channel/TG_Expression_ChannelSplitter.h
--- a/Engine/Plugins/TextureGraph/Source/TextureGraph/Public/Expressions/Channel/TG_Expression_ChannelSplitter.h
+++ b/Engine/Plugins/TextureGraph/Source/TextureGraph/Public/Expressions/Channel/TG_Expression_ChannelSplitter.h
@@ -9,36 +9,36 @@
 
 #include "TG_Expression_ChannelSplitter.generated.h"
 
-UCLASS(CollapseCategories)
+UCLASS()
 class TEXTUREGRAPH_API UTG_Expression_ChannelSplitter : public UTG_Expression
 {
 	GENERATED_BODY()
 
 public:
 	TG_DECLARE_EXPRESSION(TG_Category::Channel);
-	virtual void						Evaluate(FTG_EvaluationContext* InContext) override;
+	virtual void			Evaluate(FTG_EvaluationContext* InContext) override;
 
 	// The input texture
 	UPROPERTY(meta = (TGType = "TG_Input", PinDisplayName = "RGBA"))
-	FTG_Texture							Input;
+	FTG_Texture				Input;
 
 	// Red channel of the input
 	UPROPERTY(EditAnywhere, Category = NoCategory, meta = (TGType = "TG_Output", PinDisplayName = "R", HideInnerPropertiesInNode))
-	FTG_Texture							Red = FTG_Texture::GetBlack();
+	FTG_Texture				Red = FTG_Texture::GetBlack();
 
 	// Green channel of the input
 	UPROPERTY(EditAnywhere, Category = NoCategory, meta = (TGType = "TG_Output", PinDisplayName = "G", HideInnerPropertiesInNode))
-	FTG_Texture							Green = FTG_Texture::GetBlack();
+	FTG_Texture				Green = FTG_Texture::GetBlack();
 
 	// Blue channel of the input
 	UPROPERTY(EditAnywhere, Category = NoCategory, meta = (TGType = "TG_Output", PinDisplayName = "B", HideInnerPropertiesInNode))
-	FTG_Texture							Blue = FTG_Texture::GetBlack();
+	FTG_Texture				Blue = FTG_Texture::GetBlack();
 
 	// Alpha channel of the input
 	UPROPERTY(EditAnywhere, Category = NoCategory, meta = (TGType = "TG_Output", PinDisplayName = "A", HideInnerPropertiesInNode))
 	FTG_Texture							Alpha = FTG_Texture::GetBlack();
 
-	virtual FTG_Name					GetDefaultName() const override { return TEXT("SplitChannels");}
-	virtual FText						GetTooltipText() const override { return FText::FromString(TEXT("Splits an RGBA image into individual R, G, B, A channels.")); } 
+	virtual FTG_Name		GetDefaultName() const override { return TEXT("SplitChannels");}
+	virtual FText			GetTooltipText() const override { return FText::FromString(TEXT("Splits an RGBA image into individual R, G, B, A channels.")); } 
 };
 
diff --git a/Engine/Plugins/TextureGraph/Source/TextureGraph/Public/Expressions/Channel/TG_Expression_ChannelSwizzle.h b/Engine/Plugins/TextureGraph/Source/TextureGraph/Public/Expressions/Channel/TG_Expression_ChannelSwizzle.h
new file mode 100644
--- /dev/null
+++ b/Engine/Plugins/TextureGraph/Source/TextureGraph/Public/Expressions/Channel/TG_Expression_ChannelSwizzle.h
@@ -0,0 +1,50 @@
+// Copyright Epic Games, Inc. All Rights Reserved.
+
+#pragma once
+
+#include "CoreMinimal.h"
+
+#include "Expressions/TG_Expression.h"
+#include "TG_Texture.h"
+#include "Transform/Expressions/T_Channel.h"
+
+#include "TG_Expression_ChannelSwizzle.generated.h"
+
+UCLASS()
+class TEXTUREGRAPH_API UTG_Expression_ChannelSwizzle : public UTG_Expression
+{
+	GENERATED_BODY()
+
+public:
+	TG_DECLARE_EXPRESSION(TG_Category::Channel);
+	virtual void			Evaluate(FTG_EvaluationContext* InContext) override;
+
+	// The input texture
+	UPROPERTY(meta = (TGType = "TG_Input", PinDisplayName = ""))
+	FTG_Texture				Input;
+
+	// The output in the red channel
+	UPROPERTY(EditAnywhere, BlueprintReadWrite, Category = NoCategory, meta = (TGType = "TG_Setting", DisplayName = "Red"))
+	EColorChannel			RedChannel = EColorChannel::Red;
+
+	// The output in the Green channel
+	UPROPERTY(EditAnywhere, BlueprintReadWrite, Category = NoCategory, meta = (TGType = "TG_Setting", DisplayName = "Green"))
+	EColorChannel			GreenChannel = EColorChannel::Green;
+
+	// The output in the Blue channel
+	UPROPERTY(EditAnywhere, BlueprintReadWrite, Category = NoCategory, meta = (TGType = "TG_Setting", DisplayName = "Blue"))
+	EColorChannel			BlueChannel = EColorChannel::Blue;
+
+	// The output in the Alpha channel
+	UPROPERTY(EditAnywhere, BlueprintReadWrite, Category = NoCategory, meta = (TGType = "TG_Setting", DisplayName = "Alpha"))
+	EColorChannel			AlphaChannel = EColorChannel::Alpha;
+
+	// The output texture after swizzle operation
+	UPROPERTY(EditAnywhere, Category = NoCategory, meta = (TGType = "TG_Output", PinDisplayName = ""))
+	FTG_Texture				Output;
+
+
+	virtual FTG_Name		GetDefaultName() const override { return TEXT("Swizzle");}
+	virtual FText			GetTooltipText() const override { return FText::FromString(TEXT("Performs a swizzle operation.")); } 
+};
+
diff --git a/Engine/Plugins/TextureGraph/Source/TextureGraphEngine/Transform/Expressions/T_Channel.cpp b/Engine/Plugins/TextureGraph/Source/TextureGraphEngine/Transform/Expressions/T_Channel.cpp
--- a/Engine/Plugins/TextureGraph/Source/TextureGraphEngine/Transform/Expressions/T_Channel.cpp
+++ b/Engine/Plugins/TextureGraph/Source/TextureGraphEngine/Transform/Expressions/T_Channel.cpp
@@ -14,3 +14,5 @@ IMPLEMENT_GLOBAL_SHADER(FSH_ChannelSplitter_Blue, "/Plugin/TextureGraph/Expressi
 IMPLEMENT_GLOBAL_SHADER(FSH_ChannelSplitter_Alpha, "/Plugin/TextureGraph/Expressions/Expression_ChannelSplitter.usf", "FSH_ChannelSplitter_Alpha", SF_Pixel);
 
 IMPLEMENT_GLOBAL_SHADER(FSH_ChannelCombiner, "/Plugin/TextureGraph/Expressions/Expression_ChannelCombiner.usf", "FSH_ChannelCombiner", SF_Pixel);
+
+IMPLEMENT_GLOBAL_SHADER(FSH_ChannelSwizzle, "/Plugin/TextureGraph/Expressions/Expression_ChannelSwizzle.usf", "FSH_ChannelSwizzle", SF_Pixel);
diff --git a/Engine/Plugins/TextureGraph/Source/TextureGraphEngine/Transform/Expressions/T_Channel.h b/Engine/Plugins/TextureGraph/Source/TextureGraphEngine/Transform/Expressions/T_Channel.h
--- a/Engine/Plugins/TextureGraph/Source/TextureGraphEngine/Transform/Expressions/T_Channel.h
+++ b/Engine/Plugins/TextureGraph/Source/TextureGraphEngine/Transform/Expressions/T_Channel.h
@@ -6,6 +6,16 @@
 #include "FxMat/FxMaterial.h"
 #include "Model/Mix/MixUpdateCycle.h"
 #include <DataDrivenShaderPlatformInfo.h>
+#include "T_Channel.generated.h"
+
+UENUM(BlueprintType)
+enum class EColorChannel : uint8
+{
+	Red		= 0		UMETA(DisplayName = "Red"),
+	Green	= 1		UMETA(DisplayName = "Green"),
+	Blue	= 2		UMETA(DisplayName = "Blue"),
+	Alpha	= 3		UMETA(DisplayName = "Alpha"),
+};
 
 //////////////////////////////////////////////////////////////////////////
 /// Basic Channel Op
@@ -20,11 +30,8 @@ public:
 		SHADER_PARAMETER_TEXTURE(Texture2D, SourceTexture)
 	END_SHADER_PARAMETER_STRUCT()
 
-	static bool						ShouldCompilePermutation(const FGlobalShaderPermutationParameters& Parameters) 
-	{
-		return IsPCPlatform(Parameters.Platform) && EnumHasAllFlags(Parameters.Flags, EShaderPermutationFlags::HasEditorOnlyData);
-	}
-	static void						ModifyCompilationEnvironment(const FGlobalShaderPermutationParameters& params, FShaderCompilerEnvironment& env) {}
+	TEXTURE_ENGINE_DEFAULT_PERMUTATION;
+	TEXTUREGRAPH_ENGINE_DEFAULT_COMPILATION_ENV;
 };
 
 //////////////////////////////////////////////////////////////////////////
@@ -68,7 +75,7 @@ public:
 };
 
 //////////////////////////////////////////////////////////////////////////
-/// Basic Channel Op
+/// Channel Combiner
 //////////////////////////////////////////////////////////////////////////
 class TEXTUREGRAPHENGINE_API FSH_ChannelCombiner : public FSH_Base
 {
@@ -84,8 +91,33 @@ public:
 		SHADER_PARAMETER_TEXTURE(Texture2D, SourceAlpha)
 	END_SHADER_PARAMETER_STRUCT()
 
-	static bool						ShouldCompilePermutation(const FGlobalShaderPermutationParameters& Parameters) { return true; }
-	static void						ModifyCompilationEnvironment(const FGlobalShaderPermutationParameters& params, FShaderCompilerEnvironment& env) {}
+	TEXTURE_ENGINE_DEFAULT_PERMUTATION;
+	TEXTUREGRAPH_ENGINE_DEFAULT_COMPILATION_ENV;
+};
+
+//////////////////////////////////////////////////////////////////////////
+/// Channel Combiner
+//////////////////////////////////////////////////////////////////////////
+class FVar_SwizzleDstChannelRed : SHADER_PERMUTATION_INT("DST_CHANNEL_RED", 4);
+class FVar_SwizzleDstChannelGreen : SHADER_PERMUTATION_INT("DST_CHANNEL_GREEN", 4);
+class FVar_SwizzleDstChannelBlue : SHADER_PERMUTATION_INT("DST_CHANNEL_BLUE", 4);
+class FVar_SwizzleDstChannelAlpha : SHADER_PERMUTATION_INT("DST_CHANNEL_ALPHA", 4);
+
+class TEXTUREGRAPHENGINE_API FSH_ChannelSwizzle : public FSH_Base
+{
+public:
+	DECLARE_GLOBAL_SHADER(FSH_ChannelSwizzle);
+	SHADER_USE_PARAMETER_STRUCT(FSH_ChannelSwizzle, FSH_Base);
+
+	BEGIN_SHADER_PARAMETER_STRUCT(FParameters, )
+		SHADER_PARAMETER_STRUCT(FStandardSamplerStates, SamplerStates)
+		SHADER_PARAMETER_TEXTURE(Texture2D, SourceTexture)
+	END_SHADER_PARAMETER_STRUCT()
+
+	using FPermutationDomain = TShaderPermutationDomain<FVar_SwizzleDstChannelRed, FVar_SwizzleDstChannelGreen, FVar_SwizzleDstChannelBlue, FVar_SwizzleDstChannelAlpha>;
+public:
+	TEXTURE_ENGINE_DEFAULT_PERMUTATION;
+	TEXTUREGRAPH_ENGINE_DEFAULT_COMPILATION_ENV;
 };
 
 //////////////////////////////////////////////////////////////////////////
