Array concatenation customization

diff --git a/Engine/Plugins/TextureGraph/Source/TextureGraph/Private/Expressions/Arrays/TG_Expression_ArrayOps.cpp b/Engine/Plugins/TextureGraph/Source/TextureGraph/Private/Expressions/Arrays/TG_Expression_ArrayOps.cpp
--- a/Engine/Plugins/TextureGraph/Source/TextureGraph/Private/Expressions/Arrays/TG_Expression_ArrayOps.cpp
+++ b/Engine/Plugins/TextureGraph/Source/TextureGraph/Private/Expressions/Arrays/TG_Expression_ArrayOps.cpp
@@ -13,13 +13,34 @@ void UTG_Expression_ArrayConcat::Evaluate(FTG_EvaluationContext* InContext)
 {
 	Super::Evaluate(InContext);
 
-	Output.SetNum(Input1.Num() + Input2.Num());
+	uint32 Input1Count = (uint32)Input1.Num();
+	uint32 Input2Count = (uint32)Input2.Num();
+	uint32 MaxCount = Input1Count + Input2Count;
+	Output.SetNum(MaxCount);
 
-	for (int i = 0; i < Input1.Num(); i++)
-		Output.Set(i, Input1.GetArray()[i]);
+	uint32 ActualStartIndex = StartIndex;
 
-	for (int i = 0; i < Input2.Num(); i++)
-		Output.Set(i + Input1.Num(), Input2.GetArray()[i]);
+	if (StartIndex >= MaxCount)
+	{
+		FTextureGraphErrorReporter* ErrorReporter = TextureGraphEngine::GetErrorReporter(InContext->Cycle->GetMix());
+
+		if (ErrorReporter)
+		{
+			ErrorReporter->ReportWarning((int32)ETextureGraphErrorType::NODE_WARNING, FString::Printf(TEXT("Invalid starting index for Input-1 in the output array specified: %d [Output Array Size: %d]."), 
+				StartIndex, MaxCount), GetParentNode());
+		}
+
+		ActualStartIndex = 0;
+	}
+
+	for (uint32 i = 0; i < Input2Count && i < ActualStartIndex; i++)
+		Output.Set(i, Input2.GetArray()[i]);
+
+	for (uint32 i = 0; i < Input1Count; i++)
+		Output.Set(i + ActualStartIndex, Input1.GetArray()[i]);
+
+	for (uint32 i = ActualStartIndex; i < Input2Count; i++)
+		Output.Set(i + Input1Count, Input2.GetArray()[i]);
 }
 
 /////////////////////////////////////////////////////////////
diff --git a/Engine/Plugins/TextureGraph/Source/TextureGraph/Public/Expressions/Arrays/TG_Expression_ArrayOps.h b/Engine/Plugins/TextureGraph/Source/TextureGraph/Public/Expressions/Arrays/TG_Expression_ArrayOps.h
--- a/Engine/Plugins/TextureGraph/Source/TextureGraph/Public/Expressions/Arrays/TG_Expression_ArrayOps.h
+++ b/Engine/Plugins/TextureGraph/Source/TextureGraph/Public/Expressions/Arrays/TG_Expression_ArrayOps.h
@@ -20,11 +20,15 @@ public:
 	virtual void					Evaluate(FTG_EvaluationContext* InContext) override;
 	
 	/// The first array as part of the concatenation process
-	UPROPERTY(meta = (TGType = "TG_Input", PinDisplayName = ""))
+	UPROPERTY(meta = (TGType = "TG_Input", PinDisplayName = "Input-1"))
 	FTG_VariantArray				Input1;
 
+	/// The index at which Input1 array is going to start. Input 2 is going occupy the rest of the places within the array
+	UPROPERTY(meta = (TGType = "TG_Input", PinDisplayName = "Input-1 Out Index"))
+	uint32							StartIndex = 0;
+
 	/// The second array as part of the concatenation process
-	UPROPERTY(meta = (TGType = "TG_Input", PinDisplayName = ""))
+	UPROPERTY(meta = (TGType = "TG_Input", PinDisplayName = "Input-2"))
 	FTG_VariantArray				Input2;
 
 	/// Concatenated array
