Adding a make vector4 node

diff --git a/Engine/Plugins/TextureGraph/Source/TextureGraph/Private/Expressions/Utilities/TG_Expression_Utils.cpp b/Engine/Plugins/TextureGraph/Source/TextureGraph/Private/Expressions/Utilities/TG_Expression_Utils.cpp
--- a/Engine/Plugins/TextureGraph/Source/TextureGraph/Private/Expressions/Utilities/TG_Expression_Utils.cpp
+++ b/Engine/Plugins/TextureGraph/Source/TextureGraph/Private/Expressions/Utilities/TG_Expression_Utils.cpp
@@ -56,6 +56,11 @@ void UTG_Expression_Utils_GetHeight::Evaluate(FTG_EvaluationContext* InContext)
 }
 #endif 
 
+void UTG_Expression_Utils_MakeVector4::Evaluate(FTG_EvaluationContext* InContext)
+{
+	Output = FVector4f(X, Y, Z, W);
+}
+
 //////////////////////////////////////////////////////////
 void UTG_Expression_Utils_Resize::Evaluate(FTG_EvaluationContext* InContext)
 {
diff --git a/Engine/Plugins/TextureGraph/Source/TextureGraph/Public/Expressions/Utilities/TG_Expression_Utils.h b/Engine/Plugins/TextureGraph/Source/TextureGraph/Public/Expressions/Utilities/TG_Expression_Utils.h
--- a/Engine/Plugins/TextureGraph/Source/TextureGraph/Public/Expressions/Utilities/TG_Expression_Utils.h
+++ b/Engine/Plugins/TextureGraph/Source/TextureGraph/Public/Expressions/Utilities/TG_Expression_Utils.h
@@ -50,6 +50,35 @@ public:
 	virtual FText GetTooltipText() const override { return FText::FromString(TEXT("Extracts the height from an input texture. Useful for doing maths on image dimensions.")); }
 };
 #endif 
+//////////////////////////////////////////////////////////
+UCLASS()
+class TEXTUREGRAPH_API UTG_Expression_Utils_MakeVector4 : public UTG_Expression
+{
+	GENERATED_BODY()
+public:
+	TG_DECLARE_EXPRESSION(TG_Category::Utilities)
+	virtual void Evaluate(FTG_EvaluationContext* InContext) override;
+
+	UPROPERTY(meta = (TGType = "TG_Input"))
+	float X = 0;
+
+	UPROPERTY(meta = (TGType = "TG_Input"))
+	float Y = 0;
+
+	UPROPERTY(meta = (TGType = "TG_Input"))
+	float Z = 0;
+
+	UPROPERTY(meta = (TGType = "TG_Input"))
+	float W = 0;
+
+	/// The height of the texture
+	UPROPERTY(meta = (TGType = "TG_Output"))
+	FVector4f Output;
+
+public:
+	virtual FTG_Name GetDefaultName() const override { return TEXT("MakeVector4"); }
+	virtual FText GetTooltipText() const override { return FText::FromString(TEXT("Converts a bunch of floating point values to a vector.")); }
+};
 
 //////////////////////////////////////////////////////////
 UCLASS()
