Modify the ZOom behavior in the Transform Expression:
The zoom is now a percentage expressing the size of the source image relative to the cell destination rectangle

diff --git a/Engine/Plugins/TextureGraph/Shaders/Expressions/Expression_Transform.usf b/Engine/Plugins/TextureGraph/Shaders/Expressions/Expression_Transform.usf
--- a/Engine/Plugins/TextureGraph/Shaders/Expressions/Expression_Transform.usf
+++ b/Engine/Plugins/TextureGraph/Shaders/Expressions/Expression_Transform.usf
@@ -127,7 +127,7 @@ float4 FSH_TransformBlit(float2 uv : TEXCOORD0) : SV_Target0
 	float2 spacedCellPos = ( (cellPos - cellTilePos) - spacing * 0.5) / (1 - spacing);
 
 	// and a last scaling to apply the zoom control
-	float scaling = lerp((exp(log(rcp(2 * length(pixelSize))) * (-Zoom))), (1 - Zoom), (Zoom >= 0)); 
+	float scaling = rcp(max(length(pixelSize), Zoom)); 
 	spacedCellPos = ((spacedCellPos * 2 - 1) * scaling) * 0.5 + 0.5;
 
 	// Eval the cell grid masked  for debug
diff --git a/Engine/Plugins/TextureGraph/Source/TextureGraph/Private/Expressions/Procedural/TG_Expression_Transform.cpp b/Engine/Plugins/TextureGraph/Source/TextureGraph/Private/Expressions/Procedural/TG_Expression_Transform.cpp
--- a/Engine/Plugins/TextureGraph/Source/TextureGraph/Private/Expressions/Procedural/TG_Expression_Transform.cpp
+++ b/Engine/Plugins/TextureGraph/Source/TextureGraph/Private/Expressions/Procedural/TG_Expression_Transform.cpp
@@ -30,7 +30,7 @@ void UTG_Expression_Transform::Evaluate(FTG_EvaluationContext* InContext)
 	};
 
 	T_Transform::CellParameter CellParam{
-		.Zoom = Zoom,
+		.Zoom = Zoom * 0.01f,
 		.StretchToFit = StretchToFit,
 		.Spacing = Spacing,
 		.Stagger = {(StaggerOffset * StaggerHorizontally), (StaggerOffset * (1 - StaggerHorizontally))},
diff --git a/Engine/Plugins/TextureGraph/Source/TextureGraph/Public/Expressions/Procedural/TG_Expression_Transform.h b/Engine/Plugins/TextureGraph/Source/TextureGraph/Public/Expressions/Procedural/TG_Expression_Transform.h
--- a/Engine/Plugins/TextureGraph/Source/TextureGraph/Public/Expressions/Procedural/TG_Expression_Transform.h
+++ b/Engine/Plugins/TextureGraph/Source/TextureGraph/Public/Expressions/Procedural/TG_Expression_Transform.h
@@ -58,12 +58,12 @@ public:
 	UPROPERTY(EditAnywhere, Category = NoCategory, meta = (TGType = "TG_Setting", UIMin = "-180", ClampMin = "-180", UIMax = "180", ClampMax = "180", Delta = "1", Units="Degrees"))
 	float									Rotation = 0;
 
-	// Uniform scaling of the cell aka the zoom effect:
-	// -1 => Zoom OUT => scale so that entire image fits in 1 pixel of the cell 
-	// 0 =>  No Zoom  => scale of 1 Source image completely visible in the cell
-	// +1 => Zoom IN  => scale so that one pixel of the source fills the cell 
-	UPROPERTY(EditAnywhere, Category = NoCategory, meta = (TGType = "TG_Setting", UIMin = "-1.0", ClampMin = "-1.0", UIMax = "1.0", ClampMax = "1.0", Multiple = "0.001"))
-	float									Zoom = 0.0f;
+	// Uniform scaling of the image within the tiled cell.
+	// Value is a percentage, default value is 100% meaning the image fully fill the cell:
+	// Zoom out with value in the range [0,100]
+	// Zoom in with value above 100%
+	UPROPERTY(EditAnywhere, Category = NoCategory, meta = (TGType = "TG_Setting", UIMin = "0", ClampMin = "0", UIMax = "100", ClampMax = "100000", Multiple = "1", Delta = "1", Units="Percent"))
+	float									Zoom = 100.0f;
 
 	// Keep the AspectRatio of the source image (0) or stretch to fit the cell rectangle (1)
 	UPROPERTY(EditAnywhere, Category = NoCategory, meta = (TGType = "TG_Setting", DisplayName = "Stretch to fit", UIMin = "0.0", ClampMin = "0.0", UIMax = "1.0", ClampMax = "1.0", Multiple = "0.01"))
