Fix crash on clamp nodes

diff --git a/Engine/Plugins/TextureGraph/Source/TextureGraph/Private/Expressions/Arrays/TG_Expression_Array.cpp b/Engine/Plugins/TextureGraph/Source/TextureGraph/Private/Expressions/Arrays/TG_Expression_Array.cpp
--- a/Engine/Plugins/TextureGraph/Source/TextureGraph/Private/Expressions/Arrays/TG_Expression_Array.cpp
+++ b/Engine/Plugins/TextureGraph/Source/TextureGraph/Private/Expressions/Arrays/TG_Expression_Array.cpp
@@ -17,8 +17,23 @@ void UTG_Expression_Array4::Evaluate(FTG_EvaluationContext* InContext)
 	//Output[2] = Input3;
 	//Output[3] = Input4;
 
-	Output.Set(0, Input1);
-	Output.Set(1, Input2);
-	Output.Set(2, Input3);
-	Output.Set(3, Input4);
+	if (Input1)
+		Output.Set(0, Input1);
+	else 
+		Output.Set(0, FTG_Texture::GetBlack());
+
+	if (Input2)
+		Output.Set(1, Input2);
+	else 
+		Output.Set(1, FTG_Texture::GetBlack());
+
+	if (Input3)
+		Output.Set(2, Input3);
+	else 
+		Output.Set(2, FTG_Texture::GetBlack());
+
+	if (Input4)
+		Output.Set(3, Input4);
+	else 
+		Output.Set(3, FTG_Texture::GetBlack());
 }
diff --git a/Engine/Plugins/TextureGraph/Source/TextureGraph/Private/Expressions/TG_Expression_Variant.cpp b/Engine/Plugins/TextureGraph/Source/TextureGraph/Private/Expressions/TG_Expression_Variant.cpp
--- a/Engine/Plugins/TextureGraph/Source/TextureGraph/Private/Expressions/TG_Expression_Variant.cpp
+++ b/Engine/Plugins/TextureGraph/Source/TextureGraph/Private/Expressions/TG_Expression_Variant.cpp
@@ -93,7 +93,9 @@ void UTG_Expression_Variant::Evaluate(FTG_EvaluationContext* InContext)
 {
 	Super::Evaluate(InContext);
 
-	switch (GetCommonInputVariantType())
+	FTG_Variant::EType VariantType = GetCommonInputVariantType();
+
+	switch (VariantType)
 	{
 	case FTG_Variant::EType::Scalar:
 		Output.EditScalar() = EvaluateScalar(InContext);
diff --git a/Engine/Plugins/TextureGraph/Source/TextureGraph/Private/TG_GraphEvaluation.cpp b/Engine/Plugins/TextureGraph/Source/TextureGraph/Private/TG_GraphEvaluation.cpp
--- a/Engine/Plugins/TextureGraph/Source/TextureGraph/Private/TG_GraphEvaluation.cpp
+++ b/Engine/Plugins/TextureGraph/Source/TextureGraph/Private/TG_GraphEvaluation.cpp
@@ -714,7 +714,7 @@ void FTG_Evaluation::EvaluateNodeArray(UTG_Node* InNode, const TArray<TObjectPtr
 				const FTG_Texture& ArrayTexture = ArrayValue.GetTexture();
 				FString ToType = OutputVar.Value.Argument.GetCPPTypeName().ToString();
 
-				if (ToType.Find("FTG_Variant") >= 0)
+				if (ToType == TEXT("FTG_Variant"))
 					OutputVar.Value.Var->EditAs<FTG_Variant>().EditTexture().Descriptor = ArrayTexture.Descriptor;
 				else if (ToType == TEXT("FTG_Texture"))
 					OutputVar.Value.Var->EditAs<FTG_Texture>().Descriptor = ArrayTexture.Descriptor;
@@ -724,6 +724,8 @@ void FTG_Evaluation::EvaluateNodeArray(UTG_Node* InNode, const TArray<TObjectPtr
 
 	for (int VarIndex = 0; VarIndex < MaxCount; VarIndex++)
 	{
+		FTG_Variant::EType CommonVariantType = FTG_Variant::EType::Invalid;
+
 		/// Run through array inputs, one at a time
 		for (int ArrayIndex = 0; ArrayIndex < ArrayInputs.Num(); ArrayIndex++)
 		{
@@ -746,11 +748,17 @@ void FTG_Evaluation::EvaluateNodeArray(UTG_Node* InNode, const TArray<TObjectPtr
 					->ReportWarning((int32)ETextureGraphErrorType::INPUT_ARRAY_WARNING, ErrorMsg, InNode);
 			}
 
+			const FTG_Variant& Variant = VarArray.GetArray()[VarArrayIndex];
+			if (Variant.GetType() > CommonVariantType)
+				CommonVariantType = Variant.GetType();
+
 			FTG_Evaluation::TransferVarToPin(ArrayPin, InContext, VarArrayIndex);
 		}
 
 		/// Now we evaluate the expression
 		UTG_Expression* Expression = InNode->GetExpression();
+		Expression->ResetCommonInputVariantType(CommonVariantType);
+
 		/// Trigger the Evaluation
 		Expression->SetupAndEvaluate(InContext);
 
@@ -784,6 +792,10 @@ void FTG_Evaluation::EvaluateNodeArray(UTG_Node* InNode, const TArray<TObjectPtr
 				/// We just copy it over in the array
 				ArrayOutput->Set(VarIndex, VariantVar.GetAs<FTG_Variant>());
 			}
+			else if (FromType == TEXT("FTG_VariantArray"))
+			{
+				ArrayOutput->CopyFrom(OutputVar.Value.Var->GetAs<FTG_VariantArray>());
+			}
 			else
 			{
 				ArrayOutput->Set(VarIndex, OutputVar.Value.Var->GetAs<FTG_Variant>());
diff --git a/Engine/Plugins/TextureGraph/Source/TextureGraph/Private/TG_Variant.cpp b/Engine/Plugins/TextureGraph/Source/TextureGraph/Private/TG_Variant.cpp
--- a/Engine/Plugins/TextureGraph/Source/TextureGraph/Private/TG_Variant.cpp
+++ b/Engine/Plugins/TextureGraph/Source/TextureGraph/Private/TG_Variant.cpp
@@ -348,6 +348,11 @@ void FTG_VariantArray::Set(int Index, const FTG_Texture& Tex)
 	Array[Index].EditTexture() = Tex;
 }
 
+void FTG_VariantArray::CopyFrom(const FTG_VariantArray& RHS)
+{
+	Array = RHS.GetArray();
+}
+
 void FTG_VariantArray::SetViewIndex(int Index)
 {
 	check(Index >= 0 && Index < Array.Num());
diff --git a/Engine/Plugins/TextureGraph/Source/TextureGraph/Public/TG_Variant.h b/Engine/Plugins/TextureGraph/Source/TextureGraph/Public/TG_Variant.h
--- a/Engine/Plugins/TextureGraph/Source/TextureGraph/Public/TG_Variant.h
+++ b/Engine/Plugins/TextureGraph/Source/TextureGraph/Public/TG_Variant.h
@@ -199,6 +199,7 @@ public:
 	void SetNum(int Count);
 	void Set(int Index, const FTG_Variant& Var);
 	void Set(int Index, const FTG_Texture& Tex);
+	void CopyFrom(const FTG_VariantArray& RHS);
 	void SetViewIndex(int Index);
 
 	FORCEINLINE bool IsTexture()
