Fix undo not working for texture nodes

diff --git a/Engine/Plugins/TextureGraph/Source/TextureGraph/Private/TG_Node.cpp b/Engine/Plugins/TextureGraph/Source/TextureGraph/Private/TG_Node.cpp
--- a/Engine/Plugins/TextureGraph/Source/TextureGraph/Private/TG_Node.cpp
+++ b/Engine/Plugins/TextureGraph/Source/TextureGraph/Private/TG_Node.cpp
@@ -46,17 +46,17 @@ void UTG_Node::PostEditUndo()
 
 void UTG_Node::OnPostUndo()
 {
-	if(IsValidChecked(this))
+	if (IsValidChecked(this))
 	{
 		// We dont know what got modified.
 		// For now, just copying all the values from the expression.
-		for(const auto& Pin : Pins)
+		for (const auto& Pin : Pins)
 		{
-			// exclude the case of input textures 
-			if(!(Pin->IsInput() && Pin->IsArgTexture()))
+			/// exclude the case of input textures but only for expressions that can't handle it (MaterialID change made by Zohaib Javed in CL: 31512086)
+			if (!(Pin->IsInput() && Pin->IsArgTexture() && Expression && Expression->IgnoreInputTextureOnUndo()))
 			{
 				auto Arg = Pin->GetArgument();
-				Pin->EditSelfVar()->CopyFrom(GetExpression(), Arg);	
+				Pin->EditSelfVar()->CopyFrom(GetExpression(), Arg);
 			}
 		}
 
diff --git a/Engine/Plugins/TextureGraph/Source/TextureGraph/Public/Expressions/TG_Expression.h b/Engine/Plugins/TextureGraph/Source/TextureGraph/Public/Expressions/TG_Expression.h
--- a/Engine/Plugins/TextureGraph/Source/TextureGraph/Public/Expressions/TG_Expression.h
+++ b/Engine/Plugins/TextureGraph/Source/TextureGraph/Public/Expressions/TG_Expression.h
@@ -148,6 +148,8 @@ protected:
 	// this call can be modified in a specific class to get a differnet behavior, see Expression_Dot for an example.
 	virtual FTG_Variant::EType EvalExpressionCommonOutputVariantType() const;
 
+	virtual bool IgnoreInputTextureOnUndo() const { return false; }
+
 
 	// In some cases, evaluation or change in the Expression needs to be feedback to the matching pin's value
 	// This is not needed for the standard flow of evaluation but is sometime required for coupled member.
diff --git a/Engine/Plugins/TextureGraph/Source/TextureGraph/Public/Expressions/Utilities/TG_Expression_MaterialID.h b/Engine/Plugins/TextureGraph/Source/TextureGraph/Public/Expressions/Utilities/TG_Expression_MaterialID.h
--- a/Engine/Plugins/TextureGraph/Source/TextureGraph/Public/Expressions/Utilities/TG_Expression_MaterialID.h
+++ b/Engine/Plugins/TextureGraph/Source/TextureGraph/Public/Expressions/Utilities/TG_Expression_MaterialID.h
@@ -49,6 +49,7 @@ public:
 	virtual void PostEditChangeProperty(FPropertyChangedEvent& PropertyChangedEvent) override;
 #endif
 	
+	virtual bool IgnoreInputTextureOnUndo() const override { return true; }
 	virtual FTG_Name GetDefaultName() const override { return TEXT("MaterialID"); }
 	virtual FText GetTooltipText() const override { return FText::FromString(TEXT("Extracts colors from the texture and creates a mask with selected colors")); }
 	void SetMaterialIdMask(int32 Id, bool bEnable);
