Background color information for Array grid

diff --git a/Engine/Plugins/TextureGraph/Source/TextureGraph/Private/Expressions/Arrays/TG_Expression_ArrayGrid.cpp b/Engine/Plugins/TextureGraph/Source/TextureGraph/Private/Expressions/Arrays/TG_Expression_ArrayGrid.cpp
--- a/Engine/Plugins/TextureGraph/Source/TextureGraph/Private/Expressions/Arrays/TG_Expression_ArrayGrid.cpp
+++ b/Engine/Plugins/TextureGraph/Source/TextureGraph/Private/Expressions/Arrays/TG_Expression_ArrayGrid.cpp
@@ -65,17 +65,17 @@ void UTG_Expression_ArrayGrid::Evaluate(FTG_EvaluationContext* InContext)
 		FTextureGraphErrorReporter* ErrorReporter = TextureGraphEngine::GetErrorReporter(InContext->Cycle->GetMix());
 		if (ErrorReporter)
 		{
-			ErrorReporter->ReportWarning((int32)ETextureGraphErrorType::NODE_WARNING, FString::Printf(TEXT("Number of rows and columns for the tiling do not match the number of inputs. Grid: %d x %d [Max Inputs: %d]. Padding with black images."), 
+			ErrorReporter->ReportWarning((int32)ETextureGraphErrorType::NODE_WARNING, FString::Printf(TEXT("Number of rows and columns for the tiling do not match the number of inputs. Grid: %d x %d [Max Inputs: %d]. Padding with transparent images (background color will blend through)."), 
 				OutRows, OutCols, Inputs.Num()), GetParentNode());
 		}
 
 		int32 NumDesiredEntires = OutRows * OutCols;
 		for (int32 Index = Inputs.Num(); Index < NumDesiredEntires; Index++)
 		{
-			Inputs.Add(TextureHelper::GetBlack());
+			Inputs.Add(TextureHelper::GetTransparent());
 		}
 	}
 
 	BufferDescriptor DesiredDesc = Output.Descriptor;
-	Output = T_ArrayGrid::Create(InContext->Cycle, DesiredDesc, Inputs, OutRows, OutCols);
+	Output = T_ArrayGrid::Create(InContext->Cycle, DesiredDesc, Inputs, OutRows, OutCols, BackgroundColor);
 }
diff --git a/Engine/Plugins/TextureGraph/Source/TextureGraph/Public/Expressions/Arrays/TG_Expression_ArrayGrid.h b/Engine/Plugins/TextureGraph/Source/TextureGraph/Public/Expressions/Arrays/TG_Expression_ArrayGrid.h
--- a/Engine/Plugins/TextureGraph/Source/TextureGraph/Public/Expressions/Arrays/TG_Expression_ArrayGrid.h
+++ b/Engine/Plugins/TextureGraph/Source/TextureGraph/Public/Expressions/Arrays/TG_Expression_ArrayGrid.h
@@ -28,6 +28,10 @@ public:
 	UPROPERTY(meta = (TGType = "TG_Input", PinDisplayName = ""))
 	int32							Columns = 0;
 
+	/// The background color to use for the output texture
+	UPROPERTY(meta = (TGType = "TG_Input"))
+	FLinearColor					BackgroundColor = FLinearColor::Transparent;
+
 	/// A single output texture
 	UPROPERTY(meta = (TGType = "TG_Output", PinDisplayName = ""))
 	FTG_Texture						Output;
diff --git a/Engine/Plugins/TextureGraph/Source/TextureGraphEngine/2D/TextureHelper.cpp b/Engine/Plugins/TextureGraph/Source/TextureGraphEngine/2D/TextureHelper.cpp
--- a/Engine/Plugins/TextureGraph/Source/TextureGraphEngine/2D/TextureHelper.cpp
+++ b/Engine/Plugins/TextureGraph/Source/TextureGraphEngine/2D/TextureHelper.cpp
@@ -44,6 +44,7 @@
 #include "Editor.h"
 #endif
 
+TiledBlobPtr TextureHelper::GTransparent;
 TiledBlobPtr TextureHelper::GBlack;
 TiledBlobPtr TextureHelper::GWhite;
 TiledBlobPtr TextureHelper::GGray;
@@ -133,6 +134,7 @@ void TextureHelper::InitStockTextures()
 	std::vector<cti::continuable<bool>> Promises;
 
 	// TODO: Need these values to be linear.
+	Promises.push_back(InitSolidTexture(&GTransparent, FLinearColor::Transparent, TEXT("TRANSPARENT"), AlbedoDesc));
 	Promises.push_back(InitSolidTexture(&GBlack, FLinearColor::Black, TEXT("BLACK"), AlbedoDesc));
 	Promises.push_back(InitSolidTexture(&GWhite, FLinearColor::White, TEXT("WHITE"), LinearDesc));
 	Promises.push_back(InitSolidTexture(&GGray, FLinearColor::Gray, TEXT("GRAY"), AlbedoDesc));
@@ -160,6 +162,7 @@ void TextureHelper::FreeStockTextures()
 {
 	check(IsInGameThread());
 
+	GTransparent = nullptr;
 	GBlack = nullptr;
 	GWhite = nullptr;
 	GGray = nullptr;
diff --git a/Engine/Plugins/TextureGraph/Source/TextureGraphEngine/2D/TextureHelper.h b/Engine/Plugins/TextureGraph/Source/TextureGraphEngine/2D/TextureHelper.h
--- a/Engine/Plugins/TextureGraph/Source/TextureGraphEngine/2D/TextureHelper.h
+++ b/Engine/Plugins/TextureGraph/Source/TextureGraphEngine/2D/TextureHelper.h
@@ -25,6 +25,7 @@ typedef cti::continuable<FLinearColor>	AsyncLinearColor;
 
 struct TEXTUREGRAPHENGINE_API TextureHelper
 {
+	static TiledBlobPtr					GTransparent;			/// Transparent RGBA texture
 	static TiledBlobPtr					GBlack;					/// Black RGBA texture
 	static TiledBlobPtr					GWhite;					/// White RGBA texture
 	static TiledBlobPtr					GGray;					/// Gray RGBA texture
@@ -41,6 +42,7 @@ struct TEXTUREGRAPHENGINE_API TextureHelper
 	static void							InitStockTextures();
 	static void							FreeStockTextures();
 
+	static FORCEINLINE TiledBlobPtr		GetTransparent() { check(GTransparent); return GTransparent; }
 	static FORCEINLINE TiledBlobPtr		GetBlack() { check(GBlack); return GBlack; }
 	static FORCEINLINE TiledBlobPtr		GetWhite() { check(GWhite); return GWhite; }
 	static FORCEINLINE TiledBlobPtr		GetGray() { check(GGray);  return GGray; }
diff --git a/Engine/Plugins/TextureGraph/Source/TextureGraphEngine/Transform/Expressions/T_ArrayGrid.cpp b/Engine/Plugins/TextureGraph/Source/TextureGraphEngine/Transform/Expressions/T_ArrayGrid.cpp
--- a/Engine/Plugins/TextureGraph/Source/TextureGraphEngine/Transform/Expressions/T_ArrayGrid.cpp
+++ b/Engine/Plugins/TextureGraph/Source/TextureGraphEngine/Transform/Expressions/T_ArrayGrid.cpp
@@ -149,7 +149,8 @@ protected:
 };
 
 //////////////////////////////////////////////////////////////////////////
-TiledBlobPtr T_ArrayGrid::Create(MixUpdateCyclePtr Cycle, BufferDescriptor& InDesiredOutputDesc, TArray<TiledBlobPtr> Inputs, int32 NumRows, int32 NumCols, int32 TargetId)
+TiledBlobPtr T_ArrayGrid::Create(MixUpdateCyclePtr Cycle, BufferDescriptor& InDesiredOutputDesc, TArray<TiledBlobPtr> Inputs, 
+	int32 NumRows, int32 NumCols, FLinearColor BackgroundColor, int32 TargetId)
 {
 	check(!Inputs.IsEmpty());
 
@@ -172,7 +173,7 @@ TiledBlobPtr T_ArrayGrid::Create(MixUpdateCyclePtr Cycle, BufferDescriptor& InDe
 	{
 		.bFixed = false,
 		.bMaintainAspectRatio = true,
-		.BackgroundColor = FLinearColor::Transparent
+		.BackgroundColor = BackgroundColor
 	};
 
 	BufferDescriptor InputDesc = Inputs[0]->GetDescriptor();
diff --git a/Engine/Plugins/TextureGraph/Source/TextureGraphEngine/Transform/Expressions/T_ArrayGrid.h b/Engine/Plugins/TextureGraph/Source/TextureGraphEngine/Transform/Expressions/T_ArrayGrid.h
--- a/Engine/Plugins/TextureGraph/Source/TextureGraphEngine/Transform/Expressions/T_ArrayGrid.h
+++ b/Engine/Plugins/TextureGraph/Source/TextureGraphEngine/Transform/Expressions/T_ArrayGrid.h
@@ -13,6 +13,7 @@
 class TEXTUREGRAPHENGINE_API T_ArrayGrid
 {
 public:
-	static TiledBlobPtr Create(MixUpdateCyclePtr Cycle, BufferDescriptor& DesiredOutputDesc, TArray<TiledBlobPtr> Inputs, int32 NumRows, int32 NumCols, int32 TargetId = 0);
+	static TiledBlobPtr Create(MixUpdateCyclePtr Cycle, BufferDescriptor& DesiredOutputDesc, TArray<TiledBlobPtr> Inputs, 
+		int32 NumRows, int32 NumCols, FLinearColor BackgroundColor, int32 TargetId = 0);
 };
 
diff --git a/Engine/Plugins/TextureGraph/Source/TextureGraphEngine/Transform/Utility/T_CombineTiledBlob.cpp b/Engine/Plugins/TextureGraph/Source/TextureGraphEngine/Transform/Utility/T_CombineTiledBlob.cpp
--- a/Engine/Plugins/TextureGraph/Source/TextureGraphEngine/Transform/Utility/T_CombineTiledBlob.cpp
+++ b/Engine/Plugins/TextureGraph/Source/TextureGraphEngine/Transform/Utility/T_CombineTiledBlob.cpp
@@ -256,7 +256,7 @@ AsyncTransformResultPtr	CombineTiledBlob_Transform::Exec(const TransformArgs& Ar
 	DrawTilesSettings* DrawSettingsToUse = nullptr;
 	InitDrawSettings(Target, Tiles);
 
-	FLinearColor* ClearColor = nullptr;
+	FLinearColor* ClearColor = &Settings.BackgroundColor;
 
 	if (DrawSettings.Position.Rows() == Tiles.Rows())
 	{
