Better tooltip for the ShowDebugGrid option on the Transform node

diff --git a/Engine/Plugins/TextureGraph/Source/TextureGraph/Public/Expressions/Procedural/TG_Expression_Transform.h b/Engine/Plugins/TextureGraph/Source/TextureGraph/Public/Expressions/Procedural/TG_Expression_Transform.h
--- a/Engine/Plugins/TextureGraph/Source/TextureGraph/Public/Expressions/Procedural/TG_Expression_Transform.h
+++ b/Engine/Plugins/TextureGraph/Source/TextureGraph/Public/Expressions/Procedural/TG_Expression_Transform.h
@@ -89,7 +89,7 @@ public:
 	UPROPERTY(EditAnywhere, Category = NoCategory, meta = (TGType = "TG_Setting", UIMin = "0", ClampMin = "0", UIMax = "31", ClampMax = "31", Delta = "1"))
 	FVector2f								Stride = { 0, 0 };
 
-	// Output blended with the debugging grid showing the transformation applied
+	// [DEBUG ONLY OPTION] Output blended with the debugging grid showing the transformation applied. This option WILL NOT get applied while exporting!
 	UPROPERTY(EditAnywhere, Category = NoCategory, meta = (TGType = "TG_Setting", UIMin = "0", ClampMin = "0", UIMax = "1", ClampMax = "1"))
 	float									 ShowDebugGrid = 0;
 
