Test fixes

diff --git a/Engine/Plugins/TextureGraph/Source/TextureGraph/Private/TG_Node.cpp b/Engine/Plugins/TextureGraph/Source/TextureGraph/Private/TG_Node.cpp
--- a/Engine/Plugins/TextureGraph/Source/TextureGraph/Private/TG_Node.cpp
+++ b/Engine/Plugins/TextureGraph/Source/TextureGraph/Private/TG_Node.cpp
@@ -53,7 +53,8 @@ void UTG_Node::OnPostUndo()
 		for (const auto& Pin : Pins)
 		{
 			/// exclude the case of input textures but only for expressions that can't handle it (MaterialID change made by Zohaib Javed in CL: 31512086)
-			if (!(Pin->IsInput() && Pin->IsArgTexture() && Expression && Expression->IgnoreInputTextureOnUndo()))
+			//&& Expression && Expression->IgnoreInputTextureOnUndo()
+			if (!(Pin->IsInput() && Pin->IsArgTexture()))
 			{
 				auto Arg = Pin->GetArgument();
 				Pin->EditSelfVar()->CopyFrom(GetExpression(), Arg);
