Test fixes

diff --git a/Engine/Plugins/TextureGraph/Source/TextureGraph/Private/TG_Node.cpp b/Engine/Plugins/TextureGraph/Source/TextureGraph/Private/TG_Node.cpp
--- a/Engine/Plugins/TextureGraph/Source/TextureGraph/Private/TG_Node.cpp
+++ b/Engine/Plugins/TextureGraph/Source/TextureGraph/Private/TG_Node.cpp
@@ -53,8 +53,7 @@ void UTG_Node::OnPostUndo()
 		for (const auto& Pin : Pins)
 		{
 			/// exclude the case of input textures but only for expressions that can't handle it (MaterialID change made by Zohaib Javed in CL: 31512086)
-			//&& Expression && Expression->IgnoreInputTextureOnUndo()
-			if (!(Pin->IsInput() && Pin->IsArgTexture()))
+			if (!(Pin->IsInput() && Pin->IsArgTexture() && Expression && Expression->IgnoreInputTextureOnUndo()))
 			{
 				auto Arg = Pin->GetArgument();
 				Pin->EditSelfVar()->CopyFrom(GetExpression(), Arg);
diff --git a/Engine/Plugins/TextureGraph/Source/TextureGraph/Public/Expressions/Input/TG_Expression_Texture.h b/Engine/Plugins/TextureGraph/Source/TextureGraph/Public/Expressions/Input/TG_Expression_Texture.h
--- a/Engine/Plugins/TextureGraph/Source/TextureGraph/Public/Expressions/Input/TG_Expression_Texture.h
+++ b/Engine/Plugins/TextureGraph/Source/TextureGraph/Public/Expressions/Input/TG_Expression_Texture.h
@@ -48,6 +48,7 @@ public:
 	virtual void SetTitleName(FName NewName) override;
 	virtual FName GetTitleName() const override;
 	virtual FText GetTooltipText() const override { return FText::FromString(TEXT("Makes an existing texture asset available. It is automatically exposed as a graph input parameter.")); }
+	virtual bool IgnoreInputTextureOnUndo() const override { return false; }
 
 	virtual void SetAsset(UObject* Asset) override;
 	virtual bool CanHandleAsset(UObject* Asset) override;
diff --git a/Engine/Plugins/TextureGraph/Source/TextureGraph/Public/Expressions/TG_Expression.h b/Engine/Plugins/TextureGraph/Source/TextureGraph/Public/Expressions/TG_Expression.h
--- a/Engine/Plugins/TextureGraph/Source/TextureGraph/Public/Expressions/TG_Expression.h
+++ b/Engine/Plugins/TextureGraph/Source/TextureGraph/Public/Expressions/TG_Expression.h
@@ -148,8 +148,7 @@ protected:
 	// this call can be modified in a specific class to get a differnet behavior, see Expression_Dot for an example.
 	virtual FTG_Variant::EType EvalExpressionCommonOutputVariantType() const;
 
-	virtual bool IgnoreInputTextureOnUndo() const { return false; }
-
+	virtual bool IgnoreInputTextureOnUndo() const { return true; }
 
 	// In some cases, evaluation or change in the Expression needs to be feedback to the matching pin's value
 	// This is not needed for the standard flow of evaluation but is sometime required for coupled member.
diff --git a/Engine/Plugins/TextureGraph/Source/TextureGraph/Public/Expressions/Utilities/TG_Expression_MaterialID.h b/Engine/Plugins/TextureGraph/Source/TextureGraph/Public/Expressions/Utilities/TG_Expression_MaterialID.h
--- a/Engine/Plugins/TextureGraph/Source/TextureGraph/Public/Expressions/Utilities/TG_Expression_MaterialID.h
+++ b/Engine/Plugins/TextureGraph/Source/TextureGraph/Public/Expressions/Utilities/TG_Expression_MaterialID.h
@@ -49,7 +49,6 @@ public:
 	virtual void PostEditChangeProperty(FPropertyChangedEvent& PropertyChangedEvent) override;
 #endif
 	
-	virtual bool IgnoreInputTextureOnUndo() const override { return true; }
 	virtual FTG_Name GetDefaultName() const override { return TEXT("MaterialID"); }
 	virtual FText GetTooltipText() const override { return FText::FromString(TEXT("Extracts colors from the texture and creates a mask with selected colors")); }
 	void SetMaterialIdMask(int32 Id, bool bEnable);
