Fix histogram and node preview showing blank on texture input change

diff --git a/Engine/Plugins/TextureGraph/Source/TextureGraphEditor/Private/STG_NodePreview.cpp b/Engine/Plugins/TextureGraph/Source/TextureGraphEditor/Private/STG_NodePreview.cpp
--- a/Engine/Plugins/TextureGraph/Source/TextureGraphEditor/Private/STG_NodePreview.cpp
+++ b/Engine/Plugins/TextureGraph/Source/TextureGraphEditor/Private/STG_NodePreview.cpp
@@ -28,18 +28,55 @@ UE::ImageWidgets::IImageViewer::FImageInfo FNodeViewer::GetCurrentImageInfo() co
 	return {{}, FIntPoint::ZeroValue, 0, false};
 }
 
+void FNodeViewer::UpdateTexture()
+{
+	if (!CurrentBlob)
+		return;
+
+	if (CurrentBlob->IsTiled())
+	{
+		CurrentBlob->OnFinalise()
+			.then([this]()
+				{
+					UE_LOG(LogTemp, Log, TEXT("!!!Finalised. Combining!!!"));
+					const TiledBlobPtr BlobTiled = std::static_pointer_cast<TiledBlob>(CurrentBlob);
+					return BlobTiled->CombineTiles(false, false);
+				})
+			.then([this]()
+				{
+					UE_LOG(LogTemp, Log, TEXT("!!!SETTING TEXTURE!!!"));
+					FIntPoint ImageSize = FIntPoint {(int32)CurrentBlob->GetWidth(), (int32)CurrentBlob->GetHeight()};
+					SetTexture(CurrentBlob);
+				});
+	}
+	else
+	{
+		CurrentBlob->OnFinalise().then([this]()
+			{
+				SetTexture(CurrentBlob);
+			});
+	}
+}
+
 void FNodeViewer::DrawCurrentImage(FViewport*, FCanvas* Canvas, const FDrawProperties& Properties)
 {
+	check(NodeTexture);
 	if (NodeTexture)
 	{
-		if (FTextureResource* TextureResource = NodeTexture->GetResource(); TextureResource != nullptr)
+		FTextureResource* TextureResource = NodeTexture->GetResource();
+		if (TextureResource != nullptr)
 		{
 			TextureResource->bGreyScaleFormat = IsSingleChannel();
 			TextureResource->bSRGB = IsSRGB();
 
 			DrawTexture(TextureResource, Canvas, Properties.Placement, Properties.Mip);
+			return;
 		}
 	}
+	Util::OnGameThread([this]()
+		{
+			UpdateTexture();
+		});
 }
 
 
@@ -104,11 +141,13 @@ void FNodeViewer::SetTexture(const BlobPtr& InBlob, const FLinearColor InClearCo
 
 	if (InBlob)
 	{
-		if (const DeviceBufferPtr Buffer = InBlob->GetBufferRef().GetPtr())
-		{
-			NodeTexture = GetTextureFromBuffer(Buffer);
-			NodeDescriptor = Buffer->Descriptor();
-		}
+		const DeviceBufferPtr Buffer = InBlob->GetBufferRef().GetPtr();
+		check(Buffer);
+
+		NodeTexture = GetTextureFromBuffer(Buffer);
+		NodeDescriptor = Buffer->Descriptor();
+
+		check(NodeDescriptor.Width > 0 && NodeDescriptor.Height > 0);
 	}
 }
 
@@ -365,7 +404,9 @@ void STG_NodePreviewWidget::Update() const
 					.then([this, Blob]()
 						{
 							UE_LOG(LogTemp, Log, TEXT("!!!SETTING TEXTURE!!!"));
+							FIntPoint ImageSize = FIntPoint {(int32)Blob->GetWidth(), (int32)Blob->GetHeight()};
 							NodeViewer->SetTexture(Blob);
+							Viewport->ResetZoom(ImageSize);
 						});
 			}
 			else
@@ -382,11 +423,12 @@ void STG_NodePreviewWidget::Update() const
 		}
 		else
 		{
-			NodeViewer->SetTexture(nullptr);
+			NodeViewer->SetTexture(TextureHelper::GetBlack());
 		}
 
 		NodeViewer->SetLabelText(GetLabelText(Variant, bValidVariant));
 		Viewport->ResetZoom(ImageSize);
+		[[maybe_unused]] bool Result = OnNodeBlobChanged.ExecuteIfBound(Blob);
 
 		return;
 	}
diff --git a/Engine/Plugins/TextureGraph/Source/TextureGraphEditor/Private/STG_NodePreview.h b/Engine/Plugins/TextureGraph/Source/TextureGraphEditor/Private/STG_NodePreview.h
--- a/Engine/Plugins/TextureGraph/Source/TextureGraphEditor/Private/STG_NodePreview.h
+++ b/Engine/Plugins/TextureGraph/Source/TextureGraphEditor/Private/STG_NodePreview.h
@@ -42,6 +42,7 @@ public:
 
 	/** Sets the node buffer to a given node. */
 	void SetTexture(const BlobPtr& InBlob, FLinearColor InClearColor = FLinearColor(0.1, 0.1, 0.1, 1));
+	void UpdateTexture();
 
 	/** Sets the label text. */
 	void SetLabelText(FText InText);
