Fix warnings in shaders

#tests TextureScriptTests

diff --git a/Engine/Plugins/TextureGraph/Shaders/Expressions/Expression_Blend.usf b/Engine/Plugins/TextureGraph/Shaders/Expressions/Expression_Blend.usf
--- a/Engine/Plugins/TextureGraph/Shaders/Expressions/Expression_Blend.usf
+++ b/Engine/Plugins/TextureGraph/Shaders/Expressions/Expression_Blend.usf
@@ -67,7 +67,7 @@ float4 FSH_BlendSubtract(in float2 uv : TEXCOORD0) : SV_Target0
 	float MaskingValue = GetMask(MaskTexture, uv, Opacity);
 
 #if IGNORE_ALPHA
-	float3 FinalRGB = lerp(BackgroundTex, BackgroundTex - ForegroundTex, MaskingValue);
+	float3 FinalRGB = lerp(BackgroundTex.rgb, BackgroundTex.rgb - ForegroundTex.rgb, MaskingValue);
 	float4 FinalColor = float4(FinalRGB, BackgroundTex.a);
 #else
 	float4 FinalColor = lerp(BackgroundTex, BackgroundTex - ForegroundTex, MaskingValue);
@@ -88,7 +88,7 @@ float4 FSH_BlendMultiply(in float2 uv : TEXCOORD0) : SV_Target0
 	float MaskingValue = GetMask(MaskTexture, uv, Opacity);
 
 #if IGNORE_ALPHA
-	float3 FinalRGB = lerp(BackgroundTex, ForegroundTex * BackgroundTex, MaskingValue);
+	float3 FinalRGB = lerp(BackgroundTex.rgb, ForegroundTex.rgb * BackgroundTex.rgb, MaskingValue);
 	float4 FinalColor = float4(FinalRGB, BackgroundTex.a);
 #else
 	float4 FinalColor = lerp(BackgroundTex, ForegroundTex * BackgroundTex, MaskingValue);
diff --git a/Engine/Plugins/TextureGraph/Shaders/Expressions/Expression_Levels.usf b/Engine/Plugins/TextureGraph/Shaders/Expressions/Expression_Levels.usf
--- a/Engine/Plugins/TextureGraph/Shaders/Expressions/Expression_Levels.usf
+++ b/Engine/Plugins/TextureGraph/Shaders/Expressions/Expression_Levels.usf
@@ -91,7 +91,7 @@ float4 FSH_Levels(float2 uv : TEXCOORD0) : SV_Target0
 #if CONVERT_TO_GRAYSCALE
 	// The new  Grayscale level is used to scale the color darker or brighter
 	float Ratio = (SourceValue == 0 ? 0.0 : LeveledValue / SourceValue); 	
-	float3 LeveledColor = Color * Ratio;
+	float3 LeveledColor = Color.rgb * Ratio;
 #else
 	float3 LeveledColor = LeveledValue;
 #endif
diff --git a/Engine/Plugins/TextureGraph/Shaders/Expressions/Expression_RadialBlur.usf b/Engine/Plugins/TextureGraph/Shaders/Expressions/Expression_RadialBlur.usf
--- a/Engine/Plugins/TextureGraph/Shaders/Expressions/Expression_RadialBlur.usf
+++ b/Engine/Plugins/TextureGraph/Shaders/Expressions/Expression_RadialBlur.usf
@@ -13,7 +13,7 @@ float Samples;
 
 float4 FSH_RadialBlur(float2 UV : TEXCOORD0) : SV_Target0
 {
-	const float Center = float2(CenterX, CenterY);
+	const float2 Center = float2(CenterX, CenterY);
 	
 	float2 FullUV = TileInfo_fromCurrentTileToLayer(UV);
 	float4 Col = float4(0,0,0,0);
diff --git a/Engine/Plugins/TextureGraph/Shaders/Layer/AdjustColor.usf b/Engine/Plugins/TextureGraph/Shaders/Layer/AdjustColor.usf
--- a/Engine/Plugins/TextureGraph/Shaders/Layer/AdjustColor.usf
+++ b/Engine/Plugins/TextureGraph/Shaders/Layer/AdjustColor.usf
@@ -30,7 +30,7 @@ float4 FSH_AdjustColor(in float2 uvSource : TEXCOORD0) : SV_Target0
 			sourceColor.a = 1 - sourceColor.a;
 		#endif
 		float shiftedAlpha = saturate(ColorOffset.r + sourceColor.a);
-		float finalAlpha = lerp(AverageColor, shiftedAlpha, Contrast);
+		float finalAlpha = lerp(AverageColor.a, shiftedAlpha, Contrast);
 		return float4(sourceColor.rgb, finalAlpha);	
 	#else
 		#if INVERT
diff --git a/Engine/Plugins/TextureGraph/Shaders/Mask/Blend/MaskBlend_Multiply.usf b/Engine/Plugins/TextureGraph/Shaders/Mask/Blend/MaskBlend_Multiply.usf
--- a/Engine/Plugins/TextureGraph/Shaders/Mask/Blend/MaskBlend_Multiply.usf
+++ b/Engine/Plugins/TextureGraph/Shaders/Mask/Blend/MaskBlend_Multiply.usf
@@ -12,8 +12,8 @@ Texture2D	Mask1;
 
 float4 FSH_MaskBlend_Multiply(in float2 uv : TEXCOORD0) : SV_Target0
 {
-	float m0 = MaskFactor0 * Mask0.Sample(SamplerStates_Linear_Clamp, uv);
-	float m1 = MaskFactor1 * Mask1.Sample(SamplerStates_Linear_Clamp, uv);
+	float m0 = MaskFactor0 * Mask0.Sample(SamplerStates_Linear_Clamp, uv).r;
+	float m1 = MaskFactor1 * Mask1.Sample(SamplerStates_Linear_Clamp, uv).r;
 	float m = saturate(m0 * m1);
 	
     return float4(m, m, m, 1);
diff --git a/Engine/Plugins/TextureGraph/Shaders/Mask/LayerMask.usf b/Engine/Plugins/TextureGraph/Shaders/Mask/LayerMask.usf
--- a/Engine/Plugins/TextureGraph/Shaders/Mask/LayerMask.usf
+++ b/Engine/Plugins/TextureGraph/Shaders/Mask/LayerMask.usf
@@ -112,7 +112,7 @@ float4 FSH_LayerMask(in float2 uv : TEXCOORD0) : SV_Target0
 #if PAINTMASK//apply paint mask
 		float4 pSample = PaintMask.Sample(SamplerStates_Wrap, uv);
 		float removePremul = saturate(pSample.r / pSample.a);
-		output = lerp(output, float4(removePremul, removePremul, removePremul, 1), pSample.a);
+		output = lerp(output, removePremul, pSample.a);
 #endif
 
 #if UVMASK //apply uv mask
diff --git a/Engine/Plugins/TextureGraph/Shaders/Mask/NormalMask.usf b/Engine/Plugins/TextureGraph/Shaders/Mask/NormalMask.usf
--- a/Engine/Plugins/TextureGraph/Shaders/Mask/NormalMask.usf
+++ b/Engine/Plugins/TextureGraph/Shaders/Mask/NormalMask.usf
@@ -36,7 +36,7 @@ float4 FSH_NormalMask(float2 uv : TEXCOORD0) : SV_Target0
 	float4 tangent = MeshTangents.Sample(SamplerStates_Linear_Clamp, uv);
     float4 tangentOS = float4(tangent.y, tangent.z, tangent.x, tangent.w) * 2.0 - 1.0;	
 
-    float3x3 tangentToWorld = GetTransform_TS_to_WS(LocalToWorldMatrix, normalOS, tangentOS);
+    float3x3 tangentToWorld = GetTransform_TS_to_WS((float3x3)LocalToWorldMatrix, normalOS, tangentOS);
 
     normalsForMask = (ObjectSpace * mul(tangentToWorld, normalTS)) + ((1 - ObjectSpace) * normalOS);
     
diff --git a/Engine/Plugins/TextureGraph/Shaders/Utils/Histogram_comp.usf b/Engine/Plugins/TextureGraph/Shaders/Utils/Histogram_comp.usf
--- a/Engine/Plugins/TextureGraph/Shaders/Utils/Histogram_comp.usf
+++ b/Engine/Plugins/TextureGraph/Shaders/Utils/Histogram_comp.usf
@@ -37,7 +37,7 @@ void CSH_HistogramPerTile(
 	float4 Fetched = SourceTiles.Load(uint3(ThreadId.xy, 0));
 	
 	// Evaluate Luma for 4th component and find the bin index for each channel
-	float Luma = Grayscale(Fetched);
+	float Luma = Grayscale(Fetched.rgb);
 	// The considered value is in the range [0, 1];
 	float4 Value = saturate(float4(Fetched.rgb, Luma));
 	// From the float value to the matching histogram bin [0,255]
