Color correction shader.

diff --git a/Engine/Plugins/TextureGraph/Shaders/Expressions/Expression_Brightness.usf b/Engine/Plugins/TextureGraph/Shaders/Expressions/Expression_Brightness.usf
--- a/Engine/Plugins/TextureGraph/Shaders/Expressions/Expression_Brightness.usf
+++ b/Engine/Plugins/TextureGraph/Shaders/Expressions/Expression_Brightness.usf
@@ -1,6 +1,7 @@
 // Copyright Epic Games, Inc. All Rights Reserved.
 #include "/Engine/Public/Platform.ush"
 #include "/Plugin/TextureGraph/SamplerStates.ush"
+#include "/Plugin/TextureGraph/ShaderUtil.ush"
 
 Texture2D SourceTexture;
 float Brightness;
@@ -9,8 +10,5 @@ float Contrast;
 float4 FSH_Brightness(float2 uv : TEXCOORD0) : SV_Target0
 {
 	float4 Color = SourceTexture.Sample(SamplerStates_NoBorder, uv);
-	float MidPoint = 0.5 + Brightness;
-
-	// TODO: Support HDR values
-	return saturate(lerp(MidPoint, Color + float4(Brightness, Brightness, Brightness, Brightness), float4(Contrast, Contrast, Contrast, Contrast)));
+	return AdjustBrightnessAndContrast(Color, Brightness, Contrast);
 }
diff --git a/Engine/Plugins/TextureGraph/Shaders/Expressions/Expression_ColorCorrection.usf b/Engine/Plugins/TextureGraph/Shaders/Expressions/Expression_ColorCorrection.usf
new file mode 100644
--- /dev/null
+++ b/Engine/Plugins/TextureGraph/Shaders/Expressions/Expression_ColorCorrection.usf
@@ -0,0 +1,34 @@
+// Copyright Epic Games, Inc. All Rights Reserved.
+#include "/Engine/Public/Platform.ush"
+#include "/Plugin/TextureGraph/SamplerStates.ush"
+#include "/Plugin/TextureGraph/ShaderUtil.ush"
+
+Texture2D Input;
+float Brightness;
+float Contrast;
+float Gain;
+float Lift;
+float Gamma;
+float Saturation;
+float4 TemperatureRGB;
+float TemperatureStrength;
+float TemperatureBrightnessNormalization;
+
+float4 FSH_ColorCorrection(float2 uv : TEXCOORD0) : SV_Target0
+{
+	const float3 HALF = float3(0.5, 0.5, 0.5);
+
+	float4 Clr = Input.Sample(SamplerStates_NoBorder, uv);
+	float4 ClrA = AdjustBrightnessAndContrast_NoAlpha(Clr, Brightness, Contrast);
+	
+	float OriginalBrightness = LumaFromRGB(Clr.rgb);
+	float3 ClrTemp = lerp(ClrA.rgb, ClrA.rgb * TemperatureRGB.rgb, TemperatureStrength);
+	float NewBrightness = LumaFromRGB(ClrTemp.rgb);
+	float3 ClrNew = ClrTemp.rgb * lerp(1.0, (NewBrightness > 1e-6) ? (OriginalBrightness / NewBrightness) : 1.0, TemperatureBrightnessNormalization);
+	float TemperatureBrightness = LumaFromRGB(ClrNew);
+	float3 TemperatureBrightnessVec = float3(TemperatureBrightness, TemperatureBrightness, TemperatureBrightness);
+	ClrNew = lerp(TemperatureBrightnessVec, ClrNew, Saturation);
+	
+	float3 ClrFinal = pow(ClrNew, HALF / Gamma); 
+	return float4(ClrFinal.rgb, Clr.a);
+}
diff --git a/Engine/Plugins/TextureGraph/Shaders/ShaderUtil.ush b/Engine/Plugins/TextureGraph/Shaders/ShaderUtil.ush
--- a/Engine/Plugins/TextureGraph/Shaders/ShaderUtil.ush
+++ b/Engine/Plugins/TextureGraph/Shaders/ShaderUtil.ush
@@ -45,11 +45,26 @@ float3x3 GetTransform_TS_to_WS(float3x3 _ObjectToWorld, float3 worldNormal, floa
 	return tangentToWorld;
 }
 
+float4 AdjustBrightnessAndContrast(float4 Color, float Brightness, float Contrast)
+{
+	float MidPoint = 0.5 + Brightness;
+
+	// TODO: Support HDR values
+	return saturate(lerp(MidPoint, Color + float4(Brightness, Brightness, Brightness, Brightness), float4(Contrast, Contrast, Contrast, Contrast)));	
+}
+
+float4 AdjustBrightnessAndContrast_NoAlpha(float4 Color, float Brightness, float Contrast)
+{
+	float4 Result = AdjustBrightnessAndContrast(Color, Brightness, Contrast);
+	return float4(Result.rgb, Color.a);
+}
 
 float LumaFromRGB(float3 Color)
 {
 	// ITU-R BT.709 colour luminance gradients
-	return 0.2126 * Color.r + 0.7152 * Color.g + 0.0722* Color.b;
+	const float3 BrightnessFactor = float3(0.2126, 0.7152, 0.0722);
+	return dot(Color, BrightnessFactor);
+	//0.2126 * Color.r + 0.7152 * Color.g + 0.0722 * Color.b;
 }
 
 
diff --git a/Engine/Plugins/TextureGraph/Source/TextureGraph/Private/Expressions/Color/TG_Expression_ColorCorrection.cpp b/Engine/Plugins/TextureGraph/Source/TextureGraph/Private/Expressions/Color/TG_Expression_ColorCorrection.cpp
new file mode 100644
--- /dev/null
+++ b/Engine/Plugins/TextureGraph/Source/TextureGraph/Private/Expressions/Color/TG_Expression_ColorCorrection.cpp
@@ -0,0 +1,66 @@
+// Copyright Epic Games, Inc. All Rights Reserved.
+
+#include "Expressions/Color/TG_Expression_ColorCorrection.h"
+#include "FxMat/MaterialManager.h"
+#include "Job/JobBatch.h"
+#include "Transform/Expressions/T_Color.h"
+#include "TextureGraphEngine/Helper/ColorUtil.h"
+#include "TextureGraphEngine/Helper/GraphicsUtil.h"
+
+static FLinearColor ColorTemperatureRGB(float Temperature)
+{
+	static const int NumElements = 3;
+	static const float CoolConv[NumElements][NumElements] = {
+		{ 0.0f, -2902.1955373783176f, -8257.7997278925690f },
+		{ 0.0f, 1669.5803561666639f, 2575.2827530017594f },
+		{ 1.0f, 1.3302673723350029f, 1.8993753891711275f }
+	};
+
+	static const float WarmConv[NumElements][NumElements] = {
+		{ 1745.0425298314172f, 1216.6168361476490f, -8257.7997278925690f },
+		{ -2666.3474220535695f, -2173.1012343082230f, 2575.2827530017594f },
+		{ 0.55995389139931482f, 0.70381203140554553f, 1.8993753891711275f }
+	};
+
+	const auto Mat = (Temperature <= 6500.0f) ? CoolConv : WarmConv;
+	float V0 = FMath::Clamp(Mat[0][0] / (Temperature + Mat[1][0]) + Mat[2][0], 0.0f, 1.0f);
+	float V1 = FMath::Clamp(Mat[0][1] / (Temperature + Mat[1][1]) + Mat[2][1], 0.0f, 1.0f);
+	float V2 = FMath::Clamp(Mat[0][2] / (Temperature + Mat[1][2]) + Mat[2][2], 0.0f, 1.0f);
+
+	return FLinearColor(V0, V1, V2, 0.0f);
+}
+
+void UTG_Expression_ColorCorrection::Evaluate(FTG_EvaluationContext* InContext)
+{
+	Super::Evaluate(InContext);
+
+	const RenderMaterial_FXPtr RenderMaterial = TextureGraphEngine::GetMaterialManager()->CreateMaterial_FX<VSH_Simple, FSH_ColorCorrection>(TEXT("T_ColorCorrection"));
+
+	check(RenderMaterial);
+
+	if (!Input)
+	{
+		Output = FTG_Texture::GetBlack();
+		return;
+	}
+
+	FLinearColor TemperatureRGB = ColorTemperatureRGB(FMath::Clamp(Temperature, 1000.0f, 20000.0f));
+
+	JobUPtr RenderJob = std::make_unique<Job>(InContext->Cycle->GetMix(), InContext->TargetId, std::static_pointer_cast<BlobTransform>(RenderMaterial));
+
+	RenderJob
+		->AddArg(ARG_BLOB(Input, "Input"))
+		->AddArg(ARG_FLOAT(Brightness, "Brightness"))
+		->AddArg(ARG_FLOAT(Contrast, "Contrast"))
+		->AddArg(ARG_FLOAT(Gamma, "Gamma"))
+		->AddArg(ARG_FLOAT(Saturation, "Saturation"))
+		->AddArg(ARG_LINEAR_COLOR(TemperatureRGB, "TemperatureRGB"))
+		->AddArg(ARG_FLOAT(TemperatureStrength, "TemperatureStrength"))
+		->AddArg(ARG_FLOAT(TemperatureBrightnessNormalization, "TemperatureBrightnessNormalization"))
+		;
+
+	const FString Name = TEXT("ColorCorrection");
+	BufferDescriptor Desc = Output.GetBufferDescriptor();
+	Output = RenderJob->InitResult(Name, &Desc);
+	InContext->Cycle->AddJob(InContext->TargetId, std::move(RenderJob));
+}
diff --git a/Engine/Plugins/TextureGraph/Source/TextureGraph/Public/Expressions/Color/TG_Expression_ColorCorrection.h b/Engine/Plugins/TextureGraph/Source/TextureGraph/Public/Expressions/Color/TG_Expression_ColorCorrection.h
new file mode 100644
--- /dev/null
+++ b/Engine/Plugins/TextureGraph/Source/TextureGraph/Public/Expressions/Color/TG_Expression_ColorCorrection.h
@@ -0,0 +1,60 @@
+// Copyright Epic Games, Inc. All Rights Reserved.
+
+#pragma once
+
+#include "CoreMinimal.h"
+
+#include "Expressions/TG_Expression.h"
+#include "TG_Texture.h"
+#include "TG_Expression_ColorCorrection.generated.h"
+
+UCLASS()
+class TEXTUREGRAPH_API UTG_Expression_ColorCorrection : public UTG_Expression
+{
+	GENERATED_BODY()
+
+public:
+	TG_DECLARE_EXPRESSION(TG_Category::Adjustment);
+	virtual void						Evaluate(FTG_EvaluationContext* InContext) override;
+
+	// The input image to be converted to grayscale
+	UPROPERTY(meta = (TGType = "TG_Input", PinDisplayName = ""))
+	FTG_Texture							Input;
+
+	// Adjusts the overall lightness or darkness of an image. Increasing brightness makes it brighter, while decreasing brightness makes it darker.
+	UPROPERTY(EditAnywhere, Category = NoCategory, meta = (TGType = "TG_Setting", UIMin = "-1", ClampMin = "-1", UIMax = "1", ClampMax = "1"))
+	float								Brightness = 0;
+
+	// Modifies the difference between light and dark areas within an image. Increasing contrast makes light areas lighter and dark areas darker, while decreasing contrast reduces this difference, resulting in a more uniform appearance.
+	UPROPERTY(EditAnywhere, Category = NoCategory, meta = (TGType = "TG_Setting", UIMin = "0", ClampMin = "0", UIMax = "10", ClampMax = "10"))
+	float								Contrast = 1;
+
+	// Gamma can be used to manipulate the curve
+	UPROPERTY(EditAnywhere, Category = NoCategory, meta = (TGType = "TG_Setting", UIMin = "0", ClampMin = "0", UIMax = "10", ClampMax = "10"))
+	float								Gamma = 0.5f;
+
+	// Saturation is how intense and pure a color appears. 
+	UPROPERTY(EditAnywhere, Category = NoCategory, meta = (TGType = "TG_Setting", UIMin = "0", ClampMin = "0", UIMax = "10", ClampMax = "10"))
+	float								Saturation = 1;
+
+	// The temperature adjustment for the image
+	UPROPERTY(EditAnywhere, Category = NoCategory, meta = (TGType = "TG_Setting", UIMin = "1000", ClampMin = "1000", UIMax = "20000", ClampMax = "20000"))
+	float								Temperature = 6500.0;
+
+	// The strength of the temperature application. It just gives more control over mixing the original image with the temperature applied image
+	UPROPERTY(EditAnywhere, Category = NoCategory, meta = (TGType = "TG_Setting", UIMin = "0", ClampMin = "0", UIMax = "1", ClampMax = "1"))
+	float								TemperatureStrength = 0;
+
+	// How to normalize the original brightness of the image with the new brigthness (after temperature adjustment)
+	UPROPERTY(EditAnywhere, Category = NoCategory, meta = (TGType = "TG_Setting", UIMin = "0", ClampMin = "0", UIMax = "1", ClampMax = "1"))
+	float								TemperatureBrightnessNormalization = 0;
+
+	// The output image which is a single channel grayscale version of the input image 
+	UPROPERTY(EditAnywhere, Category = NoCategory, meta = (TGType = "TG_Output", PinDisplayName = ""))
+	FTG_Texture							Output;
+
+	virtual FTG_Name					GetDefaultName() const override { return TEXT("Color Correction");}
+	virtual FText						GetTooltipText() const override { return FText::FromString(TEXT("Provides color tweaking and correction options.")); } 
+
+};
+
diff --git a/Engine/Plugins/TextureGraph/Source/TextureGraphEngine/Transform/Expressions/T_Color.cpp b/Engine/Plugins/TextureGraph/Source/TextureGraphEngine/Transform/Expressions/T_Color.cpp
--- a/Engine/Plugins/TextureGraph/Source/TextureGraphEngine/Transform/Expressions/T_Color.cpp
+++ b/Engine/Plugins/TextureGraph/Source/TextureGraphEngine/Transform/Expressions/T_Color.cpp
@@ -13,3 +13,4 @@ IMPLEMENT_GLOBAL_SHADER(FSH_Threshold, "/Plugin/TextureGraph/Expressions/Express
 IMPLEMENT_GLOBAL_SHADER(FSH_HSV, "/Plugin/TextureGraph/Expressions/Expression_HSV.usf", "FSH_HSV", SF_Pixel);
 IMPLEMENT_GLOBAL_SHADER(FSH_RGB2HSV, "/Plugin/TextureGraph/Expressions/Expression_HSV.usf", "FSH_RGB2HSV", SF_Pixel);
 IMPLEMENT_GLOBAL_SHADER(FSH_HSV2RGB, "/Plugin/TextureGraph/Expressions/Expression_HSV.usf", "FSH_HSV2RGB", SF_Pixel);
+IMPLEMENT_GLOBAL_SHADER(FSH_ColorCorrection, "/Plugin/TextureGraph/Expressions/Expression_ColorCorrection.usf", "FSH_ColorCorrection", SF_Pixel);
diff --git a/Engine/Plugins/TextureGraph/Source/TextureGraphEngine/Transform/Expressions/T_Color.h b/Engine/Plugins/TextureGraph/Source/TextureGraphEngine/Transform/Expressions/T_Color.h
--- a/Engine/Plugins/TextureGraph/Source/TextureGraphEngine/Transform/Expressions/T_Color.h
+++ b/Engine/Plugins/TextureGraph/Source/TextureGraphEngine/Transform/Expressions/T_Color.h
@@ -21,11 +21,8 @@ public:
 		SHADER_PARAMETER_TEXTURE(Texture2D, SourceTexture)
 	END_SHADER_PARAMETER_STRUCT()
 
-	static bool						ShouldCompilePermutation(const FGlobalShaderPermutationParameters& Parameters) 
-	{
-		return IsPCPlatform(Parameters.Platform) && EnumHasAllFlags(Parameters.Flags, EShaderPermutationFlags::HasEditorOnlyData);
-	}
-	static void						ModifyCompilationEnvironment(const FGlobalShaderPermutationParameters& params, FShaderCompilerEnvironment& env) {}
+	TEXTURE_ENGINE_DEFAULT_PERMUTATION;
+	TEXTUREGRAPH_ENGINE_DEFAULT_COMPILATION_ENV;
 };
 
 //////////////////////////////////////////////////////////////////////////
@@ -41,13 +38,10 @@ public:
 		SHADER_PARAMETER_STRUCT(FStandardSamplerStates, SamplerStates)
 		SHADER_PARAMETER_TEXTURE(Texture2D, Input)
 		SHADER_PARAMETER(float, Threshold)
-		END_SHADER_PARAMETER_STRUCT()
+	END_SHADER_PARAMETER_STRUCT()
 
-		static bool						ShouldCompilePermutation(const FGlobalShaderPermutationParameters& Parameters)
-	{
-		return IsPCPlatform(Parameters.Platform) && EnumHasAllFlags(Parameters.Flags, EShaderPermutationFlags::HasEditorOnlyData);
-	}
-	static void						ModifyCompilationEnvironment(const FGlobalShaderPermutationParameters& params, FShaderCompilerEnvironment& env) {}
+	TEXTURE_ENGINE_DEFAULT_PERMUTATION;
+	TEXTUREGRAPH_ENGINE_DEFAULT_COMPILATION_ENV;
 };
 
 //////////////////////////////////////////////////////////////////////////
@@ -67,11 +61,8 @@ public:
 		SHADER_PARAMETER(float, Value)
 	END_SHADER_PARAMETER_STRUCT()
 
-	static bool						ShouldCompilePermutation(const FGlobalShaderPermutationParameters& Parameters) 
-	{
-		return IsPCPlatform(Parameters.Platform) && EnumHasAllFlags(Parameters.Flags, EShaderPermutationFlags::HasEditorOnlyData);
-	}
-	static void						ModifyCompilationEnvironment(const FGlobalShaderPermutationParameters& params, FShaderCompilerEnvironment& env) {}
+	TEXTURE_ENGINE_DEFAULT_PERMUTATION;
+	TEXTUREGRAPH_ENGINE_DEFAULT_COMPILATION_ENV;
 };
 
 //////////////////////////////////////////////////////////////////////////
@@ -88,11 +79,8 @@ public:
 		SHADER_PARAMETER_TEXTURE(Texture2D, Input)
 	END_SHADER_PARAMETER_STRUCT()
 
-	static bool						ShouldCompilePermutation(const FGlobalShaderPermutationParameters& Parameters) 
-	{
-		return IsPCPlatform(Parameters.Platform) && EnumHasAllFlags(Parameters.Flags, EShaderPermutationFlags::HasEditorOnlyData);
-	}
-	static void						ModifyCompilationEnvironment(const FGlobalShaderPermutationParameters& params, FShaderCompilerEnvironment& env) {}
+	TEXTURE_ENGINE_DEFAULT_PERMUTATION;
+	TEXTUREGRAPH_ENGINE_DEFAULT_COMPILATION_ENV;
 };
 
 //////////////////////////////////////////////////////////////////////////
@@ -109,11 +97,33 @@ public:
 		SHADER_PARAMETER_TEXTURE(Texture2D, Input)
 	END_SHADER_PARAMETER_STRUCT()
 
-	static bool						ShouldCompilePermutation(const FGlobalShaderPermutationParameters& Parameters) 
-	{
-		return IsPCPlatform(Parameters.Platform) && EnumHasAllFlags(Parameters.Flags, EShaderPermutationFlags::HasEditorOnlyData);
-	}
-	static void						ModifyCompilationEnvironment(const FGlobalShaderPermutationParameters& params, FShaderCompilerEnvironment& env) {}
+	TEXTURE_ENGINE_DEFAULT_PERMUTATION;
+	TEXTUREGRAPH_ENGINE_DEFAULT_COMPILATION_ENV;
+};
+
+//////////////////////////////////////////////////////////////////////////
+/// Color correction shader
+//////////////////////////////////////////////////////////////////////////
+class TEXTUREGRAPHENGINE_API FSH_ColorCorrection : public FSH_Base
+{
+public:
+	SHADER_USE_PARAMETER_STRUCT(FSH_ColorCorrection, FSH_Base);
+	DECLARE_GLOBAL_SHADER(FSH_ColorCorrection);
+
+	BEGIN_SHADER_PARAMETER_STRUCT(FParameters, )
+		SHADER_PARAMETER_STRUCT(FStandardSamplerStates, SamplerStates)
+		SHADER_PARAMETER_TEXTURE(Texture2D, Input)
+		SHADER_PARAMETER(float, Brightness)
+		SHADER_PARAMETER(float, Contrast)
+		SHADER_PARAMETER(float, Gamma)
+		SHADER_PARAMETER(float, Saturation)
+		SHADER_PARAMETER(FLinearColor, TemperatureRGB)
+		SHADER_PARAMETER(float, TemperatureStrength)
+		SHADER_PARAMETER(float, TemperatureBrightnessNormalization)
+	END_SHADER_PARAMETER_STRUCT()
+
+	TEXTURE_ENGINE_DEFAULT_PERMUTATION;
+	TEXTUREGRAPH_ENGINE_DEFAULT_COMPILATION_ENV;
 };
 
 //////////////////////////////////////////////////////////////////////////
