Added support for connecting scalar node to unsigned integer node

diff --git a/Engine/Plugins/TextureGraph/Source/TextureGraph/Private/TG_GraphEvaluation.cpp b/Engine/Plugins/TextureGraph/Source/TextureGraph/Private/TG_GraphEvaluation.cpp
--- a/Engine/Plugins/TextureGraph/Source/TextureGraph/Private/TG_GraphEvaluation.cpp
+++ b/Engine/Plugins/TextureGraph/Source/TextureGraph/Private/TG_GraphEvaluation.cpp
@@ -121,6 +121,15 @@ void FloatToInt_Converter(FTG_Evaluation::VarConverterInfo& Info)
 	Info.OutVar->EditAs<int>() = Input;
 }
 
+void FloatToUInt_Converter(FTG_Evaluation::VarConverterInfo& Info)
+{
+	float Input = Info.InVar->GetAs<float>();
+	uint32 Output = 0;
+	if (Input > 0)
+		Output = (uint32)Input;
+	Info.OutVar->EditAs<uint32>() = Output;
+}
+
 void FloatToFLinearColor_Converter(FTG_Evaluation::VarConverterInfo& Info)
 {
 	auto Input = Info.InVar->GetAs<float>();
@@ -536,6 +545,7 @@ FTG_Evaluation::ConverterMap FTG_Evaluation::DefaultConverters
 	{
 		VAR_CONVERTER(float, int, FloatToInt_Converter),
 		VAR_CONVERTER(float, int32, FloatToInt_Converter),
+		VAR_CONVERTER(float, uint32, FloatToUInt_Converter),
 		VAR_CONVERTER(float, FLinearColor, FloatToFLinearColor_Converter),
 		VAR_CONVERTER(float, FVector4f, FloatToFVector4f_Converter),
 		VAR_CONVERTER(float, FTG_Texture, FloatToFTG_Texture_Converter),
