Fixed crash with array concatenate node

diff --git a/Engine/Plugins/TextureGraph/Source/TextureGraph/Private/Expressions/Arrays/TG_Expression_ArrayOps.cpp b/Engine/Plugins/TextureGraph/Source/TextureGraph/Private/Expressions/Arrays/TG_Expression_ArrayOps.cpp
--- a/Engine/Plugins/TextureGraph/Source/TextureGraph/Private/Expressions/Arrays/TG_Expression_ArrayOps.cpp
+++ b/Engine/Plugins/TextureGraph/Source/TextureGraph/Private/Expressions/Arrays/TG_Expression_ArrayOps.cpp
@@ -18,9 +18,9 @@ void UTG_Expression_ArrayConcat::Evaluate(FTG_EvaluationContext* InContext)
 	uint32 MaxCount = Input1Count + Input2Count;
 	Output.SetNum(MaxCount);
 
-	uint32 ActualStartIndex = StartIndex;
+	uint32 ActualStartIndex = StartIndex >= 0 ? StartIndex : 0;
 
-	if (StartIndex >= MaxCount)
+	if (ActualStartIndex >= MaxCount)
 	{
 		FTextureGraphErrorReporter* ErrorReporter = TextureGraphEngine::GetErrorReporter(InContext->Cycle->GetMix());
 
diff --git a/Engine/Plugins/TextureGraph/Source/TextureGraph/Public/Expressions/Arrays/TG_Expression_ArrayOps.h b/Engine/Plugins/TextureGraph/Source/TextureGraph/Public/Expressions/Arrays/TG_Expression_ArrayOps.h
--- a/Engine/Plugins/TextureGraph/Source/TextureGraph/Public/Expressions/Arrays/TG_Expression_ArrayOps.h
+++ b/Engine/Plugins/TextureGraph/Source/TextureGraph/Public/Expressions/Arrays/TG_Expression_ArrayOps.h
@@ -25,7 +25,7 @@ public:
 
 	/// The index at which Input1 array is going to start. Input 2 is going occupy the rest of the places within the array
 	UPROPERTY(meta = (TGType = "TG_Input", PinDisplayName = "Input-1 Out Index"))
-	uint32							StartIndex = 0;
+	int32							StartIndex = 0;
 
 	/// The second array as part of the concatenation process
 	UPROPERTY(meta = (TGType = "TG_Input", PinDisplayName = "Input-2"))
