Premultiply alpha node

diff --git a/Engine/Plugins/TextureGraph/Shaders/Expressions/Expression_Premult.usf b/Engine/Plugins/TextureGraph/Shaders/Expressions/Expression_Premult.usf
new file mode 100644
--- /dev/null
+++ b/Engine/Plugins/TextureGraph/Shaders/Expressions/Expression_Premult.usf
@@ -0,0 +1,12 @@
+// Copyright Epic Games, Inc. All Rights Reserved.
+#include "/Engine/Public/Platform.ush"
+#include "/Plugin/TextureGraph/SamplerStates.ush"
+#include "/Plugin/TextureGraph/ShaderUtil.ush"
+
+Texture2D Input;
+
+float4 FSH_Premult(float2 uv : TEXCOORD0) : SV_Target0
+{
+	float4 Clr = Input.Sample(SamplerStates_NoBorder, uv);
+	return float4(Clr.rgb * float3(Clr.a, Clr.a, Clr.a), Clr.a);
+}
diff --git a/Engine/Plugins/TextureGraph/Source/TextureGraph/Private/Expressions/Color/TG_Expression_Premult.cpp b/Engine/Plugins/TextureGraph/Source/TextureGraph/Private/Expressions/Color/TG_Expression_Premult.cpp
new file mode 100644
--- /dev/null
+++ b/Engine/Plugins/TextureGraph/Source/TextureGraph/Private/Expressions/Color/TG_Expression_Premult.cpp
@@ -0,0 +1,32 @@
+// Copyright Epic Games, Inc. All Rights Reserved.
+
+#include "Expressions/Color/TG_Expression_Premult.h"
+#include "FxMat/MaterialManager.h"
+#include "Job/JobBatch.h"
+#include "Transform/Expressions/T_Color.h"
+#include "TextureGraphEngine/Helper/ColorUtil.h"
+
+void UTG_Expression_Premult::Evaluate(FTG_EvaluationContext* InContext)
+{
+	Super::Evaluate(InContext);
+
+	const RenderMaterial_FXPtr RenderMaterial = TextureGraphEngine::GetMaterialManager()->CreateMaterial_FX<VSH_Simple, FSH_Premult>(TEXT("T_Premult"));
+
+	check(RenderMaterial);
+
+	if (!Input)
+	{
+		Output = FTG_Texture::GetBlack();
+		return;
+	}
+
+	JobUPtr RenderJob = std::make_unique<Job>(InContext->Cycle->GetMix(), InContext->TargetId, std::static_pointer_cast<BlobTransform>(RenderMaterial));
+	RenderJob
+		->AddArg(ARG_BLOB(Input, "Input"))
+		;
+
+	const FString Name = TEXT("Premult");
+	BufferDescriptor Desc = Output.GetBufferDescriptor();
+	Output = RenderJob->InitResult(Name, &Desc);
+	InContext->Cycle->AddJob(InContext->TargetId, std::move(RenderJob));
+}
diff --git a/Engine/Plugins/TextureGraph/Source/TextureGraph/Public/Expressions/Color/TG_Expression_ColorCorrection.h b/Engine/Plugins/TextureGraph/Source/TextureGraph/Public/Expressions/Color/TG_Expression_ColorCorrection.h
--- a/Engine/Plugins/TextureGraph/Source/TextureGraph/Public/Expressions/Color/TG_Expression_ColorCorrection.h
+++ b/Engine/Plugins/TextureGraph/Source/TextureGraph/Public/Expressions/Color/TG_Expression_ColorCorrection.h
@@ -17,7 +17,7 @@ public:
 	TG_DECLARE_EXPRESSION(TG_Category::Adjustment);
 	virtual void						Evaluate(FTG_EvaluationContext* InContext) override;
 
-	// The input image to be converted to grayscale
+	// The input image to be color corrected
 	UPROPERTY(meta = (TGType = "TG_Input", PinDisplayName = ""))
 	FTG_Texture							Input;
 
@@ -49,7 +49,7 @@ public:
 	UPROPERTY(EditAnywhere, Category = NoCategory, meta = (TGType = "TG_Setting", UIMin = "0", ClampMin = "0", UIMax = "1", ClampMax = "1"))
 	float								TemperatureBrightnessNormalization = 0;
 
-	// The output image which is a single channel grayscale version of the input image 
+	// The output image 
 	UPROPERTY(EditAnywhere, Category = NoCategory, meta = (TGType = "TG_Output", PinDisplayName = ""))
 	FTG_Texture							Output;
 
diff --git a/Engine/Plugins/TextureGraph/Source/TextureGraph/Public/Expressions/Color/TG_Expression_Premult.h b/Engine/Plugins/TextureGraph/Source/TextureGraph/Public/Expressions/Color/TG_Expression_Premult.h
new file mode 100644
--- /dev/null
+++ b/Engine/Plugins/TextureGraph/Source/TextureGraph/Public/Expressions/Color/TG_Expression_Premult.h
@@ -0,0 +1,33 @@
+// Copyright Epic Games, Inc. All Rights Reserved.
+
+#pragma once
+
+#include "CoreMinimal.h"
+
+#include "Expressions/TG_Expression.h"
+#include "TG_Texture.h"
+#include "Expressions/Maths/TG_Expression_Maths_OneInput.h"
+
+#include "TG_Expression_Premult.generated.h"
+
+UCLASS()
+class TEXTUREGRAPH_API UTG_Expression_Premult : public UTG_Expression
+{
+	GENERATED_BODY()
+
+public:
+	TG_DECLARE_EXPRESSION(TG_Category::Adjustment);
+	virtual void						Evaluate(FTG_EvaluationContext* InContext) override;
+
+	// The input image to be used
+	UPROPERTY(meta = (TGType = "TG_Input", PinDisplayName = ""))
+	FTG_Texture							Input;
+
+	// The output of the image with the premultiplied alpha
+	UPROPERTY(EditAnywhere, Category = NoCategory, meta = (TGType = "TG_Output", PinDisplayName = ""))
+	FTG_Texture							Output;
+
+	virtual FTG_Name					GetDefaultName() const override { return TEXT("Premultiply Alpha");}
+	virtual FText						GetTooltipText() const override { return FText::FromString(TEXT("Premultiply alpha with the RGB color values at each pixel. The alpha remains unchanged.")); } 
+};
+
diff --git a/Engine/Plugins/TextureGraph/Source/TextureGraphEngine/Transform/Expressions/T_Color.cpp b/Engine/Plugins/TextureGraph/Source/TextureGraphEngine/Transform/Expressions/T_Color.cpp
--- a/Engine/Plugins/TextureGraph/Source/TextureGraphEngine/Transform/Expressions/T_Color.cpp
+++ b/Engine/Plugins/TextureGraph/Source/TextureGraphEngine/Transform/Expressions/T_Color.cpp
@@ -14,3 +14,4 @@ IMPLEMENT_GLOBAL_SHADER(FSH_HSV, "/Plugin/TextureGraph/Expressions/Expression_HS
 IMPLEMENT_GLOBAL_SHADER(FSH_RGB2HSV, "/Plugin/TextureGraph/Expressions/Expression_HSV.usf", "FSH_RGB2HSV", SF_Pixel);
 IMPLEMENT_GLOBAL_SHADER(FSH_HSV2RGB, "/Plugin/TextureGraph/Expressions/Expression_HSV.usf", "FSH_HSV2RGB", SF_Pixel);
 IMPLEMENT_GLOBAL_SHADER(FSH_ColorCorrection, "/Plugin/TextureGraph/Expressions/Expression_ColorCorrection.usf", "FSH_ColorCorrection", SF_Pixel);
+IMPLEMENT_GLOBAL_SHADER(FSH_Premult, "/Plugin/TextureGraph/Expressions/Expression_Premult.usf", "FSH_Premult", SF_Pixel);
diff --git a/Engine/Plugins/TextureGraph/Source/TextureGraphEngine/Transform/Expressions/T_Color.h b/Engine/Plugins/TextureGraph/Source/TextureGraphEngine/Transform/Expressions/T_Color.h
--- a/Engine/Plugins/TextureGraph/Source/TextureGraphEngine/Transform/Expressions/T_Color.h
+++ b/Engine/Plugins/TextureGraph/Source/TextureGraphEngine/Transform/Expressions/T_Color.h
@@ -101,6 +101,24 @@ public:
 	TEXTUREGRAPH_ENGINE_DEFAULT_COMPILATION_ENV;
 };
 
+//////////////////////////////////////////////////////////////////////////
+/// Color correction shader
+//////////////////////////////////////////////////////////////////////////
+class TEXTUREGRAPHENGINE_API FSH_Premult : public FSH_Base
+{
+public:
+	SHADER_USE_PARAMETER_STRUCT(FSH_Premult, FSH_Base);
+	DECLARE_GLOBAL_SHADER(FSH_Premult);
+
+	BEGIN_SHADER_PARAMETER_STRUCT(FParameters, )
+		SHADER_PARAMETER_STRUCT(FStandardSamplerStates, SamplerStates)
+		SHADER_PARAMETER_TEXTURE(Texture2D, Input)
+	END_SHADER_PARAMETER_STRUCT()
+
+	TEXTURE_ENGINE_DEFAULT_PERMUTATION;
+	TEXTUREGRAPH_ENGINE_DEFAULT_COMPILATION_ENV;
+};
+
 //////////////////////////////////////////////////////////////////////////
 /// Color correction shader
 //////////////////////////////////////////////////////////////////////////
