Erode/Dilate support

diff --git a/Engine/Plugins/TextureGraph/Shaders/Expressions/Expression_ErodeDilate.usf b/Engine/Plugins/TextureGraph/Shaders/Expressions/Expression_ErodeDilate.usf
new file mode 100644
--- /dev/null
+++ b/Engine/Plugins/TextureGraph/Shaders/Expressions/Expression_ErodeDilate.usf
@@ -0,0 +1,82 @@
+// Copyright Epic Games, Inc. All Rights Reserved.
+#include "/Engine/Public/Platform.ush"
+#include "/Plugin/TextureGraph/SamplerStates.ush"
+#include "/Plugin/TextureGraph/TileInfo.ush"
+#include "/Plugin/TextureGraph/ShaderUtil.ush"
+
+#ifndef ED_KERNEL_BOX
+#define ED_KERNEL_BOX 0
+#endif
+
+#ifndef ED_KERNEL_CIRCULAR
+#define ED_KERNEL_CIRCULAR 1
+#endif
+
+#ifndef ED_KERNEL_DIAMOND
+#define ED_KERNEL_DIAMOND 2
+#endif
+
+#ifndef ED_ERODE
+#define ED_ERODE 0
+#endif 
+
+#ifndef ED_DILATE
+#define ED_DILATE 1
+#endif 
+
+Texture2D Input;
+int Size;
+
+float4 FSH_ErodeDilate(float2 InUV : TEXCOORD0) : SV_Target0
+{
+	int Width = TileInfo_TileCountX * TileInfo_TileWidth;
+	int Height = TileInfo_TileCountY * TileInfo_TileHeight;
+	float WidthInv = rcp(float(Width));
+	float HeightInv = rcp(float(Height));
+	float2 UV = TileInfo_fromCurrentTileToLayer(InUV);
+	
+	float4 OrgClr = Input.Sample(SamplerStates_Linear_Clamp, UV);
+	float3 Clr = OrgClr.rgb;
+	float3 TmpClr = Clr;
+	float Brightness = LumaFromRGB(Clr);
+
+	for (int X = -Size; X <= Size; ++X)
+	{
+		float Xf = float(X) * WidthInv;
+		for (int Y = -Size; Y <= Size; ++Y)
+		{
+			float Yf = float(Y) * HeightInv;
+        // default kernel is box shaped anyway, so we don't need to do anything
+        
+        // diamond shaped kernel (45° rotated square)
+#if ED_KERNEL == ED_KERNEL_DIAMOND
+        if (abs(X) > Size - abs(Y))
+            continue;
+#elif ED_KERNEL == ED_KERNEL_CIRCULAR
+        if (distance(float2(X, Y), float2(0, 0)) > float(Size))
+            continue;
+#endif
+
+			float2 OffsetUV = UV + float2(Xf, Yf);
+			float3 NClr = Input.Sample(SamplerStates_Linear_Clamp, OffsetUV).rgb;
+			float NBrightness = LumaFromRGB(NClr);
+        
+#if ED_TYPE == ED_DILATE
+			if (Brightness < NBrightness)
+#else
+			if (Brightness > NBrightness)
+#endif 
+			{
+				Brightness = NBrightness;
+				TmpClr = NClr;
+			}
+		}
+	}
+	
+#if ED_TYPE == ED_DILATE
+	float3 FinalClr = lerp(Clr, TmpClr, smoothstep(0, 1, Brightness));
+#else
+	float3 FinalClr = lerp(TmpClr, Clr, smoothstep(0, 1, Brightness));
+#endif 
+	return float4(FinalClr.rgb, OrgClr.a);
+}
diff --git a/Engine/Plugins/TextureGraph/Source/TextureGraph/Private/Expressions/FIlter/TG_Expression_ErodeDilate.cpp b/Engine/Plugins/TextureGraph/Source/TextureGraph/Private/Expressions/FIlter/TG_Expression_ErodeDilate.cpp
new file mode 100644
--- /dev/null
+++ b/Engine/Plugins/TextureGraph/Source/TextureGraph/Private/Expressions/FIlter/TG_Expression_ErodeDilate.cpp
@@ -0,0 +1,53 @@
+// Copyright Epic Games, Inc. All Rights Reserved.
+
+#include "Expressions/Filter/TG_Expression_ErodeDilate.h"
+#include "FxMat/MaterialManager.h"
+#include "Job/JobBatch.h"
+#include "Transform/Expressions/T_Filter.h"
+#include "Transform/Utility/T_CombineTiledBlob.h"
+
+//////////////////////////////////////////////////////////////////////////
+void UTG_Expression_ErodeDilate::Evaluate(FTG_EvaluationContext* InContext)
+{
+	Super::Evaluate(InContext);
+
+	if (!Input)
+	{
+		Output = FTG_Texture::GetBlack();
+		return;
+	}
+
+	if (Size == 0)
+	{
+		Output = Input;
+		return;
+	}
+
+	FSH_ErodeDilate::FPermutationDomain PermutationVector;
+	PermutationVector.Set<FVar_ErodeDilate_Type>((int32)Type);
+	PermutationVector.Set<FVar_ErodeDilate_Kernel>((int32)Kernel);
+
+	RenderMaterial_FXPtr RenderMaterial = TextureGraphEngine::GetMaterialManager()->CreateMaterial_FX<VSH_Simple, FSH_ErodeDilate>(TEXT("T_ErodeDilate"), PermutationVector);
+	check(RenderMaterial);
+
+	// Combine the tiled texture
+	TiledBlobPtr CombinedBlob = T_CombineTiledBlob::Create(InContext->Cycle, Input->GetDescriptor(), InContext->TargetId, Input);
+	BufferDescriptor Desc = BufferDescriptor::Combine(Output.Descriptor, Input->GetDescriptor());
+
+	FTileInfo TileInfo;
+	JobUPtr RenderJob = std::make_unique<Job>(InContext->Cycle->GetMix(), InContext->TargetId, std::static_pointer_cast<BlobTransform>(RenderMaterial));
+	RenderJob
+		->AddArg(ARG_TILEINFO(TileInfo, "TileInfo"))
+		->AddArg(ARG_BLOB(CombinedBlob, "Input"))
+		->AddArg(ARG_INT(Size, "Size"))
+		->AddArg(WithUnbounded(ARG_INT((int)Type, "Type")))
+		->AddArg(WithUnbounded(ARG_INT((int)Kernel, "Kernel")))
+		;
+
+	const FString Name = FString::Printf(TEXT("T_ErodeDilate"));
+
+	TiledBlobPtr Result = RenderJob->InitResult(Name, &Desc);
+	InContext->Cycle->AddJob(InContext->TargetId, std::move(RenderJob));
+
+	Output = Result;
+}
diff --git a/Engine/Plugins/TextureGraph/Source/TextureGraph/Public/Expressions/Filter/TG_Expression_ErodeDilate.h b/Engine/Plugins/TextureGraph/Source/TextureGraph/Public/Expressions/Filter/TG_Expression_ErodeDilate.h
new file mode 100644
--- /dev/null
+++ b/Engine/Plugins/TextureGraph/Source/TextureGraph/Public/Expressions/Filter/TG_Expression_ErodeDilate.h
@@ -0,0 +1,77 @@
+// Copyright Epic Games, Inc. All Rights Reserved.
+
+#pragma once
+
+#include "CoreMinimal.h"
+
+#include "Expressions/TG_Expression.h"
+#include "TG_Texture.h"
+#include "Expressions/Maths/TG_Expression_Maths_OneInput.h"
+
+#include "TG_Expression_ErodeDilate.generated.h"
+
+//////////////////////////////////////////////////////////////////////////
+/// Type
+//////////////////////////////////////////////////////////////////////////
+UENUM(BlueprintType)
+enum class EErodeOrDilate : uint8
+{
+	Erode = 0		UMETA(DisplayName = "Erode"),
+	Dilate = 1		UMETA(DisplayName = "Dilate"),
+};
+
+//////////////////////////////////////////////////////////////////////////
+/// Type
+//////////////////////////////////////////////////////////////////////////
+UENUM(BlueprintType)
+enum class EErodeDilateKernelType : uint8
+{
+	Box = 0			UMETA(DisplayName = "Box"),
+	Circular = 1 	UMETA(DisplayName = "Circular"),
+	Diamond = 2		UMETA(DisplayName = "Diamond"),
+};
+
+//////////////////////////////////////////////////////////////////////////
+/// Type
+//////////////////////////////////////////////////////////////////////////
+UCLASS()
+class TEXTUREGRAPH_API UTG_Expression_ErodeDilate : public UTG_Expression
+{
+	GENERATED_BODY()
+
+public:
+	TG_DECLARE_EXPRESSION(TG_Category::Filter);
+	virtual void						Evaluate(FTG_EvaluationContext* InContext) override;
+
+	/// The input image to 
+	UPROPERTY(meta = (TGType = "TG_Input", PinDisplayName = ""))
+	FTG_Texture							Input;
+
+	///// The minimum brightness level to use. Stick with 0 for non-hdr images
+	//UPROPERTY(meta = (TGType = "TG_Input"))
+	//float								MinBrightness = 0.0f;
+
+	///// The maximum brightness level to use. Stick with 1 for non-hdr images
+	//UPROPERTY(meta = (TGType = "TG_Input"))
+	//float								MinBrightness = 0.0f;
+
+	/// The size of the kernel to use
+	UPROPERTY(meta = (TGType = "TG_Input", ClampMin = 0, ClampMax = 64, UIMin = 0, UIMax = 64))
+	int32								Size = 2;
+
+	/// The type of the kernel to use
+	UPROPERTY(EditAnywhere, BlueprintReadWrite, Category = NoCategory, meta = (TGType = "TG_Setting"))
+	EErodeDilateKernelType				Kernel = EErodeDilateKernelType::Box;
+
+	/// Whether to erode or dilate
+	UPROPERTY(EditAnywhere, BlueprintReadWrite, Category = NoCategory, meta = (TGType = "TG_Setting"))
+	EErodeOrDilate						Type = EErodeOrDilate::Erode;
+
+	/// The output image which is a single channel grayscale version of the input image 
+	UPROPERTY(EditAnywhere, Category = NoCategory, meta = (TGType = "TG_Output", PinDisplayName = ""))
+	FTG_Texture							Output;
+
+	virtual FTG_Name					GetDefaultName() const override { return TEXT("Erode/Dilate");}
+	virtual FText						GetTooltipText() const override { return FText::FromString(TEXT("Erode or dilate a particular image.")); } 
+};
+
diff --git a/Engine/Plugins/TextureGraph/Source/TextureGraphEngine/Transform/Expressions/T_Filter.cpp b/Engine/Plugins/TextureGraph/Source/TextureGraphEngine/Transform/Expressions/T_Filter.cpp
--- a/Engine/Plugins/TextureGraph/Source/TextureGraphEngine/Transform/Expressions/T_Filter.cpp
+++ b/Engine/Plugins/TextureGraph/Source/TextureGraphEngine/Transform/Expressions/T_Filter.cpp
@@ -9,6 +9,7 @@ IMPLEMENT_GLOBAL_SHADER(FSH_EdgeDetect, "/Plugin/TextureGraph/Expressions/Expres
 IMPLEMENT_GLOBAL_SHADER(FSH_DirectionalWarp, "/Plugin/TextureGraph/Expressions/Expression_Warp.usf", "FSH_DirectionalWarp", SF_Pixel);
 IMPLEMENT_GLOBAL_SHADER(FSH_NormalWarp, "/Plugin/TextureGraph/Expressions/Expression_Warp.usf", "FSH_NormalWarp", SF_Pixel);
 IMPLEMENT_GLOBAL_SHADER(FSH_SineWarp, "/Plugin/TextureGraph/Expressions/Expression_Warp.usf", "FSH_SineWarp", SF_Pixel);
+IMPLEMENT_GLOBAL_SHADER(FSH_ErodeDilate, "/Plugin/TextureGraph/Expressions/Expression_ErodeDilate.usf", "FSH_ErodeDilate", SF_Pixel);
 
 TiledBlobPtr T_Filter::CreateEdgeDetect(MixUpdateCyclePtr InCycle, BufferDescriptor DesiredDesc, TiledBlobPtr SourceTexture, float Thickness, int32 InTargetId)
 {
diff --git a/Engine/Plugins/TextureGraph/Source/TextureGraphEngine/Transform/Expressions/T_Filter.h b/Engine/Plugins/TextureGraph/Source/TextureGraphEngine/Transform/Expressions/T_Filter.h
--- a/Engine/Plugins/TextureGraph/Source/TextureGraphEngine/Transform/Expressions/T_Filter.h
+++ b/Engine/Plugins/TextureGraph/Source/TextureGraphEngine/Transform/Expressions/T_Filter.h
@@ -107,6 +107,30 @@ public:
 	TEXTUREGRAPH_ENGINE_DEFAULT_COMPILATION_ENV;
 };
 
+class FVar_ErodeDilate_Type : SHADER_PERMUTATION_INT("ED_TYPE", 2);
+class FVar_ErodeDilate_Kernel : SHADER_PERMUTATION_INT("ED_KERNEL", 3);
+//////////////////////////////////////////////////////////////////////////
+/// Erode/Dilate
+//////////////////////////////////////////////////////////////////////////
+class TEXTUREGRAPHENGINE_API FSH_ErodeDilate : public FSH_Base
+{
+public:
+	DECLARE_GLOBAL_SHADER(FSH_ErodeDilate);
+	SHADER_USE_PARAMETER_STRUCT(FSH_ErodeDilate, FSH_Base);
+	BEGIN_SHADER_PARAMETER_STRUCT(FParameters, )
+		SHADER_PARAMETER_STRUCT(FTileInfo, TileInfo)
+		SHADER_PARAMETER_STRUCT(FStandardSamplerStates, SamplerStates)
+		SHADER_PARAMETER_TEXTURE(Texture2D, Input)
+		SHADER_PARAMETER(int, Size)
+	END_SHADER_PARAMETER_STRUCT()
+
+	using FPermutationDomain = TShaderPermutationDomain<FVar_ErodeDilate_Type, FVar_ErodeDilate_Kernel>;
+public:
+	TEXTURE_ENGINE_DEFAULT_PERMUTATION;
+	TEXTUREGRAPH_ENGINE_DEFAULT_COMPILATION_ENV;
+};
+
+
 //////////////////////////////////////////////////////////////////////////
 class TEXTUREGRAPHENGINE_API T_Filter
 {
