Added support for dedicated texture settings, per node

diff --git a/Engine/Plugins/TextureGraph/Source/TextureGraph/Private/Expressions/TG_Expression.cpp b/Engine/Plugins/TextureGraph/Source/TextureGraph/Private/Expressions/TG_Expression.cpp
--- a/Engine/Plugins/TextureGraph/Source/TextureGraph/Private/Expressions/TG_Expression.cpp
+++ b/Engine/Plugins/TextureGraph/Source/TextureGraph/Private/Expressions/TG_Expression.cpp
@@ -152,7 +152,15 @@ bool UTG_Expression::CanEditChange(const FProperty* InProperty) const
 			// Make property not editable if pin is connected and is of type input or settings
 			bIsEditable = !(Pin->IsConnected() && (Pin->IsInput() || Pin->IsSetting()));	
 		}
+
+		const FName PropertyName = InProperty->GetFName();
+
+		if (PropertyName == GET_MEMBER_NAME_CHECKED(UTG_Expression, BaseOutputSettings))
+		{
+			bIsEditable = ShouldShowSettings();
+		}
 	}
+
 	return bIsEditable;
 }
 
@@ -279,6 +287,17 @@ void UTG_Expression::SetupAndEvaluate(FTG_EvaluationContext* InContext)
 	// Copy Output var values TO this expression's matching properties (Mainly to copy existing thumbs for Output FTG_Textures)
 	for (auto& OutputVar : InContext->Outputs.VarArguments)
 	{
+		if (OutputVar.Value.Argument.IsTexture())
+		{
+			FTG_Texture& Texture = OutputVar.Value.Var->GetAs<FTG_Texture>();
+			Texture.Descriptor = BaseOutputSettings;
+		}
+		else if (OutputVar.Value.Argument.IsVariant())
+		{
+			FTG_Variant& Variant = OutputVar.Value.Var->GetAs<FTG_Variant>();
+			if (Variant.IsTexture())
+				Variant.GetTexture().Descriptor = BaseOutputSettings;
+		}
 		OutputVar.Value.Var->CopyTo(this, OutputVar.Value.Argument);
 	}
 
diff --git a/Engine/Plugins/TextureGraph/Source/TextureGraph/Private/TG_GraphEvaluation.cpp b/Engine/Plugins/TextureGraph/Source/TextureGraph/Private/TG_GraphEvaluation.cpp
--- a/Engine/Plugins/TextureGraph/Source/TextureGraph/Private/TG_GraphEvaluation.cpp
+++ b/Engine/Plugins/TextureGraph/Source/TextureGraph/Private/TG_GraphEvaluation.cpp
@@ -710,28 +710,6 @@ void FTG_Evaluation::EvaluateNodeArray(UTG_Node* InNode, const TArray<TObjectPtr
 		FTG_Evaluation::TransferVarToPin(Pin, InContext, -1);
 	}
 
-	/// Transfer the descriptors of the outputs to the output vars ... one index at a time (only if they've been modified)
-	for (auto& OutputVar : InContext->Outputs.VarArguments)
-	{
-		const FTG_VariantArray* ArrayOutput = ArrayOutputs.Find(OutputVar.Key);
-		if (!ArrayOutput->GetArray().IsEmpty())
-		{
-			const FTG_Variant& ArrayValue = ArrayOutput->Get(0);
-
-			/// We only need to do this for textures
-			if (ArrayValue.IsTexture())
-			{
-				const FTG_Texture& ArrayTexture = ArrayValue.GetTexture();
-				FString ToType = OutputVar.Value.Argument.GetCPPTypeName().ToString();
-
-				if (ToType == TEXT("FTG_Variant"))
-					OutputVar.Value.Var->EditAs<FTG_Variant>().EditTexture().Descriptor = ArrayTexture.Descriptor;
-				else if (ToType == TEXT("FTG_Texture"))
-					OutputVar.Value.Var->EditAs<FTG_Texture>().Descriptor = ArrayTexture.Descriptor;
-			}
-		}
-	}
-
 	for (int VarIndex = 0; VarIndex < MaxCount; VarIndex++)
 	{
 		FTG_Variant::EType CommonVariantType = FTG_Variant::EType::Invalid;
diff --git a/Engine/Plugins/TextureGraph/Source/TextureGraph/Private/TG_Node.cpp b/Engine/Plugins/TextureGraph/Source/TextureGraph/Private/TG_Node.cpp
--- a/Engine/Plugins/TextureGraph/Source/TextureGraph/Private/TG_Node.cpp
+++ b/Engine/Plugins/TextureGraph/Source/TextureGraph/Private/TG_Node.cpp
@@ -246,6 +246,36 @@ void UTG_Node::Serialize(FArchive& Ar)
 		Expression = NewObject<UTG_Expression_Null>(this, UTG_Expression_Null::StaticClass(), NAME_None, RF_Transactional);
 	}
 
+	int32 Version = Ar.CustomVer(FTG_CustomVersion::GUID);
+
+	/// Backwards compatibility with Texture descriptor
+	if (Version < FTG_CustomVersion::TGExpressionAddedBaseOutputSettings)
+	{
+		TArray<FTG_Texture> Textures;
+		GetAllOutputValues(Textures, nullptr);
+
+		if (!Textures.IsEmpty())
+		{
+			Expression->BaseOutputSettings = Textures[0].Descriptor;
+		}
+		else /// Try to see if there are texture variants 
+		{
+			TArray<FTG_Variant> Variants;
+			GetAllOutputValues(Variants, nullptr);
+
+			/// Just copy the descriptor off the first texture that we encounter
+			for (int32 Index = 0; Index < Variants.Num(); Index++)
+			{
+				if (Variants[Index].IsTexture())
+				{
+					FTG_Texture Texture = Variants[Index].GetTexture();
+					Expression->BaseOutputSettings = Texture.Descriptor;
+					break;
+				}
+			}
+		}
+	}
+
 	UE_LOG(LogTextureGraph, VeryVerbose, TEXT("    %s Node: %s"),
 		(Ar.IsSaving() ? TEXT("Saved") : TEXT("Loaded")),
 		*GetId().ToString());
diff --git a/Engine/Plugins/TextureGraph/Source/TextureGraph/Private/TG_Var.cpp b/Engine/Plugins/TextureGraph/Source/TextureGraph/Private/TG_Var.cpp
--- a/Engine/Plugins/TextureGraph/Source/TextureGraph/Private/TG_Var.cpp
+++ b/Engine/Plugins/TextureGraph/Source/TextureGraph/Private/TG_Var.cpp
@@ -80,6 +80,12 @@ template <> FString TG_Var_LogValue(FTG_OutputSettings& Value)
 	return LogMessage;
 }
 
+template <> FString TG_Var_LogValue(FTG_TextureDescriptor& Value)
+{
+	FString LogMessage = Value.ToString();
+	return LogMessage;
+}
+
 template <> FString TG_Var_LogValue(FPatternMaskPlacement_TS& Value)
 {
 	FString LogMessage = Value.ToString();
@@ -167,22 +173,32 @@ template <> void TG_Var_SetValueFromString(FTG_OutputSettings& Value, const FStr
 {
 	Value.InitFromString(StrVal);
 }
+
+template <> void TG_Var_SetValueFromString(FTG_TextureDescriptor& Value, const FString& StrVal)
+{
+	Value.InitFromString(StrVal);
+}
+
 template <> void TG_Var_SetValueFromString(FPatternMaskPlacement_TS& Value, const FString& StrVal)
 {
 	Value.InitFromString(StrVal);
 }
+
 template <> void TG_Var_SetValueFromString(FPatternMaskJitter_TS& Value, const FString& StrVal)
 {
 	Value.InitFromString(StrVal);
 }
+
 template <> void TG_Var_SetValueFromString(FPatternMaskBevel_TS& Value, const FString& StrVal)
 {
 	Value.InitFromString(StrVal);
 }
+
 template <> void TG_Var_SetValueFromString(FPatternMaskCutout_TS& Value, const FString& StrVal)
 {
 	Value.InitFromString(StrVal);
 }
+
 template <> void TG_Var_SetValueFromString(FGradientDir_TS& Value, const FString& StrVal)
 {
 	Value.InitFromString(StrVal);
@@ -471,6 +487,11 @@ void VarPropertySerializer_FTG_OutputSettings(FTG_Var::VarPropertySerialInfo& In
 	Generic_Struct_Serializer<FTG_OutputSettings>(Info);
 }
 
+void VarPropertySerializer_FTG_TextureDescriptor(FTG_Var::VarPropertySerialInfo& Info)
+{
+	Generic_Struct_Serializer<FTG_TextureDescriptor>(Info);
+}
+
 void VarPropertySerializer_FGradientDir_TS(FTG_Var::VarPropertySerialInfo& Info)
 {
 	Generic_Struct_Serializer<FGradientDir_TS>(Info);
@@ -510,6 +531,7 @@ FTG_Var::VarPropertySerializerMap FTG_Var::DefaultPropertySerializers
 		VAR_PROPERTY_SERIALIZER_DEF(FVector2f),
 		VAR_PROPERTY_SERIALIZER_DEF(FLinearColor),
 		VAR_PROPERTY_SERIALIZER_DEF(FTG_OutputSettings),
+		VAR_PROPERTY_SERIALIZER_DEF(FTG_TextureDescriptor),
 		VAR_PROPERTY_SERIALIZER_DEF(FTG_Variant),
 		VAR_PROPERTY_SERIALIZER_DEF(FGradientDir_TS),
 		VAR_PROPERTY_SERIALIZER_DEF(FPatternMaskCutout_TS),
@@ -742,6 +764,11 @@ void VarArchiveSerializer_FTG_OutputSettings(FTG_Var::VarArchiveSerialInfo& Info
 {
 	Generic_Simple_ArSerializer<FTG_OutputSettings>(Info);
 }
+
+void VarArchiveSerializer_FTG_TextureDescriptor(FTG_Var::VarArchiveSerialInfo& Info)
+{
+	Generic_Simple_ArSerializer<FTG_TextureDescriptor>(Info);
+}
 //////////////////////////////////////////////////////////////////////////////////////////////////////////////////////////////////////
 FTG_Var::VarArchiveSerializerMap FTG_Var::DefaultArchiveSerializers
 (
@@ -756,7 +783,8 @@ FTG_Var::VarArchiveSerializerMap FTG_Var::DefaultArchiveSerializers
 		VAR_ARCHIVE_SERIALIZER_DEF(FVector2f),
 		VAR_ARCHIVE_SERIALIZER_DEF(FLinearColor),
 		VAR_ARCHIVE_SERIALIZER_DEF(FTG_VariantArray),
-		VAR_ARCHIVE_SERIALIZER_DEF(FTG_OutputSettings)
+		VAR_ARCHIVE_SERIALIZER_DEF(FTG_OutputSettings),
+		VAR_ARCHIVE_SERIALIZER_DEF(FTG_TextureDescriptor)
 	}
 ); 
 
diff --git a/Engine/Plugins/TextureGraph/Source/TextureGraph/Public/Expressions/Input/TG_Expression_InputParam.h b/Engine/Plugins/TextureGraph/Source/TextureGraph/Public/Expressions/Input/TG_Expression_InputParam.h
--- a/Engine/Plugins/TextureGraph/Source/TextureGraph/Public/Expressions/Input/TG_Expression_InputParam.h
+++ b/Engine/Plugins/TextureGraph/Source/TextureGraph/Public/Expressions/Input/TG_Expression_InputParam.h
@@ -17,6 +17,7 @@ protected:
 	// these 2 methods need to be implemented to provide concrete signatures for derived InputParam expression classes
 	virtual FTG_SignaturePtr BuildInputParameterSignature() const;
 	virtual FTG_SignaturePtr BuildInputConstantSignature() const;
+	virtual bool ShouldShowSettings() const override { return false; }
 public:
 
 #if WITH_EDITOR
diff --git a/Engine/Plugins/TextureGraph/Source/TextureGraph/Public/Expressions/Output/TG_Expression_Output.h b/Engine/Plugins/TextureGraph/Source/TextureGraph/Public/Expressions/Output/TG_Expression_Output.h
--- a/Engine/Plugins/TextureGraph/Source/TextureGraph/Public/Expressions/Output/TG_Expression_Output.h
+++ b/Engine/Plugins/TextureGraph/Source/TextureGraph/Public/Expressions/Output/TG_Expression_Output.h
@@ -40,6 +40,7 @@ public:
 	virtual void SetTitleName(FName NewName) override;
 	virtual FName GetTitleName() const override;
 	virtual FText GetTooltipText() const override { return FText::FromString(TEXT("Uses the node input as a graph output. It can be exported as a texture asset and is automatically exposed as a graph output parameter. ")); } 
+	virtual bool ShouldShowSettings() const { return false; }
 
 	void UpdateBufferDescriptorValues();
 	void InitializeOutputSettings();
diff --git a/Engine/Plugins/TextureGraph/Source/TextureGraph/Public/Expressions/TG_Expression.h b/Engine/Plugins/TextureGraph/Source/TextureGraph/Public/Expressions/TG_Expression.h
--- a/Engine/Plugins/TextureGraph/Source/TextureGraph/Public/Expressions/TG_Expression.h
+++ b/Engine/Plugins/TextureGraph/Source/TextureGraph/Public/Expressions/TG_Expression.h
@@ -44,8 +44,13 @@ protected:
 	// Method returning this Expression class version that can be overwritten in a sub class requiring to
 	// support version change
 	virtual int32		GetExpressionClassVersion() const { return 0; }
+	virtual bool		ShouldShowSettings() const { return true; }
 
 public:
+	//When we will work on Node UI for FOutputSettings we will set the category as TG_Setting
+	UPROPERTY(EditAnywhere, Category = NoCategory, meta = (TGType = "TG_Setting", CollapsableChildProperties, ShowOnlyInnerProperties, FullyExpand, NoResetToDefault, PinDisplayName = "Settings", EditConditionHides))
+	FTG_TextureDescriptor BaseOutputSettings;
+
 	// Override serialization to save Class Serialization version
 	virtual	void		Serialize(FArchive& Ar) override;
 
diff --git a/Engine/Plugins/TextureGraph/Source/TextureGraph/Public/TG_CustomVersion.h b/Engine/Plugins/TextureGraph/Source/TextureGraph/Public/TG_CustomVersion.h
--- a/Engine/Plugins/TextureGraph/Source/TextureGraph/Public/TG_CustomVersion.h
+++ b/Engine/Plugins/TextureGraph/Source/TextureGraph/Public/TG_CustomVersion.h
@@ -15,6 +15,9 @@ struct FTG_CustomVersion
 		// Added TexturePath to FTG_Texture
 		TGTextureAddedTexturePath,
 
+		// Added BaseOutputSettings to TG_Expression
+		TGExpressionAddedBaseOutputSettings,
+
 		// -----<new versions can be added above this line>-------------------------------------------------
 		VersionPlusOne,
 		LatestVersion = VersionPlusOne - 1
diff --git a/Engine/Plugins/TextureGraph/Source/TextureGraph/Public/TG_Texture.h b/Engine/Plugins/TextureGraph/Source/TextureGraph/Public/TG_Texture.h
--- a/Engine/Plugins/TextureGraph/Source/TextureGraph/Public/TG_Texture.h
+++ b/Engine/Plugins/TextureGraph/Source/TextureGraph/Public/TG_Texture.h
@@ -6,13 +6,14 @@
 #include "Data/RawBuffer.h"
 #include "Model/Mix/MixSettings.h"
 #include "2D/TextureHelper.h"
+#include "Misc/OutputDeviceNull.h"
 #include "TG_Texture.generated.h"
 
-USTRUCT()
+USTRUCT(BlueprintType)
 struct TEXTUREGRAPH_API FTG_TextureDescriptor
 {
-	GENERATED_BODY()
-		FTG_TextureDescriptor()
+	GENERATED_USTRUCT_BODY()
+	FTG_TextureDescriptor()
 	{
 	}
 
@@ -25,11 +26,11 @@ struct TEXTUREGRAPH_API FTG_TextureDescriptor
 
 public:
 	// Width of the texture in pixels. Auto means system will detect automatically based on other images
-	UPROPERTY(EditAnywhere, Category = "TextureDescriptor")
+	UPROPERTY(BlueprintReadWrite, EditAnywhere, Category = "Advanced", Meta = (NoResetToDefault))
 		EResolution Width = EResolution::Auto;		
 
 	// Height of the texture in pixels. Auto means system will detect automatically based on other images
-	UPROPERTY(EditAnywhere, Category = "TextureDescriptor")
+	UPROPERTY(BlueprintReadWrite, EditAnywhere, Category = "Advanced", Meta = (NoResetToDefault))
 		EResolution Height = EResolution::Auto;	
 
 #if WITH_EDITORONLY_DATA
@@ -42,9 +43,15 @@ PRAGMA_DISABLE_DEPRECATION_WARNINGS
 PRAGMA_ENABLE_DEPRECATION_WARNINGS
 #endif
 	
-	UPROPERTY(EditAnywhere, Category = "TextureDescriptor", DisplayName = "Texture Format")
+	//UPROPERTY(EditAnywhere, Category = "TextureDescriptor", DisplayName = "Texture Format")
+	UPROPERTY(BlueprintReadWrite, EditAnywhere, Category = "Advanced", DisplayName = "Texture Format", Meta = (NoResetToDefault))
 		ETG_TextureFormat TextureFormat = ETG_TextureFormat::Auto;
 
+	bool operator==(const FTG_TextureDescriptor& RHS) const
+	{
+		return Width == RHS.Width && Height == RHS.Height && TextureFormat == RHS.TextureFormat;
+	}
+
 	friend FArchive& operator<<(FArchive& Ar, FTG_TextureDescriptor& D)
 	{
  		Ar << D.Width;
@@ -71,6 +78,19 @@ PRAGMA_ENABLE_DEPRECATION_WARNINGS
 	{
 		return Width == EResolution::Auto && Height == EResolution::Auto && TextureFormat == ETG_TextureFormat::Auto;
 	}
+
+	void InitFromString(const FString& StrVal)
+	{
+		FOutputDeviceNull NullOut;
+		FTG_TextureDescriptor::StaticStruct()->ImportText(*StrVal, this, /*OwnerObject*/nullptr, 0, &NullOut, FTG_TextureDescriptor::StaticStruct()->GetName(), /*bAllowNativeOverride*/true);
+	}
+
+	FString ToString() const
+	{
+		FString ExportString;
+		FTG_TextureDescriptor::StaticStruct()->ExportText(ExportString, this, this, /*OwnerObject*/nullptr, /*PortFlags*/0, /*ExportRootScope*/nullptr);
+		return ExportString;
+	}
 };
 
 USTRUCT()
diff --git a/Engine/Plugins/TextureGraph/Source/TextureGraph/Public/TG_Var.h b/Engine/Plugins/TextureGraph/Source/TextureGraph/Public/TG_Var.h
--- a/Engine/Plugins/TextureGraph/Source/TextureGraph/Public/TG_Var.h
+++ b/Engine/Plugins/TextureGraph/Source/TextureGraph/Public/TG_Var.h
@@ -226,6 +226,7 @@ template <> FString TG_Var_LogValue(TObjectPtr<UObject>& Value);
 template <> FString TG_Var_LogValue(struct FTG_Texture& Value);
 template <> FString TG_Var_LogValue(struct FTG_Scalar& Value);
 template <> FString TG_Var_LogValue(struct FTG_OutputSettings& Value);
+template <> FString TG_Var_LogValue(struct FTG_TextureDescriptor& Value);
 template <> FString TG_Var_LogValue(struct FTG_Variant& Value);
 
 
@@ -244,6 +245,7 @@ template <> void TG_Var_SetValueFromString(FString& Value, const FString& StrVal
 //template <> void TG_Var_SetValueFromString(struct FTG_Texture& Value, const FString& StrVal);
 template <> void TG_Var_SetValueFromString(struct FTG_Scalar& Value, const FString& StrVal);
 template <> void TG_Var_SetValueFromString(struct FTG_OutputSettings& Value, const FString& StrVal);
+template <> void TG_Var_SetValueFromString(struct FTG_TextureDescriptor& Value, const FString& StrVal);
 template <> void TG_Var_SetValueFromString(struct FTG_Variant& Value, const FString& StrVal);
 
 
diff --git a/Engine/Plugins/TextureGraph/Source/TextureGraphEditor/Private/Customizations/TG_TextureCustomization.h b/Engine/Plugins/TextureGraph/Source/TextureGraphEditor/Private/Customizations/TG_TextureCustomization.h
--- a/Engine/Plugins/TextureGraph/Source/TextureGraphEditor/Private/Customizations/TG_TextureCustomization.h
+++ b/Engine/Plugins/TextureGraph/Source/TextureGraphEditor/Private/Customizations/TG_TextureCustomization.h
@@ -24,6 +24,7 @@ public:
 
 	virtual void CustomizeHeader(TSharedRef<IPropertyHandle> PropertyHandle, FDetailWidgetRow& HeaderRow, IPropertyTypeCustomizationUtils& CustomizationUtils) override
 	{
+#if 0
 		const FName TGType = FName(TEXT("TGType"));
 		const FName CategoryName = PropertyHandle->GetDefaultCategoryName();
 		const FString ParentTGType = PropertyHandle->GetParentHandle()->GetMetaData(TGType);
@@ -66,10 +67,12 @@ public:
 				TextureHistogram->SetTexture(*TexturePtr, TextureGraph);
 			}
 		}
+#endif 
 	}
 
 	virtual void CustomizeChildren(TSharedRef<IPropertyHandle> PropertyHandle, IDetailChildrenBuilder& ChildBuilder, IPropertyTypeCustomizationUtils& CustomizationUtils) override
 	{
+#if 0
 		const FString TGType = PropertyHandle->GetMetaData(FName(TEXT("TGType"))); 
 		if (TGType == TEXT("TG_Output"))
 		{
@@ -82,5 +85,6 @@ public:
 					.ShouldAutoExpand(true);
 			}
 		}
+#endif 
 	}
 };
diff --git a/Engine/Plugins/TextureGraph/Source/TextureGraphEditor/Private/EdGraph/STG_EditorGraphNode.cpp b/Engine/Plugins/TextureGraph/Source/TextureGraphEditor/Private/EdGraph/STG_EditorGraphNode.cpp
--- a/Engine/Plugins/TextureGraph/Source/TextureGraphEditor/Private/EdGraph/STG_EditorGraphNode.cpp
+++ b/Engine/Plugins/TextureGraph/Source/TextureGraphEditor/Private/EdGraph/STG_EditorGraphNode.cpp
@@ -329,7 +329,7 @@ void STG_EditorGraphNode::UpdateGraphNode()
 	CreateInputSideAddButton(LeftNodeBox);
 	CreateOutputSideAddButton(RightNodeBox);
 	CreateAdvancedViewArrow(InnerVerticalBox);
-	CreateNodeSettingsLayout();
+	//CreateNodeSettingsLayout();
 }
 END_SLATE_FUNCTION_BUILD_OPTIMIZATION
 
@@ -546,6 +546,38 @@ TSharedRef<SWidget> STG_EditorGraphNode::CreateNodeSettings()
 	{
 		UTG_Pin* TGPin = Schema->GetTGPinFromEdPin(Pin);
 		FTG_Texture& Texture = TGPin->GetNodePtr()->GetGraph()->GetVar(TGPin->GetId())->EditAs<FTG_Texture>();
+
+		TSharedPtr<SGraphPin> NewPin = CreatePinWidget(Pin);
+		check(NewPin.IsValid());
+		SetPinIcon(NewPin.ToSharedRef());
+
+		this->AddPin(NewPin.ToSharedRef());
+
+			//bool bShowLabel = true;
+			//// check if there is a display name defined for the property, we use that as the Pin Name
+			//if (Property && Property->HasMetaData("PinDisplayName"))
+			//{
+			//	FString DisplayName = Property->GetMetaData("PinDisplayName");
+			//	if (!DisplayName.IsEmpty())
+			//	{
+			//		CurPin->PinFriendlyName = FText::FromString(DisplayName);
+			//	}
+			//	else
+			//	{
+			//		NewPin->SetShowLabel(false);
+			//		bShowLabel = false;
+			//	}
+			//}
+
+			//Hide pin icon for not connectable pins
+			//If the pin do not have label then Icon space is not required
+			//If the pin has a label then we need to the hidden icon space to allign better with connectable pins
+		//	if (!NewPin->GetIsConnectable())
+		//	{
+		//		NewPin->GetPinImageWidget()->SetVisibility(bShowLabel ? EVisibility::Hidden : EVisibility::Collapsed);
+		//	}
+		//}
+
 			
 		Widget->AddSlot().HAlign(HAlign_Fill)
 		.Padding(Settings->GetNonPinNodeBodyPadding())
diff --git a/Engine/Plugins/TextureGraph/Source/TextureGraphEditor/Private/EdGraph/TG_EditorGraphPanelPinFactory.cpp b/Engine/Plugins/TextureGraph/Source/TextureGraphEditor/Private/EdGraph/TG_EditorGraphPanelPinFactory.cpp
--- a/Engine/Plugins/TextureGraph/Source/TextureGraphEditor/Private/EdGraph/TG_EditorGraphPanelPinFactory.cpp
+++ b/Engine/Plugins/TextureGraph/Source/TextureGraphEditor/Private/EdGraph/TG_EditorGraphPanelPinFactory.cpp
@@ -35,6 +35,7 @@
 #include "Pins/STG_GraphMaterialIdMaskInfos.h"
 #include "Transform/Expressions/T_ExtractMaterialIds.h"
 #include "Pins/STG_GraphPinOutputSettings.h"
+#include "Pins/STG_GraphPinTextureDescriptor.h"
 #include "Pins/STG_GraphPinString.h"
 
 #include "NumericPropertyParams.h"
@@ -106,6 +107,7 @@ TSharedPtr<class SGraphPin> FTG_EditorGraphPanelPinFactory::CreatePin(class UEdG
 			UScriptStruct* Vector2DStruct = TBaseStructure<FVector2D>::Get();
 			UScriptStruct* RotatorStruct = TBaseStructure<FRotator>::Get();
 			UScriptStruct* OutputSettingsStruct = TBaseStructure<FTG_OutputSettings>::Get();
+			UScriptStruct* TextureDescSettingsStruct = TBaseStructure<FTG_TextureDescriptor>::Get();
 
 			if (InPin->PinType.PinSubCategoryObject == ColorStruct)
 			{
@@ -131,6 +133,10 @@ TSharedPtr<class SGraphPin> FTG_EditorGraphPanelPinFactory::CreatePin(class UEdG
 			{
 				return SNew(STG_GraphPinOutputSettings, InPin);
 			}
+			else if (InPin->PinType.PinSubCategoryObject == TextureDescSettingsStruct)
+			{
+				return SNew(STG_GraphPinTextureDescriptor, InPin);
+			}
 		}
 		else if(InPin->PinType.PinCategory == UTG_EdGraphSchema::PC_Array)
 		{
diff --git a/Engine/Plugins/TextureGraph/Source/TextureGraphEditor/Private/Pins/STG_GraphPinTextureDescriptor.cpp b/Engine/Plugins/TextureGraph/Source/TextureGraphEditor/Private/Pins/STG_GraphPinTextureDescriptor.cpp
new file mode 100644
--- /dev/null
+++ b/Engine/Plugins/TextureGraph/Source/TextureGraphEditor/Private/Pins/STG_GraphPinTextureDescriptor.cpp
@@ -0,0 +1,155 @@
+// Copyright Epic Games, Inc. All Rights Reserved.
+
+#include "STG_GraphPinTextureDescriptor.h"
+#include "EdGraph/EdGraphSchema.h"
+#include "EdGraph/TG_EdGraphSchema.h"
+#include "ScopedTransaction.h"
+#include "TG_Pin.h"
+#include "STG_GraphPinTextureDescriptorWidget.h"
+#include "Widgets/SBoxPanel.h"
+#include "Widgets/Input/SCheckBox.h"
+#include "Widgets/Images/SImage.h"
+
+#define LOCTEXT_NAMESPACE "TextureDescriptorGraphPin"
+
+void STG_GraphPinTextureDescriptor::Construct(const FArguments& InArgs, UEdGraphPin* InGraphPinObj)
+{
+	SGraphPin::Construct( SGraphPin::FArguments(), InGraphPinObj );
+
+	bIsUIHidden = !ShowChildProperties();
+	
+	if (GetDirection() == EEdGraphPinDirection::EGPD_Input)
+	{
+		if (!CollapsibleChildProperties() && ShowChildProperties())
+		{
+			PinImage->SetVisibility(EVisibility::Collapsed);
+		}
+
+		GetPinObj()->bAdvancedView = ShowChildProperties();
+		if (GetPinObj()->GetOwningNode()->AdvancedPinDisplay != ENodeAdvancedPins::Shown)
+		{
+			GetPinObj()->GetOwningNode()->AdvancedPinDisplay = GetPinObj()->bAdvancedView ? ENodeAdvancedPins::Hidden : ENodeAdvancedPins::NoPins;
+		}
+	}
+
+	CachedImg_Pin_BackgroundHovered = CachedImg_Pin_Background;
+}
+
+void STG_GraphPinTextureDescriptor::OnTextureDescriptorChanged(const FTG_TextureDescriptor& NewTextureDescriptor)
+{
+	TextureDescriptor = NewTextureDescriptor;
+	const FString TextureDescriptorExportText = TextureDescriptor.ToString();
+
+	if (TextureDescriptorExportText != GraphPinObj->GetDefaultAsString())
+	{
+		// Set Pin Data
+		const FScopedTransaction Transaction(NSLOCTEXT("GraphEditor", "ChangePinValue", "Change Pin Value"));
+		GraphPinObj->GetSchema()->TrySetDefaultValue(*GraphPinObj, TextureDescriptorExportText);
+	}
+}
+
+FProperty* STG_GraphPinTextureDescriptor::GetPinProperty() const
+{
+	const UTG_EdGraphSchema* Schema = Cast<const UTG_EdGraphSchema>(GraphPinObj->GetOwningNode()->GetSchema());
+	UTG_Pin* TSPin = Schema->GetTGPinFromEdPin(GraphPinObj);
+	FProperty* Property = TSPin->GetExpressionProperty();
+	return Property;
+}
+
+bool STG_GraphPinTextureDescriptor::ShowChildProperties() const
+{
+	FProperty* Property = GetPinProperty();
+	bool ShowChildProperties = true;
+	// check if there is a display name defined for the property, we use that as the Pin Name
+	if (Property && Property->HasMetaData("HideChildProperties"))
+	{
+		ShowChildProperties = false;
+	}
+	return ShowChildProperties;
+}
+
+bool STG_GraphPinTextureDescriptor::CollapsibleChildProperties() const
+{
+	FProperty* Property = GetPinProperty();
+	bool Collapsible = false;
+	// check if there is a display name defined for the property, we use that as the Pin Name
+	if (Property && Property->HasMetaData("CollapsableChildProperties"))
+	{
+		Collapsible = true;
+	}
+	return Collapsible;
+}
+
+EVisibility STG_GraphPinTextureDescriptor::ShowLabel() const
+{
+	bool bOutput = GetDirection() == EEdGraphPinDirection::EGPD_Output;
+	bool bHide = ShowChildProperties() && !bOutput;
+
+	if ((GraphPinObj->GetOwningNode()->AdvancedPinDisplay == ENodeAdvancedPins::Type::Hidden && GraphPinObj->LinkedTo.Num() > 0))
+	{
+		bHide = false;
+	}
+
+	return bHide ? EVisibility::Collapsed : EVisibility::Visible;
+}
+
+TSharedRef<SWidget>	STG_GraphPinTextureDescriptor::GetDefaultValueWidget()
+{
+	if (ShowChildProperties())
+	{
+		ParseDefaultValueData();
+
+		return SNew(STG_GraphPinTextureDescriptorWidget, GraphPinObj)
+			.Visibility(this, &STG_GraphPinTextureDescriptor::IsUIEnabled)
+			.DescriptionMaxWidth(250.0f)
+			.TextureDescriptor(this, &STG_GraphPinTextureDescriptor::GetTextureDescriptor)
+			.OnTextureDescriptorChanged(this, &STG_GraphPinTextureDescriptor::OnTextureDescriptorChanged)
+			.IsEnabled(this, &STG_GraphPinTextureDescriptor::GetDefaultValueIsEnabled);
+	}
+	else
+	{
+		return SGraphPin::GetDefaultValueWidget();
+	}
+}
+
+TSharedRef<SWidget> STG_GraphPinTextureDescriptor::GetLabelWidget(const FName& InLabelStyle)
+{
+	return SNew(STextBlock)
+			.Text(this, &STG_GraphPinTextureDescriptor::GetPinLabel)
+			.TextStyle(FAppStyle::Get(), InLabelStyle)
+			.Visibility(this, &STG_GraphPinTextureDescriptor::ShowLabel)
+			.ColorAndOpacity(this, &STG_GraphPinTextureDescriptor::GetPinTextColor);
+}
+
+void STG_GraphPinTextureDescriptor::OnAdvancedViewChanged(const ECheckBoxState NewCheckedState)
+{
+	bIsUIHidden = NewCheckedState != ECheckBoxState::Checked;
+}
+
+ECheckBoxState STG_GraphPinTextureDescriptor::IsAdvancedViewChecked() const
+{
+	return bIsUIHidden ? ECheckBoxState::Unchecked : ECheckBoxState::Checked;
+}
+
+EVisibility STG_GraphPinTextureDescriptor::IsUIEnabled() const
+{
+	return bIsUIHidden ? EVisibility::Collapsed : EVisibility::Visible;
+}
+
+const FSlateBrush* STG_GraphPinTextureDescriptor::GetAdvancedViewArrow() const
+{
+	return FAppStyle::GetBrush(bIsUIHidden ? TEXT("Icons.ChevronDown") : TEXT("Icons.ChevronUp"));
+}
+
+FTG_TextureDescriptor STG_GraphPinTextureDescriptor::GetTextureDescriptor() const
+{
+	return TextureDescriptor;
+}
+
+void STG_GraphPinTextureDescriptor::ParseDefaultValueData()
+{
+	FString const TextureDescriptorString = GraphPinObj->GetDefaultAsString();
+	TextureDescriptor.InitFromString(TextureDescriptorString);
+}
+
+#undef LOCTEXT_NAMESPACE
diff --git a/Engine/Plugins/TextureGraph/Source/TextureGraphEditor/Private/Pins/STG_GraphPinTextureDescriptor.h b/Engine/Plugins/TextureGraph/Source/TextureGraphEditor/Private/Pins/STG_GraphPinTextureDescriptor.h
new file mode 100644
--- /dev/null
+++ b/Engine/Plugins/TextureGraph/Source/TextureGraphEditor/Private/Pins/STG_GraphPinTextureDescriptor.h
@@ -0,0 +1,46 @@
+// Copyright Epic Games, Inc. All Rights Reserved.
+
+#pragma once
+
+#include "SGraphPin.h"
+#include "TG_Texture.h"
+
+class SComboButton;
+
+class STG_GraphPinTextureDescriptor : public SGraphPin
+{
+public:
+	SLATE_BEGIN_ARGS(STG_GraphPinTextureDescriptor) {}
+	SLATE_END_ARGS()
+
+	void Construct(const FArguments& InArgs, UEdGraphPin* InGraphPinObj);
+	FProperty* GetPinProperty() const;
+	bool ShowChildProperties() const;
+	bool CollapsibleChildProperties() const;
+	EVisibility ShowLabel() const;
+	//~ Begin SGraphPin Interface
+	virtual TSharedRef<SWidget>	GetDefaultValueWidget() override;
+	virtual TSharedRef<SWidget> GetLabelWidget(const FName& InLabelStyle) override;
+	//~ End SGraphPin Interface
+
+private:
+	/** Parses the Data from the pin to fill in the names of the array. */
+	void ParseDefaultValueData();
+
+	bool GetDefaultValueIsEnabled() const
+	{
+		return !GraphPinObj->bDefaultValueIsReadOnly;
+	}
+
+	const FSlateBrush* GetAdvancedViewArrow() const;
+	void OnAdvancedViewChanged(const ECheckBoxState NewCheckedState);
+	ECheckBoxState IsAdvancedViewChecked() const;
+	EVisibility IsUIEnabled() const;
+
+	FTG_TextureDescriptor GetTextureDescriptor() const;
+	void OnTextureDescriptorChanged(const FTG_TextureDescriptor& NewTextureDescriptor);
+
+	/** Parse TextureDescriptor used for editing. */
+	FTG_TextureDescriptor TextureDescriptor;
+	bool bIsUIHidden = true;
+};
diff --git a/Engine/Plugins/TextureGraph/Source/TextureGraphEditor/Private/Pins/STG_GraphPinTextureDescriptorWidget.cpp b/Engine/Plugins/TextureGraph/Source/TextureGraphEditor/Private/Pins/STG_GraphPinTextureDescriptorWidget.cpp
new file mode 100644
--- /dev/null
+++ b/Engine/Plugins/TextureGraph/Source/TextureGraphEditor/Private/Pins/STG_GraphPinTextureDescriptorWidget.cpp
@@ -0,0 +1,342 @@
+// Copyright Epic Games, Inc. All Rights Reserved.
+
+#include "STG_GraphPinTextureDescriptorWidget.h"
+#include "DetailLayoutBuilder.h"
+#include "Widgets/Text/STextBlock.h"
+#include "Widgets/Input/SButton.h"
+#include "Widgets/Images/SImage.h"
+#include "Widgets/Input/SEditableTextBox.h"
+#include "Widgets/Input/SComboButton.h"
+#include "ScopedTransaction.h"
+#include "Editor.h"
+#include "Widgets/Layout/SWrapBox.h"
+#include "SGraphPinComboBox.h"
+#include "Framework/MultiBox/MultiBoxBuilder.h"
+
+#define LOCTEXT_NAMESPACE "STG_GraphPinTextureDescriptorWidget"
+
+//------------------------------------------------------------------------------
+// STG_GraphPinTextureDescriptorWidget
+//------------------------------------------------------------------------------
+
+SLATE_IMPLEMENT_WIDGET(STG_GraphPinTextureDescriptorWidget)
+void STG_GraphPinTextureDescriptorWidget::PrivateRegisterAttributes(FSlateAttributeInitializer& AttributeInitializer)
+{
+	SLATE_ADD_MEMBER_ATTRIBUTE_DEFINITION_WITH_NAME(AttributeInitializer, "TextureDescriptor", TextureDescriptorAttribute, EInvalidateWidgetReason::Layout)
+	.OnValueChanged(FSlateAttributeDescriptor::FAttributeValueChangedDelegate::CreateLambda([](SWidget& Widget)
+		{
+			//static_cast<STG_GraphPinTextureDescriptorWidget&>(Widget).CacheQueryList();
+		}));
+}
+
+STG_GraphPinTextureDescriptorWidget::STG_GraphPinTextureDescriptorWidget()
+	: TextureDescriptorAttribute(*this)
+{
+
+}
+
+STG_GraphPinTextureDescriptorWidget::~STG_GraphPinTextureDescriptorWidget()
+{
+	/*if (bRegisteredForUndo)
+	{
+	GEditor->UnregisterForUndo(this);
+	}*/
+}
+
+void STG_GraphPinTextureDescriptorWidget::Construct(const FArguments& InArgs, UEdGraphPin* InGraphPinObj)
+{
+	OnGenerateWidthMenu.BindRaw(this, &STG_GraphPinTextureDescriptorWidget::OnGenerateWidthEnumMenu);
+	GetWidthDelegate.BindRaw(this, &STG_GraphPinTextureDescriptorWidget::HandleWidthText);
+	OnGenerateHeightMenu.BindRaw(this, &STG_GraphPinTextureDescriptorWidget::OnGenerateHeightEnumMenu);
+	GetHeightDelegate.BindRaw(this, &STG_GraphPinTextureDescriptorWidget::HandleHeightText);
+	OnGenerateFormatMenu.BindRaw(this, &STG_GraphPinTextureDescriptorWidget::OnGenerateFormatEnumMenu);
+	GetFormatDelegate.BindRaw(this, &STG_GraphPinTextureDescriptorWidget::HandleFormatText);
+
+	GraphPinObj = InGraphPinObj;
+
+	OnTextureDescriptorChanged = InArgs._OnTextureDescriptorChanged;
+
+	TextureDescriptorAttribute.Assign(*this, InArgs._TextureDescriptor);
+
+	const int UniformPadding = 2;
+
+	ChildSlot
+		[
+			//Param Box hidden when pin is connected and Advanced view is collapsed
+			SNew(SBox)
+				.MinDesiredWidth(LabelSize * 2)
+				.MaxDesiredWidth(LabelSize * 4)
+				.Visibility(this, &STG_GraphPinTextureDescriptorWidget::ShowParameters)
+				[
+					SNew(SVerticalBox)
+						//Path Width
+						+SVerticalBox::Slot()
+						.Padding(UniformPadding)
+						[
+							AddEnumComobox(LOCTEXT("OutputWidth", "Width"), GetWidthDelegate, OnGenerateWidthMenu)
+						]
+
+						//Path Height
+						+ SVerticalBox::Slot()
+						.Padding(UniformPadding)
+						[
+							AddEnumComobox(LOCTEXT("OutputHeight", "Height"), GetHeightDelegate, OnGenerateHeightMenu)
+						]
+
+						//Path Format
+						+ SVerticalBox::Slot()
+						.Padding(UniformPadding)
+						[
+							AddEnumComobox(LOCTEXT("OutputFormat", "Format"), GetFormatDelegate, OnGenerateFormatMenu)
+						]
+				]
+		];
+}
+
+EVisibility STG_GraphPinTextureDescriptorWidget::ShowPinLabel() const
+{
+	return ShowParameters() == EVisibility::Collapsed ? EVisibility::Visible : EVisibility::Collapsed;
+}
+
+EVisibility STG_GraphPinTextureDescriptorWidget::ShowParameters() const
+{
+	return (GraphPinObj->GetOwningNode()->AdvancedPinDisplay == ENodeAdvancedPins::Type::Hidden && GraphPinObj->LinkedTo.Num() > 0) ? EVisibility::Collapsed : EVisibility::Visible;
+}
+
+TSharedRef<SWidget> STG_GraphPinTextureDescriptorWidget::AddEnumComobox(FText Label, FGetTextDelegate GetText, FGenerateEnumMenu OnGenerateEnumMenu)
+{
+	return 
+		SNew(SHorizontalBox)
+
+		+ SHorizontalBox::Slot()
+		.VAlign(VAlign_Center)
+		.HAlign(HAlign_Center)
+		.MinSize(LabelSize)
+		.FillWidth(1.0f)
+		[
+			SNew(SBox)
+				.MinDesiredWidth(LabelSize)
+				.MaxDesiredWidth(LabelSize)
+				.HAlign(HAlign_Left)
+				[
+					SNew(STextBlock)
+						.Text(Label)
+						.TextStyle(FAppStyle::Get(), TEXT("Graph.Node.PinName"))
+				]
+		]
+
+		+ SHorizontalBox::Slot()
+		.MaxWidth(LabelSize * 3)
+		[
+			SNew(SComboButton)
+				.HAlign(HAlign_Right)
+				.OnGetMenuContent(OnGenerateEnumMenu)
+				.ButtonContent()
+				[
+					SNew(STextBlock)
+						.Text_Lambda([GetText]() { return GetText.Execute(); })
+				]
+		];
+}
+
+FTG_TextureDescriptor STG_GraphPinTextureDescriptorWidget::GetSettings() const
+{
+	FString TextureDescriptorString = GraphPinObj->GetDefaultAsString();
+	FTG_TextureDescriptor Settings;
+	Settings.InitFromString(TextureDescriptorString);
+
+	return Settings;
+}
+
+void STG_GraphPinTextureDescriptorWidget::GenerateStringsFromEnum(TArray<FString>& OutEnumNames, const FString& EnumPathName)
+{
+	UEnum* EnumPtr = FindObject<UEnum>(nullptr, *EnumPathName);
+	if (EnumPtr)
+	{
+		for (int32 i = 0; i < EnumPtr->NumEnums() - 1; ++i)
+		{
+			if (!EnumPtr->HasMetaData(TEXT("Hidden"), i))
+			{
+				FString DisplayName = EnumPtr->GetDisplayNameTextByIndex(i).ToString();
+				uint8 EnumValue = EnumPtr->GetValueByIndex(i);
+				UE_LOG(LogTemp, Warning, TEXT("Enum Value: %d, Display Name: %s"), EnumValue, *DisplayName);
+				OutEnumNames.Add(DisplayName);
+			}
+		}
+	}
+}
+
+template<typename T>
+void STG_GraphPinTextureDescriptorWidget::GenerateValuesFromEnum(TArray<T>& OutEnumValues, const FString& EnumPathName) const
+{
+	UEnum* EnumPtr = FindObject<UEnum>(nullptr, *EnumPathName);
+	if (EnumPtr)
+	{
+		for (int32 i = 0; i < EnumPtr->NumEnums() - 1; ++i)
+		{
+			if (!EnumPtr->HasMetaData(TEXT("Hidden"), i))
+			{
+				FString DisplayName = EnumPtr->GetDisplayNameTextByIndex(i).ToString();
+				T EnumValue = EnumPtr->GetValueByIndex(i);
+				UE_LOG(LogTemp, Warning, TEXT("Enum Value: %d, Display Name: %s"), EnumValue, *DisplayName);
+				OutEnumValues.Add(EnumValue);
+			}
+		}
+	}
+}
+
+template<typename T>
+int STG_GraphPinTextureDescriptorWidget::GetValueFromIndex(const FString& EnumPathName, int Index) const
+{
+	int Value = 0;
+	TArray<T> EnumValues;
+	UEnum* EnumPtr = FindObject<UEnum>(nullptr, *EnumPathName);
+	if (EnumPtr)
+	{
+		GenerateValuesFromEnum<T>(EnumValues, EnumPathName);
+		if (EnumValues.Num() > Index)
+		{
+			Value = EnumValues[Index];
+		}
+	}
+	return Value;
+}
+
+FString STG_GraphPinTextureDescriptorWidget::GetEnumValueDisplayName(const FString& EnumPathName, int EnumValue) const
+{
+	UEnum* EnumPtr = FindObject<UEnum>(nullptr, *EnumPathName, true);
+	if (EnumPtr)
+	{
+		FString DisplayName = EnumPtr->GetDisplayNameTextByValue(EnumValue).ToString();
+		return DisplayName;
+	}
+	return FString();
+}
+
+TSharedRef<SWidget> STG_GraphPinTextureDescriptorWidget::OnGenerateWidthEnumMenu()
+{
+	FMenuBuilder MenuBuilder(true, nullptr);
+	//Get list for Enum
+	TArray<FString> ResolutionEnumItems;
+	UEnum* Resolution = StaticEnum<EResolution>();
+	GenerateStringsFromEnum(ResolutionEnumItems, Resolution->GetPathName());
+
+	for (int i = 0; i < ResolutionEnumItems.Num(); i++)
+	{
+		auto Item = ResolutionEnumItems[i];
+		MenuBuilder.AddMenuEntry(
+			FText::FromString(Item),
+			FText::FromString(Item),
+			FSlateIcon(),
+			FUIAction(
+				FExecuteAction::CreateSP(this, &STG_GraphPinTextureDescriptorWidget::HandleWidthChanged, Item, i),
+				FCanExecuteAction(),
+				FIsActionChecked::CreateLambda([this, i]() {return SelectedWidthIndex == i; })
+			));
+	}
+
+	return MenuBuilder.MakeWidget();
+}
+
+void STG_GraphPinTextureDescriptorWidget::HandleWidthChanged(FString Name, int Index)
+{
+	auto Settings = GetSettings();
+	Settings.Width = (EResolution)GetValueFromIndex<int>(StaticEnum<EResolution>()->GetPathName(), Index);
+
+	OnTextureDescriptorChanged.ExecuteIfBound(Settings);
+
+	SelectedWidthIndex = Index;
+	SelectedWidthName = Name;
+}
+
+FText STG_GraphPinTextureDescriptorWidget::HandleWidthText() const
+{
+	return FText::FromString(GetEnumValueDisplayName(StaticEnum<EResolution>()->GetPathName(), (int)GetSettings().Width));
+}
+
+TSharedRef<SWidget> STG_GraphPinTextureDescriptorWidget::OnGenerateHeightEnumMenu()
+{
+	FMenuBuilder MenuBuilder(true, nullptr);
+	//Get list for Enum
+	TArray<FString> ResolutionEnumItems;
+	GenerateStringsFromEnum(ResolutionEnumItems, StaticEnum<EResolution>()->GetPathName());
+
+	for (int i = 0; i < ResolutionEnumItems.Num(); i++)
+	{
+		auto Item = ResolutionEnumItems[i];
+		MenuBuilder.AddMenuEntry(
+			FText::FromString(Item),
+			FText::FromString(Item),
+			FSlateIcon(),
+			FUIAction(
+				FExecuteAction::CreateSP(this, &STG_GraphPinTextureDescriptorWidget::HandleHeightChanged, Item, i),
+				FCanExecuteAction(),
+				FIsActionChecked::CreateLambda([this, i]() {return SelectedHeightIndex == i; })
+			));
+	}
+
+	return MenuBuilder.MakeWidget();
+}
+
+void STG_GraphPinTextureDescriptorWidget::HandleHeightChanged(FString Name, int Index)
+{
+	auto Settings = GetSettings();
+	Settings.Height = (EResolution)GetValueFromIndex<int>(StaticEnum<EResolution>()->GetPathName(), Index);
+	OnTextureDescriptorChanged.ExecuteIfBound(Settings);
+
+	SelectedHeightIndex = Index;
+}
+
+FText STG_GraphPinTextureDescriptorWidget::HandleHeightText() const
+{
+	return FText::FromString(GetEnumValueDisplayName(StaticEnum<EResolution>()->GetPathName(), (int)GetSettings().Height));
+}
+
+TSharedRef<SWidget> STG_GraphPinTextureDescriptorWidget::OnGenerateFormatEnumMenu()
+{
+	FMenuBuilder MenuBuilder(true, nullptr);
+	//Get list for Enum
+	TArray<FString> ResolutionEnumItems;
+	GenerateStringsFromEnum(ResolutionEnumItems, StaticEnum<ETG_TextureFormat>()->GetPathName());
+
+	for (int i = 0; i < ResolutionEnumItems.Num(); i++)
+	{
+		auto Item = ResolutionEnumItems[i];
+		MenuBuilder.AddMenuEntry(
+			FText::FromString(Item),
+			FText::FromString(Item),
+			FSlateIcon(),
+			FUIAction(
+				FExecuteAction::CreateSP(this, &STG_GraphPinTextureDescriptorWidget::HandleFormatChanged, Item, i),
+				FCanExecuteAction(),
+				FIsActionChecked::CreateLambda([this, i]() {return SelectedFormatIndex == i; })
+			));
+	}
+
+	return MenuBuilder.MakeWidget();
+}
+
+void STG_GraphPinTextureDescriptorWidget::HandleFormatChanged(FString Name, int Index)
+{
+	auto Settings = GetSettings();
+	Settings.TextureFormat = (ETG_TextureFormat)GetValueFromIndex<uint8>(StaticEnum<ETG_TextureFormat>()->GetPathName(), Index);
+
+	OnTextureDescriptorChanged.ExecuteIfBound(Settings);
+
+	SelectedFormatIndex = Index;
+}
+
+FText STG_GraphPinTextureDescriptorWidget::HandleFormatText() const
+{
+	return FText::FromString(GetEnumValueDisplayName(StaticEnum<ETG_TextureFormat>()->GetPathName(), (int)GetSettings().TextureFormat));
+}
+
+void STG_GraphPinTextureDescriptorWidget::PostUndo(bool bSuccess)
+{
+}
+
+void STG_GraphPinTextureDescriptorWidget::PostRedo(bool bSuccess)
+{
+}
+
+#undef LOCTEXT_NAMESPACE
diff --git a/Engine/Plugins/TextureGraph/Source/TextureGraphEditor/Private/Pins/STG_GraphPinTextureDescriptorWidget.h b/Engine/Plugins/TextureGraph/Source/TextureGraphEditor/Private/Pins/STG_GraphPinTextureDescriptorWidget.h
new file mode 100644
--- /dev/null
+++ b/Engine/Plugins/TextureGraph/Source/TextureGraphEditor/Private/Pins/STG_GraphPinTextureDescriptorWidget.h
@@ -0,0 +1,105 @@
+// Copyright Epic Games, Inc. All Rights Reserved.
+
+#pragma once
+
+#include "Widgets/SWidget.h"
+#include "Widgets/SCompoundWidget.h"
+#include "Widgets/DeclarativeSyntaxSupport.h"
+#include "Widgets/Input/SCheckBox.h"
+#include "TG_Texture.h"
+#include "EditorUndoClient.h"
+#include "EdGraph/EdGraphPin.h"
+
+class IPropertyHandle;
+class SHorizontalBox;
+
+/**
+ * Widget for editing TextureDescriptor.
+ */
+class STG_GraphPinTextureDescriptorWidget : public SCompoundWidget, public FEditorUndoClient
+{
+	SLATE_DECLARE_WIDGET(STG_GraphPinTextureDescriptorWidget, SCompoundWidget)
+
+private:
+	DECLARE_DELEGATE_RetVal(FText, FGetTextDelegate);
+	DECLARE_DELEGATE_TwoParams(FTextCommitted, const FText&, ETextCommit::Type);
+	DECLARE_DELEGATE_RetVal(TSharedRef<SWidget>, FGenerateEnumMenu);
+
+public:
+
+	DECLARE_DELEGATE_OneParam(FOnTextureDescriptorChanged, const FTG_TextureDescriptor& /*TextureDescriptor*/)
+
+	SLATE_BEGIN_ARGS(STG_GraphPinTextureDescriptorWidget)
+		:
+		_DescriptionMaxWidth(250.0f)
+		, _PropertyHandle(nullptr)
+		{}
+
+		/** Maximum with of the query description field. */
+		SLATE_ARGUMENT(float, DescriptionMaxWidth)
+
+		/** TextureDescriptor to edit */
+		SLATE_ATTRIBUTE(FTG_TextureDescriptor, TextureDescriptor)
+
+		/** If set, the TextureDescriptor is read from the property, and the property is update when TextureDescriptor is edited. */ 
+		SLATE_ARGUMENT(TSharedPtr<IPropertyHandle>, PropertyHandle)
+
+		SLATE_EVENT(FOnTextureDescriptorChanged, OnTextureDescriptorChanged)
+	SLATE_END_ARGS();
+
+	STG_GraphPinTextureDescriptorWidget();
+	virtual ~STG_GraphPinTextureDescriptorWidget() override;
+
+	void Construct(const FArguments& InArgs, UEdGraphPin* InGraphPinObj);
+	TSharedRef<SWidget> AddEnumComobox(FText Label, FGetTextDelegate GetText, FGenerateEnumMenu OnGenerateEnumMenu);
+	FTG_TextureDescriptor GetSettings() const;
+	void GenerateStringsFromEnum(TArray<FString>& OutEnumNames, const FString& EnumPathName);
+
+	template<typename T>
+	void GenerateValuesFromEnum(TArray<T>& OutEnumValues, const FString& EnumPathName) const;
+
+	template<typename T>
+	int GetValueFromIndex(const FString& EnumPathName, int Index) const;
+	FString GetEnumValueDisplayName(const FString& EnumPathName, int EnumValue) const;
+	EVisibility ShowParameters() const;
+	EVisibility ShowPinLabel() const;
+
+protected:
+	//~ Begin FEditorUndoClient Interface
+	virtual void PostUndo(bool bSuccess) override;
+	virtual void PostRedo(bool bSuccess) override;
+	//~ End FEditorUndoClient Interface
+
+private:
+	UEdGraphPin* GraphPinObj;
+
+	TSlateAttribute<FTG_TextureDescriptor> TextureDescriptorAttribute;
+	FOnTextureDescriptorChanged OnTextureDescriptorChanged;
+	FTG_TextureDescriptor CachedTextureDescriptor;
+	FString SelectedWidthName;
+
+	int SelectedWidthIndex;
+	int SelectedHeightIndex;
+	int SelectedFormatIndex;
+
+	FGenerateEnumMenu OnGenerateWidthMenu;
+	FGetTextDelegate GetWidthDelegate;
+	FGenerateEnumMenu OnGenerateHeightMenu;
+	FGetTextDelegate GetHeightDelegate;
+	FGenerateEnumMenu OnGenerateFormatMenu;
+	FGetTextDelegate GetFormatDelegate;
+
+	TSharedRef<SWidget> OnGenerateWidthEnumMenu();
+	void HandleWidthChanged(FString OuputName, int Index);
+	FText HandleWidthText() const;
+
+	TSharedRef<SWidget> OnGenerateHeightEnumMenu();
+	void HandleHeightChanged(FString OuputName, int Index);
+	FText HandleHeightText() const;
+
+	TSharedRef<SWidget> OnGenerateFormatEnumMenu();
+	void HandleFormatChanged(FString OuputName, int Index);
+	FText HandleFormatText() const;
+
+	const int LabelSize = 75;
+};
diff --git a/Engine/Plugins/TextureGraph/Source/TextureGraphEditor/Private/Pins/STG_TextureDescriptor.h b/Engine/Plugins/TextureGraph/Source/TextureGraphEditor/Private/Pins/STG_TextureDescriptor.h
--- a/Engine/Plugins/TextureGraph/Source/TextureGraphEditor/Private/Pins/STG_TextureDescriptor.h
+++ b/Engine/Plugins/TextureGraph/Source/TextureGraphEditor/Private/Pins/STG_TextureDescriptor.h
@@ -30,7 +30,7 @@ public:
 	SLATE_BEGIN_ARGS(STG_TextureDescriptor)
 		:
 		_DescriptionMaxWidth(250.0f)
-		,_ShowingAsPin(false)
+		,_ShowingAsPin(true)
 		, _PropertyHandle(nullptr)
 	{}
 
