Load time crash fix for variant types

diff --git a/Engine/Plugins/TextureGraph/Source/TextureGraph/Private/Expressions/TG_Expression.cpp b/Engine/Plugins/TextureGraph/Source/TextureGraph/Private/Expressions/TG_Expression.cpp
--- a/Engine/Plugins/TextureGraph/Source/TextureGraph/Private/Expressions/TG_Expression.cpp
+++ b/Engine/Plugins/TextureGraph/Source/TextureGraph/Private/Expressions/TG_Expression.cpp
@@ -296,7 +296,10 @@ void UTG_Expression::SetupAndEvaluate(FTG_EvaluationContext* InContext)
 		{
 			FTG_Variant& Variant = OutputVar.Value.Var->GetAs<FTG_Variant>();
 			if (Variant.IsTexture())
-				Variant.GetTexture().Descriptor = BaseOutputSettings;
+			{
+				FTG_Texture& Texture = Variant.GetTexture();
+				Texture.Descriptor = BaseOutputSettings;
+			}
 		}
 		OutputVar.Value.Var->CopyTo(this, OutputVar.Value.Argument);
 	}
diff --git a/Engine/Plugins/TextureGraph/Source/TextureGraph/Private/Expressions/TG_Expression_Variant.cpp b/Engine/Plugins/TextureGraph/Source/TextureGraph/Private/Expressions/TG_Expression_Variant.cpp
--- a/Engine/Plugins/TextureGraph/Source/TextureGraph/Private/Expressions/TG_Expression_Variant.cpp
+++ b/Engine/Plugins/TextureGraph/Source/TextureGraph/Private/Expressions/TG_Expression_Variant.cpp
@@ -112,6 +112,7 @@ void UTG_Expression_Variant::Evaluate(FTG_EvaluationContext* InContext)
 	}
 	case FTG_Variant::EType::Texture:
 	{
+		Output.EditTexture().Descriptor = BaseOutputSettings;
 		Output.EditTexture() = EvaluateTexture(InContext).RasterBlob;
 		break;
 	}
