Fix array concatenate crash

diff --git a/Engine/Plugins/TextureGraph/Source/TextureGraph/Private/Expressions/Arrays/TG_Expression_ArrayOps.cpp b/Engine/Plugins/TextureGraph/Source/TextureGraph/Private/Expressions/Arrays/TG_Expression_ArrayOps.cpp
--- a/Engine/Plugins/TextureGraph/Source/TextureGraph/Private/Expressions/Arrays/TG_Expression_ArrayOps.cpp
+++ b/Engine/Plugins/TextureGraph/Source/TextureGraph/Private/Expressions/Arrays/TG_Expression_ArrayOps.cpp
@@ -20,14 +20,14 @@ void UTG_Expression_ArrayConcat::Evaluate(FTG_EvaluationContext* InContext)
 
 	uint32 ActualStartIndex = StartIndex >= 0 ? StartIndex : 0;
 
-	if (ActualStartIndex >= MaxCount)
+	if (ActualStartIndex >= MaxCount || ActualStartIndex >= Input1Count)
 	{
 		FTextureGraphErrorReporter* ErrorReporter = TextureGraphEngine::GetErrorReporter(InContext->Cycle->GetMix());
 
 		if (ErrorReporter)
 		{
-			ErrorReporter->ReportWarning((int32)ETextureGraphErrorType::NODE_WARNING, FString::Printf(TEXT("Invalid starting index for Input-1 in the output array specified: %d [Output Array Size: %d]."), 
-				StartIndex, MaxCount), GetParentNode());
+			ErrorReporter->ReportWarning((int32)ETextureGraphErrorType::NODE_WARNING, FString::Printf(TEXT("Invalid starting index for Input-1 in the output array specified: %d [Input-1 Size: %d | Output Size: %d]."), 
+				StartIndex, Input1Count, MaxCount), GetParentNode());
 		}
 
 		ActualStartIndex = 0;
