Undo/Redo test fixes

diff --git a/Engine/Plugins/TextureGraph/Source/TextureGraph/Private/TG_Pin.cpp b/Engine/Plugins/TextureGraph/Source/TextureGraph/Private/TG_Pin.cpp
--- a/Engine/Plugins/TextureGraph/Source/TextureGraph/Private/TG_Pin.cpp
+++ b/Engine/Plugins/TextureGraph/Source/TextureGraph/Private/TG_Pin.cpp
@@ -583,32 +583,10 @@ bool UTG_Pin::SetValue(const float* Value, size_t Count)
 
 bool UTG_Pin::SetValue(const FTG_TextureDescriptor& Value)
 {
-	if (IsArgTexture())
-	{
-		Modify();
-
-		if (IsArgVariant())
-		{
-			EditSelfVar()->EditAs<FTG_Variant>().EditTexture().Descriptor = Value;
-		}
-		else if (IsArgTextureArray())
-		{
-			FTG_VariantArray& VarArray = EditSelfVar()->EditAs<FTG_VariantArray>();
-			for (int Index = 0; Index < VarArray.Num(); Index++)
-			{
-				VarArray.Get(Index).GetTexture().Descriptor = Value;
-			}
-		}
-		else
-		{
-			EditSelfVar()->EditAs<FTG_Texture>().Descriptor = Value;
-		}
-
-		NotifyPinSelfVarChanged();
-		return true;
-	}
-
-	return false;
+	Modify(true);
+	EditSelfVar()->EditAs<FTG_TextureDescriptor>() = Value;
+	NotifyPinSelfVarChanged();
+	return true;
 }
 
 bool UTG_Pin::SetValue(float Value)
diff --git a/Engine/Plugins/TextureGraph/Source/TextureGraphEditor/Private/EdGraph/TG_EdGraphNode.cpp b/Engine/Plugins/TextureGraph/Source/TextureGraphEditor/Private/EdGraph/TG_EdGraphNode.cpp
--- a/Engine/Plugins/TextureGraph/Source/TextureGraphEditor/Private/EdGraph/TG_EdGraphNode.cpp
+++ b/Engine/Plugins/TextureGraph/Source/TextureGraphEditor/Private/EdGraph/TG_EdGraphNode.cpp
@@ -397,6 +397,8 @@ FString UTG_EdGraphNode::GetTitleDetail()
 	
 	const UEdGraphPin* SelectedPinPtr = GetSelectedPin();
 
+	FString ArrayInfo;
+
 	for (auto Pin : OutPins)
 	{
 		const UTG_Pin* TGPin = Pin;
@@ -412,8 +414,6 @@ FString UTG_EdGraphNode::GetTitleDetail()
 		{
 			if (Texture.RasterBlob)
 			{
-				FString ArrayInfo;
-
 				if (Pin->IsArgTextureArray())
 				{
 					FTG_VariantArray VarArray;
@@ -421,19 +421,25 @@ FString UTG_EdGraphNode::GetTitleDetail()
 					ArrayInfo = FString::Printf(TEXT("[%d]"), VarArray.Num());
 				}
 
-				const auto Desc = Texture.RasterBlob->GetDescriptor();
-				FStringFormatNamedArguments Args;
-				Args.Add(TEXT("Channels"), TextureHelper::GetChannelsTextFromItemsPerPoint(Desc.ItemsPerPoint));
-				Args.Add(TEXT("Format"), Desc.FormatToString(Desc.Format));
-				Args.Add(TEXT("IsSRGB"), Desc.bIsSRGB ? "sRGB" : "Linear");
-				Args.Add(TEXT("ArrayInfo"), ArrayInfo);
-				Args.Add(TEXT("Width"), FString::FromInt(Desc.Width));
-				Args.Add(TEXT("Height"), FString::FromInt(Desc.Height));
-				Details = FString::Format(TEXT("{Channels}_{Format}, {IsSRGB} {ArrayInfo}\n{Width}x{Height}"), Args);
 				break;
 			}
 		}
 	}
+
+	UTG_Expression* Expression = GetNode()->GetExpression();
+	if (Expression)
+	{
+		const BufferDescriptor Desc = Expression->BaseOutputSettings.ToBufferDescriptor();
+		FStringFormatNamedArguments Args;
+		Args.Add(TEXT("Channels"), TextureHelper::GetChannelsTextFromItemsPerPoint(Desc.ItemsPerPoint));
+		Args.Add(TEXT("Format"), Desc.FormatToString(Desc.Format));
+		Args.Add(TEXT("IsSRGB"), Desc.bIsSRGB ? "sRGB" : "Linear");
+		Args.Add(TEXT("ArrayInfo"), ArrayInfo);
+		Args.Add(TEXT("Width"), FString::FromInt(Desc.Width));
+		Args.Add(TEXT("Height"), FString::FromInt(Desc.Height));
+		Details = FString::Format(TEXT("{Channels}_{Format}, {IsSRGB} {ArrayInfo}\n{Width}x{Height}"), Args);
+	}
+
 	return Details;
 }
 
