Added srgb flag support to texture descriptor

#virtualized

diff --git a/Engine/Plugins/TextureGraph/Source/TextureGraph/Private/TG_Texture.cpp b/Engine/Plugins/TextureGraph/Source/TextureGraph/Private/TG_Texture.cpp
--- a/Engine/Plugins/TextureGraph/Source/TextureGraph/Private/TG_Texture.cpp
+++ b/Engine/Plugins/TextureGraph/Source/TextureGraph/Private/TG_Texture.cpp
@@ -6,6 +6,27 @@
 #include "TG_Var.h"
 #include "Model/StaticImageResource.h"
 
+FArchive& operator<<(FArchive& Ar, FTG_TextureDescriptor& D)
+{
+	int32 Version = Ar.CustomVer(FTG_CustomVersion::GUID);
+	Ar << D.Width;
+	Ar << D.Height;
+	Ar << D.TextureFormat;
+
+	if (Ar.IsLoading())
+	{
+		/// Older versions had this defaulted to false
+		if (Version < FTG_CustomVersion::TGTextureDescAdded_bSRGB)
+			D.bIsSRGB = false;
+		else
+			Ar << D.bIsSRGB;
+	}
+	else
+		Ar << D.bIsSRGB;
+
+	return Ar;
+}
+
 template <> FString TG_Var_LogValue(FTG_Texture& Value)
 {
 	if (!Value.RasterBlob)
diff --git a/Engine/Plugins/TextureGraph/Source/TextureGraph/Public/Expressions/Input/TG_Expression_Material.h b/Engine/Plugins/TextureGraph/Source/TextureGraph/Public/Expressions/Input/TG_Expression_Material.h
--- a/Engine/Plugins/TextureGraph/Source/TextureGraph/Public/Expressions/Input/TG_Expression_Material.h
+++ b/Engine/Plugins/TextureGraph/Source/TextureGraph/Public/Expressions/Input/TG_Expression_Material.h
@@ -25,7 +25,7 @@ public:
 #endif
 
 	// The input material to employ for rendering
-	UPROPERTY(EditAnywhere, BlueprintReadWrite, Setter, Category = NoCategory, meta = (TGType = "TG_Setting", TGPinNotConnectable))
+	UPROPERTY(EditAnywhere, BlueprintReadWrite, Setter, Category = NoCategory, meta = (TGType = "TG_Setting"))
 	TObjectPtr<UMaterialInterface> Material;
 	
 	void SetMaterial(UMaterialInterface* InMaterial);
diff --git a/Engine/Plugins/TextureGraph/Source/TextureGraph/Public/TG_CustomVersion.h b/Engine/Plugins/TextureGraph/Source/TextureGraph/Public/TG_CustomVersion.h
--- a/Engine/Plugins/TextureGraph/Source/TextureGraph/Public/TG_CustomVersion.h
+++ b/Engine/Plugins/TextureGraph/Source/TextureGraph/Public/TG_CustomVersion.h
@@ -18,6 +18,9 @@ struct FTG_CustomVersion
 		// Added BaseOutputSettings to TG_Expression
 		TGExpressionAddedBaseOutputSettings,
 
+		// Added bSRGB to FTG_TextureDescriptor
+		TGTextureDescAdded_bSRGB,
+
 		// -----<new versions can be added above this line>-------------------------------------------------
 		VersionPlusOne,
 		LatestVersion = VersionPlusOne - 1
diff --git a/Engine/Plugins/TextureGraph/Source/TextureGraph/Public/TG_Texture.h b/Engine/Plugins/TextureGraph/Source/TextureGraph/Public/TG_Texture.h
--- a/Engine/Plugins/TextureGraph/Source/TextureGraph/Public/TG_Texture.h
+++ b/Engine/Plugins/TextureGraph/Source/TextureGraph/Public/TG_Texture.h
@@ -22,6 +22,7 @@ struct TEXTUREGRAPH_API FTG_TextureDescriptor
 		Width = static_cast<EResolution>(Desc.Width);
 		Height = static_cast<EResolution>(Desc.Height);
 		TextureFormat = TextureHelper::GetTGTextureFormatFromChannelsAndFormat(Desc.ItemsPerPoint, Desc.Format);
+		bIsSRGB = Desc.bIsSRGB;
 	}
 
 public:
@@ -47,30 +48,30 @@ PRAGMA_ENABLE_DEPRECATION_WARNINGS
 	UPROPERTY(BlueprintReadWrite, EditAnywhere, Category = "Advanced", DisplayName = "Texture Format", Meta = (NoResetToDefault))
 		ETG_TextureFormat TextureFormat = ETG_TextureFormat::Auto;
 
+	UPROPERTY(BlueprintReadWrite, EditAnywhere, Category = "Advanced", DisplayName = "sRGB", Meta = (NoResetToDefault))
+		bool bIsSRGB = false;
+
 	bool operator==(const FTG_TextureDescriptor& RHS) const
 	{
-		return Width == RHS.Width && Height == RHS.Height && TextureFormat == RHS.TextureFormat;
+		return Width == RHS.Width && Height == RHS.Height && TextureFormat == RHS.TextureFormat && bIsSRGB == RHS.bIsSRGB;
 	}
 
-	friend FArchive& operator<<(FArchive& Ar, FTG_TextureDescriptor& D)
-	{
- 		Ar << D.Width;
-		Ar << D.Height;
-		Ar << D.TextureFormat;
-		return Ar;
-	}
+	friend FArchive& operator<<(FArchive& Ar, FTG_TextureDescriptor& D);
 
 	FORCEINLINE BufferDescriptor ToBufferDescriptor() const
 	{
 		uint32 NumChannels = 0;
 		BufferFormat Format = BufferFormat::Auto;
 		TextureHelper::GetBufferFormatAndChannelsFromTGTextureFormat(TextureFormat, Format, NumChannels);
-		return BufferDescriptor(
+		BufferDescriptor Desc = BufferDescriptor(
 			static_cast<int32>(Width),
 			static_cast<int32>(Height),
 			NumChannels,
 			Format
 		);
+
+		Desc.bIsSRGB = bIsSRGB;
+		return Desc;
 	}
 
 	FORCEINLINE operator BufferDescriptor() const { return ToBufferDescriptor(); }
diff --git a/Engine/Plugins/TextureGraph/Source/TextureGraphEditor/Private/EdGraph/TG_EdGraphNode.cpp b/Engine/Plugins/TextureGraph/Source/TextureGraphEditor/Private/EdGraph/TG_EdGraphNode.cpp
--- a/Engine/Plugins/TextureGraph/Source/TextureGraphEditor/Private/EdGraph/TG_EdGraphNode.cpp
+++ b/Engine/Plugins/TextureGraph/Source/TextureGraphEditor/Private/EdGraph/TG_EdGraphNode.cpp
@@ -394,10 +394,9 @@ FString UTG_EdGraphNode::GetTitleDetail()
 	const UTG_EdGraphSchema* Schema = Cast<const UTG_EdGraphSchema>(GetSchema());
 	TArray<const UTG_Pin*> OutPins;
 	GetNode()->GetOutputPins(OutPins);
-	
-	const UEdGraphPin* SelectedPinPtr = GetSelectedPin();
+	UTG_Expression* Expression = GetNode()->GetExpression();
 
-	FString ArrayInfo;
+	const UEdGraphPin* SelectedPinPtr = GetSelectedPin();
 
 	for (auto Pin : OutPins)
 	{
@@ -414,6 +413,8 @@ FString UTG_EdGraphNode::GetTitleDetail()
 		{
 			if (Texture.RasterBlob)
 			{
+				FString ArrayInfo;
+
 				if (Pin->IsArgTextureArray())
 				{
 					FTG_VariantArray VarArray;
@@ -421,25 +422,29 @@ FString UTG_EdGraphNode::GetTitleDetail()
 					ArrayInfo = FString::Printf(TEXT("[%d]"), VarArray.Num());
 				}
 
+				const auto Desc = Texture.RasterBlob->GetDescriptor();
+				FStringFormatNamedArguments Args;
+
+				bool bIsSRGB = Desc.bIsSRGB;
+
+				FTG_TextureDescriptor OutputSettings;
+
+				if (Expression)
+				{
+					bIsSRGB = Expression->BaseOutputSettings.bIsSRGB;
+				}
+
+				Args.Add(TEXT("Channels"), TextureHelper::GetChannelsTextFromItemsPerPoint(Desc.ItemsPerPoint));
+				Args.Add(TEXT("Format"), Desc.FormatToString(Desc.Format));
+				Args.Add(TEXT("IsSRGB"), bIsSRGB ? "sRGB" : "Linear");
+				Args.Add(TEXT("ArrayInfo"), ArrayInfo);
+				Args.Add(TEXT("Width"), FString::FromInt(Desc.Width));
+				Args.Add(TEXT("Height"), FString::FromInt(Desc.Height));
+				Details = FString::Format(TEXT("{Channels}_{Format}, {IsSRGB} {ArrayInfo}\n{Width}x{Height}"), Args);
 				break;
 			}
 		}
 	}
-
-	UTG_Expression* Expression = GetNode()->GetExpression();
-	if (Expression)
-	{
-		const BufferDescriptor Desc = Expression->BaseOutputSettings.ToBufferDescriptor();
-		FStringFormatNamedArguments Args;
-		Args.Add(TEXT("Channels"), TextureHelper::GetChannelsTextFromItemsPerPoint(Desc.ItemsPerPoint));
-		Args.Add(TEXT("Format"), Desc.FormatToString(Desc.Format));
-		Args.Add(TEXT("IsSRGB"), Desc.bIsSRGB ? "sRGB" : "Linear");
-		Args.Add(TEXT("ArrayInfo"), ArrayInfo);
-		Args.Add(TEXT("Width"), FString::FromInt(Desc.Width));
-		Args.Add(TEXT("Height"), FString::FromInt(Desc.Height));
-		Details = FString::Format(TEXT("{Channels}_{Format}, {IsSRGB} {ArrayInfo}\n{Width}x{Height}"), Args);
-	}
-
 	return Details;
 }
 
diff --git a/Engine/Plugins/TextureGraph/Source/TextureGraphEditor/Private/Pins/STG_GraphPinTextureDescriptorWidget.cpp b/Engine/Plugins/TextureGraph/Source/TextureGraphEditor/Private/Pins/STG_GraphPinTextureDescriptorWidget.cpp
--- a/Engine/Plugins/TextureGraph/Source/TextureGraphEditor/Private/Pins/STG_GraphPinTextureDescriptorWidget.cpp
+++ b/Engine/Plugins/TextureGraph/Source/TextureGraphEditor/Private/Pins/STG_GraphPinTextureDescriptorWidget.cpp
@@ -11,6 +11,7 @@
 #include "Editor.h"
 #include "Widgets/Layout/SWrapBox.h"
 #include "SGraphPinComboBox.h"
+#include "Widgets/Layout/SSeparator.h"
 #include "Framework/MultiBox/MultiBoxBuilder.h"
 
 #define LOCTEXT_NAMESPACE "STG_GraphPinTextureDescriptorWidget"
@@ -69,6 +70,16 @@ void STG_GraphPinTextureDescriptorWidget::Construct(const FArguments& InArgs, UE
 				.Visibility(this, &STG_GraphPinTextureDescriptorWidget::ShowParameters)
 				[
 					SNew(SVerticalBox)
+						//Separator
+						+SVerticalBox::Slot()
+						.AutoHeight()
+						.HAlign(HAlign_Fill)
+						.VAlign(VAlign_Center)
+						[
+							SNew(SSeparator)
+								.Thickness(2)
+						]
+
 						//Path Width
 						+SVerticalBox::Slot()
 						.Padding(UniformPadding)
@@ -89,6 +100,23 @@ void STG_GraphPinTextureDescriptorWidget::Construct(const FArguments& InArgs, UE
 						[
 							AddEnumComobox(LOCTEXT("OutputFormat", "Format"), GetFormatDelegate, OnGenerateFormatMenu)
 						]
+
+						//SRGB
+						+ SVerticalBox::Slot()
+						.Padding(UniformPadding)
+						[
+							AddSRGBWidget()
+						]		
+
+						//Separator
+						+SVerticalBox::Slot()
+						.AutoHeight()
+						.HAlign(HAlign_Fill)
+						.VAlign(VAlign_Center)
+						[
+							SNew(SSeparator)
+								.Thickness(2)
+						]
 				]
 		];
 }
@@ -139,6 +167,51 @@ TSharedRef<SWidget> STG_GraphPinTextureDescriptorWidget::AddEnumComobox(FText La
 		];
 }
 
+TSharedRef<SWidget> STG_GraphPinTextureDescriptorWidget::AddSRGBWidget()
+{
+	return 
+		SNew(SHorizontalBox)
+
+		+ SHorizontalBox::Slot()
+		.VAlign(VAlign_Center)
+		.HAlign(HAlign_Center)
+		.MinSize(LabelSize)
+		.FillWidth(1.0f)
+		[
+			SNew(SBox)
+				.MinDesiredWidth(LabelSize)
+				.MaxDesiredWidth(LabelSize)
+				.HAlign(HAlign_Left)
+				[
+					SNew(STextBlock)
+						.Text(FText::FromString("sRGB"))
+						.TextStyle(FAppStyle::Get(), TEXT("Graph.Node.PinName"))
+				]
+		]
+
+		+ SHorizontalBox::Slot()
+		.MaxWidth(LabelSize * 3)
+		[
+			SNew(SCheckBox)
+				.IsChecked(this, &STG_GraphPinTextureDescriptorWidget::HandleSRGBIsChecked)
+				.OnCheckStateChanged(this, &STG_GraphPinTextureDescriptorWidget::HandleSRGBExecute)
+		];
+}
+
+ECheckBoxState STG_GraphPinTextureDescriptorWidget::HandleSRGBIsChecked() const
+{
+	auto Settings = GetSettings();
+	return Settings.bIsSRGB ? ECheckBoxState::Checked : ECheckBoxState::Unchecked;
+}
+
+void STG_GraphPinTextureDescriptorWidget::HandleSRGBExecute(ECheckBoxState InNewState)
+{
+	bSRGB = InNewState == ECheckBoxState::Checked;
+	auto Settings = GetSettings();
+	Settings.bIsSRGB = bSRGB;
+	OnTextureDescriptorChanged.ExecuteIfBound(Settings);
+}
+
 FTG_TextureDescriptor STG_GraphPinTextureDescriptorWidget::GetSettings() const
 {
 	FString TextureDescriptorString = GraphPinObj->GetDefaultAsString();
diff --git a/Engine/Plugins/TextureGraph/Source/TextureGraphEditor/Private/Pins/STG_GraphPinTextureDescriptorWidget.h b/Engine/Plugins/TextureGraph/Source/TextureGraphEditor/Private/Pins/STG_GraphPinTextureDescriptorWidget.h
--- a/Engine/Plugins/TextureGraph/Source/TextureGraphEditor/Private/Pins/STG_GraphPinTextureDescriptorWidget.h
+++ b/Engine/Plugins/TextureGraph/Source/TextureGraphEditor/Private/Pins/STG_GraphPinTextureDescriptorWidget.h
@@ -52,6 +52,7 @@ public:
 
 	void Construct(const FArguments& InArgs, UEdGraphPin* InGraphPinObj);
 	TSharedRef<SWidget> AddEnumComobox(FText Label, FGetTextDelegate GetText, FGenerateEnumMenu OnGenerateEnumMenu);
+	TSharedRef<SWidget> AddSRGBWidget();
 	FTG_TextureDescriptor GetSettings() const;
 	void GenerateStringsFromEnum(TArray<FString>& OutEnumNames, const FString& EnumPathName);
 
@@ -81,6 +82,7 @@ private:
 	int SelectedWidthIndex;
 	int SelectedHeightIndex;
 	int SelectedFormatIndex;
+	bool bSRGB = false;
 
 	FGenerateEnumMenu OnGenerateWidthMenu;
 	FGetTextDelegate GetWidthDelegate;
@@ -101,5 +103,8 @@ private:
 	void HandleFormatChanged(FString OuputName, int Index);
 	FText HandleFormatText() const;
 
+	ECheckBoxState HandleSRGBIsChecked() const;
+	void HandleSRGBExecute(ECheckBoxState InNewState);
+
 	const int LabelSize = 75;
 };
diff --git a/Engine/Plugins/TextureGraph/Source/TextureGraphEngine/Job/Job.cpp b/Engine/Plugins/TextureGraph/Source/TextureGraphEngine/Job/Job.cpp
--- a/Engine/Plugins/TextureGraph/Source/TextureGraphEngine/Job/Job.cpp
+++ b/Engine/Plugins/TextureGraph/Source/TextureGraphEngine/Job/Job.cpp
@@ -261,6 +261,7 @@ CHashPtr Job::CalcTileHash(int32 RowId, int32 ColId) const
 		ArgHashes.push_back(std::make_shared<CHash>(DataUtil::Hash_Int32(ResultDesc.Height), true));
 		ArgHashes.push_back(std::make_shared<CHash>(DataUtil::Hash_Int32(ResultDesc.ItemsPerPoint), true));
 		ArgHashes.push_back(std::make_shared<CHash>(DataUtil::Hash_Int32((int32)ResultDesc.Format), true));
+		ArgHashes.push_back(std::make_shared<CHash>(DataUtil::Hash_Int32((int32)ResultDesc.bIsSRGB), true));
 		ArgHashes.push_back(std::make_shared<CHash>(DataUtil::Hash_Int32(NumRows), true));
 		ArgHashes.push_back(std::make_shared<CHash>(DataUtil::Hash_Int32(NumCols), true));
 	}
