Change erode/dilate options to TG_Setting

diff --git a/Engine/Plugins/TextureGraph/Source/TextureGraph/Public/Expressions/Filter/TG_Expression_ErodeDilate.h b/Engine/Plugins/TextureGraph/Source/TextureGraph/Public/Expressions/Filter/TG_Expression_ErodeDilate.h
--- a/Engine/Plugins/TextureGraph/Source/TextureGraph/Public/Expressions/Filter/TG_Expression_ErodeDilate.h
+++ b/Engine/Plugins/TextureGraph/Source/TextureGraph/Public/Expressions/Filter/TG_Expression_ErodeDilate.h
@@ -56,7 +56,7 @@ public:
 	//float								MinBrightness = 0.0f;
 
 	/// The size of the kernel to use
-	UPROPERTY(meta = (TGType = "TG_Input", ClampMin = 0, ClampMax = 64, UIMin = 0, UIMax = 64))
+	UPROPERTY(meta = (TGType = "TG_Setting", ClampMin = 0, ClampMax = 64, UIMin = 0, UIMax = 64))
 	int32								Size = 2;
 
 	/// The type of the kernel to use
