Allow custom texture descriptor

diff --git a/Engine/Plugins/TextureGraph/Source/TextureGraph/Private/Expressions/Input/TG_Expression_TextureDescriptor.cpp b/Engine/Plugins/TextureGraph/Source/TextureGraph/Private/Expressions/Input/TG_Expression_TextureDescriptor.cpp
new file mode 100644
--- /dev/null
+++ b/Engine/Plugins/TextureGraph/Source/TextureGraph/Private/Expressions/Input/TG_Expression_TextureDescriptor.cpp
@@ -0,0 +1,40 @@
+// Copyright Epic Games, Inc. All Rights Reserved.
+
+#include "Expressions/Input/TG_Expression_TextureDescriptor.h"
+
+void UTG_Expression_TextureDescriptor::Evaluate(FTG_EvaluationContext* InContext)
+{
+	Super::Evaluate(InContext);
+
+	FTextureGraphErrorReporter* ErrorReporter = TextureGraphEngine::GetErrorReporter(InContext->Cycle->GetMix());
+	static const int32 MinResolution = 8;
+	static const int32 MaxResolution = (int32)EResolution::Resolution8192;
+	static const int32 Auto = (int32)EResolution::Auto;
+
+	if (Width < MinResolution || Width > MaxResolution)
+	{
+		if (ErrorReporter)
+		{
+			ErrorReporter->ReportWarning((int32)ETextureGraphErrorType::NODE_WARNING, FString::Printf(TEXT("Invalid width specified: %d (Range: %d - %d). Using Auto instead."), 
+				Width, MinResolution, MaxResolution));
+		}
+		Width = Auto;
+	}
+
+	if (Height < MinResolution || Height > MaxResolution)
+	{
+		if (ErrorReporter)
+		{
+			ErrorReporter->ReportWarning((int32)ETextureGraphErrorType::NODE_WARNING, FString::Printf(TEXT("Invalid height specified: %d (Range: %d - %d). Using Auto instead."), 
+				Height, MinResolution, MaxResolution));
+		}
+		Height = Auto;
+	}
+
+	// The Value is updated either as an input or as a setting and then becomes the output for this expression
+	// The pin out is named "ValueOut"
+	Output.Width = (EResolution)Width;
+	Output.Height = (EResolution)Height;
+	Output.bIsSRGB = bIsSRGB;
+	Output.TextureFormat = Format;
+}
diff --git a/Engine/Plugins/TextureGraph/Source/TextureGraph/Public/Expressions/Input/TG_Expression_TextureDescriptor.h b/Engine/Plugins/TextureGraph/Source/TextureGraph/Public/Expressions/Input/TG_Expression_TextureDescriptor.h
new file mode 100644
--- /dev/null
+++ b/Engine/Plugins/TextureGraph/Source/TextureGraph/Public/Expressions/Input/TG_Expression_TextureDescriptor.h
@@ -0,0 +1,42 @@
+// Copyright Epic Games, Inc. All Rights Reserved.
+
+#pragma once
+
+#include "CoreMinimal.h"
+#include "TG_Expression_InputParam.h"
+#include "TG_Expression_TextureDescriptor.generated.h"
+
+UCLASS()
+class TEXTUREGRAPH_API UTG_Expression_TextureDescriptor : public UTG_Expression_InputParam
+{
+	GENERATED_BODY()
+	TG_DECLARE_INPUT_PARAM_EXPRESSION(TG_Category::Input);
+
+public:
+
+	virtual void Evaluate(FTG_EvaluationContext* InContext) override;
+
+	/// The width of the texture. 0 and negative values mean "Auto"
+	UPROPERTY(EditAnywhere, Category = NoCategory, meta = (TGType = "TG_Input"))
+    int32 Width = (int32)EResolution::Resolution2048;
+    
+	/// The height of the texture. 0 and negative values mean "Auto"
+	UPROPERTY(EditAnywhere, Category = NoCategory, meta = (TGType = "TG_Input"))
+    int32 Height = (int32)EResolution::Resolution2048;
+    
+	/// The height of the texture. 0 and negative values mean "Auto"
+	UPROPERTY(EditAnywhere, Category = NoCategory, meta = (TGType = "TG_Input", PinDisplayName = "sRGB"))
+    bool bIsSRGB = false;
+    
+	/// The height of the texture. 0 and negative values mean "Auto"
+	UPROPERTY(EditAnywhere, Category = NoCategory, meta = (TGType = "TG_Setting"))
+    ETG_TextureFormat Format = ETG_TextureFormat::Auto;
+    
+	// The output of the node, which is the color value
+    UPROPERTY(meta = (TGType = "TG_Output", PinDisplayName = ""))
+	FTG_TextureDescriptor Output;
+	
+	virtual FTG_Name GetDefaultName() const override { return TEXT("Texture Settings"); }
+	virtual FText GetTooltipText() const override { return FText::FromString(TEXT("Allows the user to customize texture settings like resolution, format etc.")); } 
+};
+
