Single channel erode/dilate support

diff --git a/Engine/Plugins/TextureGraph/Shaders/Expressions/Expression_ErodeDilate.usf b/Engine/Plugins/TextureGraph/Shaders/Expressions/Expression_ErodeDilate.usf
--- a/Engine/Plugins/TextureGraph/Shaders/Expressions/Expression_ErodeDilate.usf
+++ b/Engine/Plugins/TextureGraph/Shaders/Expressions/Expression_ErodeDilate.usf
@@ -36,7 +36,13 @@ float4 FSH_ErodeDilate(float2 InUV : TEXCOORD0) : SV_Target0
 	float2 UV = TileInfo_fromCurrentTileToLayer(InUV);
 	
 	float4 OrgClr = Input.Sample(SamplerStates_Linear_Clamp, UV);
+	
+#if ED_SINGLECHANNEL 
+	float3 Clr = OrgClr.rrr;
+#else
 	float3 Clr = OrgClr.rgb;
+#endif 
+	
 	float3 TmpClr = Clr;
 	float Brightness = LumaFromRGB(Clr);
 
@@ -58,7 +64,11 @@ float4 FSH_ErodeDilate(float2 InUV : TEXCOORD0) : SV_Target0
 #endif
 
 			float2 OffsetUV = UV + float2(Xf, Yf);
+#if ED_SINGLECHANNEL 
+			float3 NClr = Input.Sample(SamplerStates_Linear_Clamp, OffsetUV).rrr;
+#else
 			float3 NClr = Input.Sample(SamplerStates_Linear_Clamp, OffsetUV).rgb;
+#endif 
 			float NBrightness = LumaFromRGB(NClr);
         
 #if ED_TYPE == ED_DILATE
diff --git a/Engine/Plugins/TextureGraph/Source/TextureGraph/Private/Expressions/FIlter/TG_Expression_ErodeDilate.cpp b/Engine/Plugins/TextureGraph/Source/TextureGraph/Private/Expressions/FIlter/TG_Expression_ErodeDilate.cpp
--- a/Engine/Plugins/TextureGraph/Source/TextureGraph/Private/Expressions/FIlter/TG_Expression_ErodeDilate.cpp
+++ b/Engine/Plugins/TextureGraph/Source/TextureGraph/Private/Expressions/FIlter/TG_Expression_ErodeDilate.cpp
@@ -26,6 +26,8 @@ void UTG_Expression_ErodeDilate::Evaluate(FTG_EvaluationContext* InContext)
 	FSH_ErodeDilate::FPermutationDomain PermutationVector;
 	PermutationVector.Set<FVar_ErodeDilate_Type>((int32)Type);
 	PermutationVector.Set<FVar_ErodeDilate_Kernel>((int32)Kernel);
+	bool bIsSingleChannel = Input->GetDescriptor().ItemsPerPoint == 1;
+	PermutationVector.Set<FVar_ErodeDilate_IsSingleChannel>(bIsSingleChannel);
 
 	RenderMaterial_FXPtr RenderMaterial = TextureGraphEngine::GetMaterialManager()->CreateMaterial_FX<VSH_Simple, FSH_ErodeDilate>(TEXT("T_ErodeDilate"), PermutationVector);
 	check(RenderMaterial);
diff --git a/Engine/Plugins/TextureGraph/Source/TextureGraphEngine/Transform/Expressions/T_Filter.h b/Engine/Plugins/TextureGraph/Source/TextureGraphEngine/Transform/Expressions/T_Filter.h
--- a/Engine/Plugins/TextureGraph/Source/TextureGraphEngine/Transform/Expressions/T_Filter.h
+++ b/Engine/Plugins/TextureGraph/Source/TextureGraphEngine/Transform/Expressions/T_Filter.h
@@ -109,6 +109,7 @@ public:
 
 class FVar_ErodeDilate_Type : SHADER_PERMUTATION_INT("ED_TYPE", 2);
 class FVar_ErodeDilate_Kernel : SHADER_PERMUTATION_INT("ED_KERNEL", 3);
+class FVar_ErodeDilate_IsSingleChannel : SHADER_PERMUTATION_BOOL("ED_SINGLECHANNEL");
 //////////////////////////////////////////////////////////////////////////
 /// Erode/Dilate
 //////////////////////////////////////////////////////////////////////////
@@ -124,7 +125,7 @@ public:
 		SHADER_PARAMETER(int, Size)
 	END_SHADER_PARAMETER_STRUCT()
 
-	using FPermutationDomain = TShaderPermutationDomain<FVar_ErodeDilate_Type, FVar_ErodeDilate_Kernel>;
+	using FPermutationDomain = TShaderPermutationDomain<FVar_ErodeDilate_Type, FVar_ErodeDilate_Kernel, FVar_ErodeDilate_IsSingleChannel>;
 public:
 	TEXTURE_ENGINE_DEFAULT_PERMUTATION;
 	TEXTUREGRAPH_ENGINE_DEFAULT_COMPILATION_ENV;
