Resize node erroneous width/height on input image check

diff --git a/Engine/Plugins/TextureGraph/Source/TextureGraph/Private/Expressions/Utilities/TG_Expression_Utils.cpp b/Engine/Plugins/TextureGraph/Source/TextureGraph/Private/Expressions/Utilities/TG_Expression_Utils.cpp
--- a/Engine/Plugins/TextureGraph/Source/TextureGraph/Private/Expressions/Utilities/TG_Expression_Utils.cpp
+++ b/Engine/Plugins/TextureGraph/Source/TextureGraph/Private/Expressions/Utilities/TG_Expression_Utils.cpp
@@ -64,7 +64,7 @@ void UTG_Expression_Utils_MakeVector4::Evaluate(FTG_EvaluationContext* InContext
 //////////////////////////////////////////////////////////
 void UTG_Expression_Utils_Resize::Evaluate(FTG_EvaluationContext* InContext)
 {
-	if (!Input || Input->GetWidth() < 8 || Input->GetHeight() < 8)
+	if (!Input)
 	{
 		Output = FTG_Texture::GetBlack();
 		return;
