TG: Exposing Material as parameters
- Introducing the TG_Material which is a simple wrapper struct around a FSoftObjectPath and conform it to a UMaterialInterface*
- Deprecate the UMaterialInterface* Material Uproperty in FTG_Expression_Material, and replace it with a new property "FTG_Material InputMaterial"
- Introduced the UI/widget support for the FTG_Material type

diff --git a/Engine/Plugins/TextureGraph/Source/TextureGraph/Private/Expressions/Input/TG_Expression_Material.cpp b/Engine/Plugins/TextureGraph/Source/TextureGraph/Private/Expressions/Input/TG_Expression_Material.cpp
--- a/Engine/Plugins/TextureGraph/Source/TextureGraph/Private/Expressions/Input/TG_Expression_Material.cpp
+++ b/Engine/Plugins/TextureGraph/Source/TextureGraph/Private/Expressions/Input/TG_Expression_Material.cpp
@@ -7,11 +7,12 @@
 #if WITH_EDITOR
 void UTG_Expression_Material::PostEditChangeProperty(FPropertyChangedEvent& PropertyChangedEvent)
 {
-	// First catch if Material changes
-	if (PropertyChangedEvent.GetPropertyName() == GET_MEMBER_NAME_CHECKED(UTG_Expression_Material, Material))
+	// First catch if InputMaterial changes, specifically the <AssetPath> field
+	if (PropertyChangedEvent.GetPropertyName() == GET_MEMBER_NAME_CHECKED(UTG_Expression_Material, InputMaterial)
+		|| PropertyChangedEvent.GetPropertyName() == TEXT("AssetPath"))
 	{
 		UE_LOG(LogTextureGraph, VeryVerbose, TEXT("Material Expression PostEditChangeProperty."));
-		SetMaterialInternal(Material);
+		SetMaterialInternal(InputMaterial.GetMaterial());
 		FeedbackPinValue(GET_MEMBER_NAME_CHECKED(UTG_Expression_Material, RenderedAttribute), RenderedAttribute);
 	}
 	// Second catch if AttributeName changes
@@ -39,15 +40,16 @@ void UTG_Expression_Material::OnReferencedObjectPreSave(UObject* Object, FObject
 {
 	// every editor should check if your texture graph is dependent on the object being saved
 	UMaterialInterface* MaterialBeingSaved = Cast<UMaterialInterface>(Object);
+	UMaterialInterface* ReferencedMaterial = InputMaterial.GetMaterial();
+
 	// if object being saved is our linked TextureGraph, we re-create
-	if (MaterialBeingSaved && Material && (MaterialBeingSaved == Material))
+	if (MaterialBeingSaved && ReferencedMaterial && (MaterialBeingSaved == ReferencedMaterial))
 	{
-
 		SetMaterialInternal(MaterialBeingSaved);
 	}
 }
+#endif // WITH_EDITOR
 
-#endif
 
 UTG_Expression_Material::UTG_Expression_Material()
 {
@@ -63,12 +65,30 @@ UTG_Expression_Material::~UTG_Expression_Material()
 #endif
 }
 
+void UTG_Expression_Material::PostLoad()
+{
+	Super::PostLoad();
+
+	// Only do this during editor time because that is when we can fix it up.
+	// This means that properties will be fixed up during a cook!
+#if WITH_EDITORONLY_DATA
+	PRAGMA_DISABLE_DEPRECATION_WARNINGS
+
+		if (Material_DEPRECATED)
+		{
+			InputMaterial.SetMaterial(Material_DEPRECATED);
+			Material_DEPRECATED = nullptr;
+		}
+	PRAGMA_ENABLE_DEPRECATION_WARNINGS
+#endif // WITH_EDITORONLY_DATA
+}
+
 void UTG_Expression_Material::Initialize()
 {
-	// If the referenced material is valid, then we need to recreate a dubplicate
-	if (Material)
+	// If the referenced material is valid, then we need to recreate a duplicate
+ 	if (InputMaterial.IsValid())
 	{
-		MaterialCopy = DuplicateObject(Material, this);
+		MaterialCopy = DuplicateObject(InputMaterial.GetMaterial(), this);
 	}
 
 	Super::Initialize();
@@ -79,19 +99,15 @@ void UTG_Expression_Material::SetMaterialInternal(UMaterialInterface* InMaterial
 {
 	if (!InMaterial)
 	{
-		Material = nullptr;
 		MaterialCopy = nullptr;
 	}
 	else if (InMaterial->IsA<UMaterial>())
 	{
-		// create a material instance
-		Material = InMaterial;
-		MaterialCopy = DuplicateObject(Material, this);
+		MaterialCopy = DuplicateObject(InMaterial, this);
 	}
 	else if (InMaterial->IsA<UMaterialInstance>())
 	{
-		Material = InMaterial;
-		MaterialCopy = DuplicateObject(Material, this);
+		MaterialCopy = DuplicateObject(InMaterial, this);
 	}
 
 	Super::SetMaterialInternal(MaterialCopy);
@@ -99,22 +115,22 @@ void UTG_Expression_Material::SetMaterialInternal(UMaterialInterface* InMaterial
 	SetRenderedAttribute(RenderedAttribute);
 }
 
-
-void UTG_Expression_Material::SetMaterial(UMaterialInterface* InMaterial)
+void UTG_Expression_Material::SetInputMaterial(const FTG_Material& InParamMaterial)
 {
 	// This is the public setter of the material, 
 	// This is NOT called if the Material is modified from the detail panel!!!
-	// We catch that case in PostEditChangeProperty, which will call SetMaterialInternal 
+	// We catch that case in PostEditChangeProperty, which will call UpdateMaterialInternalFromSettingOrParam 
 
-	// If it is the same Material then avoid anymore work, we shoudl be good to go
-	if (InMaterial == Material)
+	// If it is the same ParamMaterial then avoid anymore work, we shoudl be good to go
+	if (InputMaterial == InParamMaterial)
 	{
-		// Just check that the MaterialInstance is valid, if not reassign below
-		if ((!Material) || (Material && MaterialCopy && MaterialPermutations))
+		// Just check that the local MaterialCopy is valid, if not reassign below
+		if ((!InputMaterial.IsValid()) || (InputMaterial.IsValid() && MaterialCopy && MaterialPermutations))
 			return;
 	}
 
-	SetMaterialInternal(InMaterial);
+	InputMaterial = InParamMaterial;
+	SetMaterialInternal(InputMaterial.GetMaterial());
 }
 
 void UTG_Expression_Material::SetRenderedAttribute(FName InRenderedAttribute)
@@ -149,16 +165,14 @@ void UTG_Expression_Material::SetAsset(UObject* Asset)
 {
 	if(UMaterialInterface* MaterialAsset = Cast<UMaterialInterface>(Asset); MaterialAsset != nullptr)
 	{
-		SetMaterial(MaterialAsset);
+		InputMaterial.SetMaterial(MaterialAsset);
 #if WITH_EDITOR
 		// We need to find its property and trigger property change event manually.
-		const auto SourcePin = GetParentNode()->GetInputPin("Material");
-
-		check(SourcePin)
-	
+		const auto SourcePin = GetParentNode()->GetInputPin(GET_MEMBER_NAME_CHECKED(UTG_Expression_Material, InputMaterial));
+		check(SourcePin)	
 		if(SourcePin)
 		{
-			auto Property = SourcePin->GetExpressionProperty();
+			FProperty* Property = SourcePin->GetExpressionProperty();
 			PropertyChangeTriggered(Property, EPropertyChangeType::ValueSet);
 		}
 #endif
@@ -167,7 +181,8 @@ void UTG_Expression_Material::SetAsset(UObject* Asset)
 
 void UTG_Expression_Material::SetTitleName(FName NewName)
 {
-	TitleName = NewName;
+	GetParentNode()->GetInputPin("InputMaterial")->SetAliasName(NewName);
+	TitleName = GetParentNode()->GetInputPin("InputMaterial")->GetAliasName();
 }
 
 FName UTG_Expression_Material::GetTitleName() const
diff --git a/Engine/Plugins/TextureGraph/Source/TextureGraph/Private/TG_Graph.cpp b/Engine/Plugins/TextureGraph/Source/TextureGraph/Private/TG_Graph.cpp
--- a/Engine/Plugins/TextureGraph/Source/TextureGraph/Private/TG_Graph.cpp
+++ b/Engine/Plugins/TextureGraph/Source/TextureGraph/Private/TG_Graph.cpp
@@ -195,6 +195,9 @@ void UTG_Graph::RegenerateNode(UTG_Node* InNode)
 			if (ReusedPin->IsParam())
 				Params.Remove(ReusedPin->GetAliasName());
 
+			// Make sure the reused pin self var is correctly reset to whatever default value is coming from the Expression instance
+			ReusedPin->ResetSelfVarDefault();
+
 			// the pin need to exists and be remaped to a different id
 			if (ReusedPin && (i != OldIdx))
 			{
@@ -207,6 +210,8 @@ void UTG_Graph::RegenerateNode(UTG_Node* InNode)
 					GetPin(Edge)->RemapEdgeId(FTG_Id(InNodeIdx, OldIdx), NewId);
 				}
 			}
+			
+			
 			InNode->Pins[OldIdx] = nullptr;
 			NewPinArray.Add(ReusedPin);
 		}
@@ -277,6 +282,10 @@ void UTG_Graph::RegenerateNode(UTG_Node* InNode)
 	{
 		InNode->GetExpression()->ResetCommonInputVariantType();
 	}
+
+	// Make sure the node is finally renamed correctly, and potentially any pin alias name
+	// which is what we use for input param pin/node
+	InNode->GetExpression()->SetTitleName(InNode->GetExpression()->GetTitleName());
 }
 
 void UTG_Graph::AllocateNodePins(UTG_Node* Node)
@@ -1170,7 +1179,7 @@ void UTG_Graph::Log() const
 	LogMessage.ParseIntoArray(Lines, TEXT("\n"));
 	for (const FString& Line : Lines)
 	{
-		UE_LOG(LogTextureGraph, VeryVerbose, TEXT("%s"), *Line);
+		UE_LOG(LogTextureGraph, Display, TEXT("%s"), *Line);
 	}
 }
 
diff --git a/Engine/Plugins/TextureGraph/Source/TextureGraph/Private/TG_Material.cpp b/Engine/Plugins/TextureGraph/Source/TextureGraph/Private/TG_Material.cpp
new file mode 100644
--- /dev/null
+++ b/Engine/Plugins/TextureGraph/Source/TextureGraph/Private/TG_Material.cpp
@@ -0,0 +1,53 @@
+// Copyright Epic Games, Inc. All Rights Reserved.
+
+#include "TG_Material.h"
+
+#include "TG_CustomVersion.h"
+#include "TG_Var.h"
+
+template <> FString TG_Var_LogValue(FTG_Material& Value)
+{
+	return FString::Printf(TEXT("FTG_Material %s"), (*Value.AssetPath.ToString()));
+}
+
+bool FTG_Material::IsValid() const
+{
+	return AssetPath.IsValid();
+}
+
+UMaterialInterface* FTG_Material::GetMaterial() const
+{
+	return Cast<UMaterialInterface>(AssetPath.TryLoad());
+}
+
+void FTG_Material::SetMaterial(UMaterialInterface* InMaterial)
+{
+	if (InMaterial)
+	{
+		AssetPath = FSoftObjectPath(InMaterial);
+	}
+	else
+	{
+		AssetPath.Reset();
+	}
+}
+
+void FTG_Material::ResetTexturePath()
+{
+	AssetPath.Reset();
+}
+
+FArchive& operator<<(FArchive& Ar, FTG_Material& T)
+{
+	Ar.UsingCustomVersion(FTG_CustomVersion::GUID);
+
+	Ar << T.AssetPath;
+
+	return Ar;
+}
+
+bool FTG_Material::Serialize(FArchive& Ar)
+{
+	Ar << *this;
+	return true;
+}
diff --git a/Engine/Plugins/TextureGraph/Source/TextureGraph/Private/TG_Pin.cpp b/Engine/Plugins/TextureGraph/Source/TextureGraph/Private/TG_Pin.cpp
--- a/Engine/Plugins/TextureGraph/Source/TextureGraph/Private/TG_Pin.cpp
+++ b/Engine/Plugins/TextureGraph/Source/TextureGraph/Private/TG_Pin.cpp
@@ -83,6 +83,25 @@ void UTG_Pin::InitSelfVar()
 	}
 }
 
+void UTG_Pin::ResetSelfVarDefault()
+{
+	SelfVar.PinId = Id; // Assign the Pinid as the Id of the Var
+	// The concrete value of the selfvar is only initialized for inputs, not output
+	if (IsInput())
+	{
+		FTG_Argument Arg = GetArgument();
+		SelfVar.CopyFrom(GetNodePtr()->GetExpression(), Arg);
+	}
+	if (IsOutput())
+	{
+		FTG_Argument Arg = GetArgument();
+		if (Arg.IsTexture())
+		{
+			SelfVar.CopyFrom(GetNodePtr()->GetExpression(), Arg);
+		}
+	}
+}
+
 void UTG_Pin::RemapId(FTG_Id InId)
 {
 	Modify();
diff --git a/Engine/Plugins/TextureGraph/Source/TextureGraph/Private/TG_Var.cpp b/Engine/Plugins/TextureGraph/Source/TextureGraph/Private/TG_Var.cpp
--- a/Engine/Plugins/TextureGraph/Source/TextureGraph/Private/TG_Var.cpp
+++ b/Engine/Plugins/TextureGraph/Source/TextureGraph/Private/TG_Var.cpp
@@ -8,6 +8,7 @@
 #include "TG_Texture.h"
 #include "TG_OutputSettings.h"
 #include "TG_Variant.h"
+#include "TG_Material.h"
 #include "Model/ModelObject.h"
 
 #include "Expressions/TG_Expression.h"
@@ -179,6 +180,11 @@ template <> void TG_Var_SetValueFromString(FTG_TextureDescriptor& Value, const F
 	Value.InitFromString(StrVal);
 }
 
+template <> void TG_Var_SetValueFromString(FTG_Material& Value, const FString& StrVal)
+{
+	Value.InitFromString(StrVal);
+}
+
 template <> void TG_Var_SetValueFromString(FPatternMaskPlacement_TS& Value, const FString& StrVal)
 {
 	Value.InitFromString(StrVal);
@@ -282,6 +288,11 @@ void VarPropertySerializer_FTG_Texture(FTG_Var::VarPropertySerialInfo& Info)
 	Generic_Struct_Serializer<FTG_Texture>(Info);
 }
 
+void VarPropertySerializer_FTG_Material(FTG_Var::VarPropertySerialInfo& Info)
+{
+	Generic_Struct_Serializer<FTG_Material>(Info);
+}
+
 void VarPropertySerializer_FTG_VariantArray(FTG_Var::VarPropertySerialInfo& Info)
 {
 	/// Just call the base generic serializer and set the var to array
@@ -533,6 +544,7 @@ FTG_Var::VarPropertySerializerMap FTG_Var::DefaultPropertySerializers
 		VAR_PROPERTY_SERIALIZER_DEF(FTG_OutputSettings),
 		VAR_PROPERTY_SERIALIZER_DEF(FTG_TextureDescriptor),
 		VAR_PROPERTY_SERIALIZER_DEF(FTG_Variant),
+		VAR_PROPERTY_SERIALIZER_DEF(FTG_Material),
 		VAR_PROPERTY_SERIALIZER_DEF(FGradientDir_TS),
 		VAR_PROPERTY_SERIALIZER_DEF(FPatternMaskCutout_TS),
 		VAR_PROPERTY_SERIALIZER_DEF(FPatternMaskBevel_TS),
@@ -769,6 +781,11 @@ void VarArchiveSerializer_FTG_TextureDescriptor(FTG_Var::VarArchiveSerialInfo& I
 {
 	Generic_Simple_ArSerializer<FTG_TextureDescriptor>(Info);
 }
+void VarArchiveSerializer_FTG_Material(FTG_Var::VarArchiveSerialInfo& Info)
+{
+	Generic_Simple_ArSerializer<FTG_Material>(Info);
+}
+
 //////////////////////////////////////////////////////////////////////////////////////////////////////////////////////////////////////
 FTG_Var::VarArchiveSerializerMap FTG_Var::DefaultArchiveSerializers
 (
@@ -784,7 +801,8 @@ FTG_Var::VarArchiveSerializerMap FTG_Var::DefaultArchiveSerializers
 		VAR_ARCHIVE_SERIALIZER_DEF(FLinearColor),
 		VAR_ARCHIVE_SERIALIZER_DEF(FTG_VariantArray),
 		VAR_ARCHIVE_SERIALIZER_DEF(FTG_OutputSettings),
-		VAR_ARCHIVE_SERIALIZER_DEF(FTG_TextureDescriptor)
+		VAR_ARCHIVE_SERIALIZER_DEF(FTG_TextureDescriptor),
+		VAR_ARCHIVE_SERIALIZER_DEF(FTG_Material)
 	}
 ); 
 
diff --git a/Engine/Plugins/TextureGraph/Source/TextureGraph/Public/Expressions/Input/TG_Expression_Material.h b/Engine/Plugins/TextureGraph/Source/TextureGraph/Public/Expressions/Input/TG_Expression_Material.h
--- a/Engine/Plugins/TextureGraph/Source/TextureGraph/Public/Expressions/Input/TG_Expression_Material.h
+++ b/Engine/Plugins/TextureGraph/Source/TextureGraph/Public/Expressions/Input/TG_Expression_Material.h
@@ -7,6 +7,7 @@
 #include "Expressions/TG_Expression_MaterialBase.h"
 #include "Materials/Material.h"
 #include "TG_Texture.h"
+#include "TG_Material.h"
 
 #include "TG_Expression_Material.generated.h"
 
@@ -16,6 +17,7 @@ UCLASS()
 class TEXTUREGRAPH_API UTG_Expression_Material : public UTG_Expression_MaterialBase
 {
 	GENERATED_BODY()
+
 public:
 	UTG_Expression_Material();
 	virtual ~UTG_Expression_Material();
@@ -24,11 +26,18 @@ public:
 	virtual void PostEditUndo() override;
 #endif
 
-	// The input material to employ for rendering
-	UPROPERTY(EditAnywhere, BlueprintReadWrite, Setter, Category = NoCategory, meta = (TGType = "TG_Setting"))
-	TObjectPtr<UMaterialInterface> Material;
-	
-	void SetMaterial(UMaterialInterface* InMaterial);
+	virtual void PostLoad() override;
+
+#if WITH_EDITORONLY_DATA
+	UE_DEPRECATED(5.6, "Material field has been deprecated.")
+	UPROPERTY(meta = (DeprecatedProperty, DeprecationMessage = "Use the InputMaterial to specify the Material asset referenced"))
+	TObjectPtr<UMaterialInterface> Material_DEPRECATED;
+#endif // WITH_EDITORONLY_DATA
+
+	// The input material referenced by this Material node
+	UPROPERTY(EditAnywhere, Setter, Category = NoCategory, meta = (TGType = "TG_InputParam", DisplayName = "Material", PinDisplayName = "Material"))
+	FTG_Material InputMaterial;
+	void SetInputMaterial(const FTG_Material& InMaterial);
 
 	// The Material attribute identifier among all the attributes of the material that is rendered in the output
 	UPROPERTY(EditAnywhere, BlueprintReadWrite, Setter, Category = NoCategory, meta = (TGType = "TG_Setting", GetOptions = "GetRenderAttributeOptions"))
@@ -39,7 +48,6 @@ public:
 	UFUNCTION(CallInEditor)
 	TArray<FName> GetRenderAttributeOptions() const;
 
-
 	virtual bool CanHandleAsset(UObject* Asset) override;
 	virtual void SetAsset(UObject* Asset) override;
 	virtual FText GetTooltipText() const override { return FText::FromString(TEXT("Renders a material into a quad and makes it available. It is automatically exposed as a graph input parameter.")); } 
@@ -49,9 +57,13 @@ protected:
 	UPROPERTY(Transient, DuplicateTransient)
 	TObjectPtr<UMaterialInterface> MaterialCopy = nullptr;
 
+	// As the "Material" or the "InputMaterial" properties change
+	// we update the current active "MaterialCopy" and pass on internally as the active "GetMaterial()"
 	virtual void SetMaterialInternal(UMaterialInterface* InMaterial) override;
 
-	UPROPERTY(EditAnywhere, BlueprintReadWrite, Category = "")
+	// Title name is still used mostly for legacy but not exposed anymore in the details.
+	// This is changed on the node itself and then call SetTitleName and rename the InnputMaterial Alias name
+	UPROPERTY()
 	FName TitleName = TEXT("Material");
 
 	virtual void Initialize() override;
diff --git a/Engine/Plugins/TextureGraph/Source/TextureGraph/Public/TG_Material.h b/Engine/Plugins/TextureGraph/Source/TextureGraph/Public/TG_Material.h
new file mode 100644
--- /dev/null
+++ b/Engine/Plugins/TextureGraph/Source/TextureGraph/Public/TG_Material.h
@@ -0,0 +1,64 @@
+// Copyright Epic Games, Inc. All Rights Reserved.
+
+#pragma once
+#include "TG_SystemTypes.h"
+#include "Materials/Material.h"
+#include "Misc/OutputDeviceNull.h"
+#include "TG_Material.generated.h"
+
+USTRUCT(BlueprintType)
+struct TEXTUREGRAPH_API FTG_Material
+{
+	GENERATED_USTRUCT_BODY()
+	
+public:	
+	UPROPERTY(Transient, EditAnywhere, Category = "Material", meta = (AllowedClasses = "/Script/Engine.MaterialInterface"))
+	FSoftObjectPath AssetPath;
+
+	FTG_Material() = default;
+	FTG_Material(const FTG_Material& InSrc) = default;
+	FTG_Material& operator = (const FTG_Material& RHS) = default;
+
+	// Check that the referenced asset is valid.
+	// If true, the GetMaterial call will NOT return nullptr
+	bool IsValid() const;
+
+	// Access the referenced Material asset
+	// return a valid MaterialInterface object or null if the reference is invalid.
+	UMaterialInterface* GetMaterial() const;
+
+	// Assign the referenced asset from an actual UMaterialInsterface* live object
+	void SetMaterial(UMaterialInterface* InMaterial);
+
+	bool Serialize(FArchive& Ar);
+	friend FArchive& operator<<(FArchive& Ar, FTG_Material& InMaterial);
+	bool operator==(const FTG_Material& RHS) const = default;
+	void ResetTexturePath();
+
+	void InitFromString(const FString& StrVal)
+	{
+		FOutputDeviceNull NullOut;
+		FTG_Material::StaticStruct()->ImportText(*StrVal, this, /*OwnerObject*/nullptr, 0, &NullOut, FTG_Material::StaticStruct()->GetName(), /*bAllowNativeOverride*/true);
+	}
+	FString ToString() const
+	{
+		FString ExportString;
+		FTG_Material::StaticStruct()->ExportText(ExportString, this, this, /*OwnerObject*/nullptr, /*PortFlags*/0, /*ExportRootScope*/nullptr);
+		return ExportString;
+	}
+
+};
+
+
+template<>
+struct TStructOpsTypeTraits<FTG_Material>
+	: public TStructOpsTypeTraitsBase2<FTG_Material>
+{
+	enum
+	{
+		WithSerializer = true,
+		WithCopy = true,
+		WithIdenticalViaEquality = true
+	};
+};
+
diff --git a/Engine/Plugins/TextureGraph/Source/TextureGraph/Public/TG_Pin.h b/Engine/Plugins/TextureGraph/Source/TextureGraph/Public/TG_Pin.h
--- a/Engine/Plugins/TextureGraph/Source/TextureGraph/Public/TG_Pin.h
+++ b/Engine/Plugins/TextureGraph/Source/TextureGraph/Public/TG_Pin.h
@@ -38,6 +38,8 @@ class TEXTUREGRAPH_API UTG_Pin : public UObject {
 	// Initialize the self var with the matching expression argument
 	void InitSelfVar();
 
+	void ResetSelfVarDefault();
+
 	// Used when remapping a pin to a different id, update the sefl Var if as well
 	void RemapId(FTG_Id InId);
 
diff --git a/Engine/Plugins/TextureGraph/Source/TextureGraph/Public/TG_Var.h b/Engine/Plugins/TextureGraph/Source/TextureGraph/Public/TG_Var.h
--- a/Engine/Plugins/TextureGraph/Source/TextureGraph/Public/TG_Var.h
+++ b/Engine/Plugins/TextureGraph/Source/TextureGraph/Public/TG_Var.h
@@ -228,6 +228,7 @@ template <> FString TG_Var_LogValue(struct FTG_Scalar& Value);
 template <> FString TG_Var_LogValue(struct FTG_OutputSettings& Value);
 template <> FString TG_Var_LogValue(struct FTG_TextureDescriptor& Value);
 template <> FString TG_Var_LogValue(struct FTG_Variant& Value);
+template <> FString TG_Var_LogValue(struct FTG_Material& Value);
 
 
 template <> void TG_Var_SetValueFromString(int& Value, const FString& StrVal);
@@ -247,5 +248,6 @@ template <> void TG_Var_SetValueFromString(struct FTG_Scalar& Value, const FStri
 template <> void TG_Var_SetValueFromString(struct FTG_OutputSettings& Value, const FString& StrVal);
 template <> void TG_Var_SetValueFromString(struct FTG_TextureDescriptor& Value, const FString& StrVal);
 template <> void TG_Var_SetValueFromString(struct FTG_Variant& Value, const FString& StrVal);
+template <> void TG_Var_SetValueFromString(struct FTG_Material& Value, const FString& StrVal);
 
 
diff --git a/Engine/Plugins/TextureGraph/Source/TextureGraphEditor/Private/Customizations/TG_MaterialCustomization.h b/Engine/Plugins/TextureGraph/Source/TextureGraphEditor/Private/Customizations/TG_MaterialCustomization.h
new file mode 100644
--- /dev/null
+++ b/Engine/Plugins/TextureGraph/Source/TextureGraphEditor/Private/Customizations/TG_MaterialCustomization.h
@@ -0,0 +1,42 @@
+// Copyright Epic Games, Inc. All Rights Reserved.
+
+#pragma once
+
+#include "CoreMinimal.h"
+#include "IPropertyTypeCustomization.h"
+#include "DetailWidgetRow.h"
+#include "DetailLayoutBuilder.h"
+#include "IDetailChildrenBuilder.h"
+#include "TextureGraph.h"
+#include "TG_Graph.h"
+#include "TG_Material.h"
+
+class FTG_MaterialCustomization : public IPropertyTypeCustomization
+{
+public:
+	static TSharedRef<IPropertyTypeCustomization> Create()
+	{
+		return MakeShareable(new FTG_MaterialCustomization);
+	}
+
+	virtual void CustomizeHeader(TSharedRef<IPropertyHandle> PropertyHandle, FDetailWidgetRow& HeaderRow, IPropertyTypeCustomizationUtils& CustomizationUtils) override
+	{
+	}
+
+	virtual void CustomizeChildren(TSharedRef<IPropertyHandle> PropertyHandle, IDetailChildrenBuilder& ChildBuilder, IPropertyTypeCustomizationUtils& CustomizationUtils) override
+	{
+		const FName TGType = FName(TEXT("TGType"));
+		const FName CategoryName = PropertyHandle->GetDefaultCategoryName();
+		const FString ParentTGType = PropertyHandle->GetParentHandle()->GetMetaData(TGType);
+
+		// show the desc UI
+		TSharedPtr<IPropertyHandle> DescPropertyHandle = PropertyHandle->GetChildHandle(TEXT("AssetPath"));
+		FString DisplayName = PropertyHandle->GetPropertyDisplayName().ToString();
+		{
+			ChildBuilder.AddProperty(DescPropertyHandle.ToSharedRef())
+				.DisplayName(FText::FromString(DisplayName))
+				.ShouldAutoExpand(true);
+		}
+	}
+
+};
diff --git a/Engine/Plugins/TextureGraph/Source/TextureGraphEditor/Private/EdGraph/TG_EditorGraphPanelPinFactory.cpp b/Engine/Plugins/TextureGraph/Source/TextureGraphEditor/Private/EdGraph/TG_EditorGraphPanelPinFactory.cpp
--- a/Engine/Plugins/TextureGraph/Source/TextureGraphEditor/Private/EdGraph/TG_EditorGraphPanelPinFactory.cpp
+++ b/Engine/Plugins/TextureGraph/Source/TextureGraphEditor/Private/EdGraph/TG_EditorGraphPanelPinFactory.cpp
@@ -37,6 +37,8 @@
 #include "Pins/STG_GraphPinOutputSettings.h"
 #include "Pins/STG_GraphPinTextureDescriptor.h"
 #include "Pins/STG_GraphPinString.h"
+#include "Pins/STG_GraphPinMaterial.h"
+
 
 #include "NumericPropertyParams.h"
 
@@ -108,6 +110,7 @@ TSharedPtr<class SGraphPin> FTG_EditorGraphPanelPinFactory::CreatePin(class UEdG
 			UScriptStruct* RotatorStruct = TBaseStructure<FRotator>::Get();
 			UScriptStruct* OutputSettingsStruct = TBaseStructure<FTG_OutputSettings>::Get();
 			UScriptStruct* TextureDescSettingsStruct = TBaseStructure<FTG_TextureDescriptor>::Get();
+			UScriptStruct* TG_Material = TBaseStructure<FTG_Material>::Get();
 
 			if (InPin->PinType.PinSubCategoryObject == ColorStruct)
 			{
@@ -137,6 +140,10 @@ TSharedPtr<class SGraphPin> FTG_EditorGraphPanelPinFactory::CreatePin(class UEdG
 			{
 				return SNew(STG_GraphPinTextureDescriptor, InPin);
 			}
+			else if (InPin->PinType.PinSubCategoryObject == TG_Material)
+			{
+				return SNew(STG_GraphPinMaterial, InPin);
+			}
 		}
 		else if(InPin->PinType.PinCategory == UTG_EdGraphSchema::PC_Array)
 		{
diff --git a/Engine/Plugins/TextureGraph/Source/TextureGraphEditor/Private/Pins/STG_GraphPinMaterial.cpp b/Engine/Plugins/TextureGraph/Source/TextureGraphEditor/Private/Pins/STG_GraphPinMaterial.cpp
new file mode 100644
--- /dev/null
+++ b/Engine/Plugins/TextureGraph/Source/TextureGraphEditor/Private/Pins/STG_GraphPinMaterial.cpp
@@ -0,0 +1,53 @@
+// Copyright Epic Games, Inc. All Rights Reserved.
+
+
+#include "STG_GraphPinMaterial.h"
+
+#include "EdGraph/EdGraphPin.h"
+#include "EdGraph/EdGraphSchema.h"
+#include "GenericPlatform/GenericApplication.h"
+#include "K2Node_CallFunction.h"
+#include "PropertyEditorUtils.h"
+#include "Misc/Attribute.h"
+#include "ScopedTransaction.h"
+#include "UObject/NameTypes.h"
+#include "EdGraph/TG_EdGraphSchema.h"
+#include "Expressions/TG_Expression.h"
+
+#include "PropertyCustomizationHelpers.h"
+
+
+void STG_GraphPinMaterial::Construct(const FArguments& InArgs, UEdGraphPin* InGraphPinObj)
+{
+	SGraphPin::Construct(SGraphPin::FArguments(), InGraphPinObj);
+}
+
+FString	STG_GraphPinMaterial::GetCurrentAssetPath() const
+{
+	const UTG_EdGraphSchema* Schema = Cast<const UTG_EdGraphSchema>(GraphPinObj->GetOwningNode()->GetSchema());
+	UTG_Pin* TSPin = Schema->GetTGPinFromEdPin(GraphPinObj);
+	return TSPin->GetSelfVar()->GetAs<FTG_Material>().AssetPath.GetAssetPathString();
+}
+
+void STG_GraphPinMaterial::OnAssetSelected(const FAssetData& AssetData)
+{
+	const UTG_EdGraphSchema* Schema = Cast<const UTG_EdGraphSchema>(GraphPinObj->GetOwningNode()->GetSchema());
+	UTG_Pin* TSPin = Schema->GetTGPinFromEdPin(GraphPinObj);
+
+	FTG_Material NewMaterial;
+	NewMaterial.AssetPath = AssetData.GetObjectPathString();
+	TSPin->SetValue(NewMaterial);
+}
+
+TSharedRef<SWidget>	STG_GraphPinMaterial::GetDefaultValueWidget()
+{
+	return  SNew(SObjectPropertyEntryBox)
+		.ObjectPath(this, &STG_GraphPinMaterial::GetCurrentAssetPath)
+		.AllowedClass(UMaterialInterface::StaticClass())
+		.OnObjectChanged(this, &STG_GraphPinMaterial::OnAssetSelected)
+		.AllowCreate(true)
+		.AllowClear(true)
+		.DisplayUseSelected(true)
+		.DisplayBrowse(true)
+		;
+}
diff --git a/Engine/Plugins/TextureGraph/Source/TextureGraphEditor/Private/Pins/STG_GraphPinMaterial.h b/Engine/Plugins/TextureGraph/Source/TextureGraphEditor/Private/Pins/STG_GraphPinMaterial.h
new file mode 100644
--- /dev/null
+++ b/Engine/Plugins/TextureGraph/Source/TextureGraphEditor/Private/Pins/STG_GraphPinMaterial.h
@@ -0,0 +1,32 @@
+// Copyright Epic Games, Inc. All Rights Reserved.
+
+#pragma once
+
+#include "CoreMinimal.h"
+#include "Internationalization/Text.h"
+#include "SGraphPin.h"
+#include "Templates/SharedPointer.h"
+#include "Types/SlateEnums.h"
+#include "Widgets/DeclarativeSyntaxSupport.h"
+#include "Widgets/SWidget.h"
+#include "TG_Material.h"
+
+class SWidget;
+class UEdGraphPin;
+
+class STG_GraphPinMaterial : public SGraphPin
+{
+public:
+	SLATE_BEGIN_ARGS(STG_GraphPinMaterial) {}
+	SLATE_END_ARGS()
+
+	void Construct(const FArguments& InArgs, UEdGraphPin* InGraphPinObj);
+
+protected:
+	FString GetCurrentAssetPath() const;
+	void OnAssetSelected(const FAssetData& AssetData);
+
+	//~ Begin SGraphPin Interface
+	virtual TSharedRef<SWidget>	GetDefaultValueWidget() override;
+	//~ End SGraphPin Interface
+};
diff --git a/Engine/Plugins/TextureGraph/Source/TextureGraphEditor/Private/TextureGraphEditorModule.cpp b/Engine/Plugins/TextureGraph/Source/TextureGraphEditor/Private/TextureGraphEditorModule.cpp
--- a/Engine/Plugins/TextureGraph/Source/TextureGraphEditor/Private/TextureGraphEditorModule.cpp
+++ b/Engine/Plugins/TextureGraph/Source/TextureGraphEditor/Private/TextureGraphEditorModule.cpp
@@ -15,6 +15,7 @@
 #include "EdGraph/TG_EditorErrorReporter.h"
 #include "Customizations/TG_ParameterCustomization.h"
 #include "Customizations/TG_TextureCustomization.h"
+#include "Customizations/TG_MaterialCustomization.h"
 #include "Customizations/TG_VariantCustomization.h"
 #include "Customizations/TG_ScalarCustomization.h"
 #include "Customizations/TG_MaterialMappingInfoCustomization.h"
@@ -36,6 +37,7 @@ void FTextureGraphEditorModule::StartupModule()
 	FPropertyEditorModule& PropertyEditorModule = FModuleManager::GetModuleChecked<FPropertyEditorModule>("PropertyEditor");
 	PropertyEditorModule.RegisterCustomPropertyTypeLayout("TG_ParameterInfo", FOnGetPropertyTypeCustomizationInstance::CreateStatic(&FTG_ParameterInfoCustomization::Create));
 	PropertyEditorModule.RegisterCustomPropertyTypeLayout("TG_Texture", FOnGetPropertyTypeCustomizationInstance::CreateStatic(&FTG_TextureCustomization::Create));
+	PropertyEditorModule.RegisterCustomPropertyTypeLayout("TG_Material", FOnGetPropertyTypeCustomizationInstance::CreateStatic(&FTG_MaterialCustomization::Create));
 	PropertyEditorModule.RegisterCustomPropertyTypeLayout("TG_Variant", FOnGetPropertyTypeCustomizationInstance::CreateStatic(&FTG_VariantCustomization::Create));
 	PropertyEditorModule.RegisterCustomPropertyTypeLayout("FloatProperty", FOnGetPropertyTypeCustomizationInstance::CreateStatic(&FTG_ScalarCustomization::Create), MakeShared<FTG_ScalarTypeIdentifier>());
 	PropertyEditorModule.RegisterCustomPropertyTypeLayout("MaterialMappingInfo", FOnGetPropertyTypeCustomizationInstance::CreateStatic(&FTG_MaterialMappingInfoCustomization::Create));
