Texture Graph:
- Adding comments to clarify and document how to access params
- Adding a helper function

diff --git a/Engine/Plugins/TextureGraph/Source/TextureGraph/Private/TG_Graph.cpp b/Engine/Plugins/TextureGraph/Source/TextureGraph/Private/TG_Graph.cpp
--- a/Engine/Plugins/TextureGraph/Source/TextureGraph/Private/TG_Graph.cpp
+++ b/Engine/Plugins/TextureGraph/Source/TextureGraph/Private/TG_Graph.cpp
@@ -748,6 +748,19 @@ void UTG_Graph::NotifyNodePostEvaluate(UTG_Node* InNode, const FTG_EvaluationCon
 #endif
 }
 
+FName UTG_Graph::GetParamName(const FTG_Id& InParamId) const
+{
+	const UTG_Pin* Pin = GetPin(InParamId);
+	if (Pin)
+	{
+		if (Pin->IsParam())
+		{
+			return Pin->GetAliasName();
+		}
+	}
+	return FName();
+}
+
 FTG_Ids			UTG_Graph::GetInputParamIds() const
 {
 	FTG_Ids Ids;
diff --git a/Engine/Plugins/TextureGraph/Source/TextureGraph/Public/TG_Graph.h b/Engine/Plugins/TextureGraph/Source/TextureGraph/Public/TG_Graph.h
--- a/Engine/Plugins/TextureGraph/Source/TextureGraph/Public/TG_Graph.h
+++ b/Engine/Plugins/TextureGraph/Source/TextureGraph/Public/TG_Graph.h
@@ -185,25 +185,37 @@ public:
 	// Accessors for nodes, pins, params
 	//////////////////////////////////////////////////////////////////////////
 
+	// Access a Node from Id
 	UTG_Node*		GetNode(FTG_Id NodeId) { return IsValidNode(NodeId) ? Nodes[NodeId.NodeIdx()] : nullptr; }
 	const UTG_Node* GetNode(FTG_Id NodeId) const { return IsValidNode(NodeId) ? Nodes[NodeId.NodeIdx()] : nullptr; }
-    UTG_Pin*		GetPin(FTG_Id PinId) { return IsValidPin(PinId) ? Nodes[PinId.NodeIdx()]->Pins[PinId.PinIdx()].Get() : nullptr; }
+    
+	// Access a Pin from Id
+	UTG_Pin*		GetPin(FTG_Id PinId) { return IsValidPin(PinId) ? Nodes[PinId.NodeIdx()]->Pins[PinId.PinIdx()].Get() : nullptr; }
 	const UTG_Pin*	GetPin(FTG_Id PinId) const { return IsValidPin(PinId) ? Nodes[PinId.NodeIdx()]->Pins[PinId.PinIdx()].Get() : nullptr; }
+	
+	// Access the Node owning a Pin from the Pin Id
 	UTG_Node*		GetNodeFromPinId(FTG_Id PinId) { return GetNode( FTG_Id(PinId.NodeIdx()) ); }
 	
+	// Access a Var from Id (the Id of a Var is similar to the Id of the Pin that owns it
 	FTG_Var*		GetVar(FTG_Id VarId) { UTG_Pin* Pin = GetPin(VarId); return (Pin ? Pin->EditSelfVar() : nullptr); }
 
+	// Accessing the params
+	// A param is simply the Id of a pin exposed as a param of the graph, the pin can be retrieved using GetPin()
+	// The Param name is the Alias Name of the associated TG_Pin
+
 	FTG_Ids			GetParamIds() const { FTG_Ids Array;  Params.GenerateValueArray(Array); return Array; }
 	TArray<FName>	GetParamNames() const { TArray<FName> Array;  Params.GenerateKeyArray(Array); return Array; }
 
-	// Find a param pin from its name
-	FTG_Id			FindParamPinId(const FName& InName) const		{ auto* PinId = Params.Find(InName); return (PinId ? (*PinId) : FTG_Id()); }
-	const UTG_Pin*	FindParamPin(const FName& InName) const		{ return GetPin(FindParamPinId(InName)); }
-	UTG_Pin*		FindParamPin(const FName& InName)				{ return GetPin(FindParamPinId(InName)); }
+	FName			GetParamName(const FTG_Id& InParamId) const;
 
 	FTG_Ids			GetInputParamIds() const;
 	FTG_Ids			GetOutputParamIds() const;
 
+	// Find a param pin from its name
+	FTG_Id			FindParamPinId(const FName& InName) const		{ auto* PinId = Params.Find(InName); return (PinId ? (*PinId) : FTG_Id()); }
+	const UTG_Pin*	FindParamPin(const FName& InName) const			{ return GetPin(FindParamPinId(InName)); }
+	UTG_Pin*		FindParamPin(const FName& InName)				{ return GetPin(FindParamPinId(InName)); }
+
 	// Get the names and ids of the output params which are Textures
 	int				GetOutputParamTextures(TArray<FName>& OutNames, FTG_Ids& OutPinIds) const;
 
