Enum pin support

diff --git a/Engine/Plugins/TextureGraph/Source/TextureGraph/Private/Expressions/TG_Expression.cpp b/Engine/Plugins/TextureGraph/Source/TextureGraph/Private/Expressions/TG_Expression.cpp
--- a/Engine/Plugins/TextureGraph/Source/TextureGraph/Private/Expressions/TG_Expression.cpp
+++ b/Engine/Plugins/TextureGraph/Source/TextureGraph/Private/Expressions/TG_Expression.cpp
@@ -69,6 +69,7 @@ FTG_Signature::FInit UTG_Expression::GetSignatureInitArgsFromClass() const
 #if WITH_EDITORONLY_DATA
 		auto PropertyTGType = Prop->GetMetaData(TEXT("TGType"));
 		bIsNotConnectable = Prop->HasMetaData(TEXT("TGPinNotConnectable"));
+		bool bIsEnum = false;
 #else
 		FString PropertyTGType;
 #endif
@@ -78,41 +79,43 @@ FTG_Signature::FInit UTG_Expression::GetSignatureInitArgsFromClass() const
 			FByteProperty* ByteProperty = CastField<FByteProperty>(*Prop);
 			bIsNotConnectable |= (ByteProperty != nullptr);
 			FEnumProperty* EnumProperty = CastField<FEnumProperty>(*Prop);
-			bIsNotConnectable |= (EnumProperty != nullptr);
+			bIsEnum = (EnumProperty != nullptr);
+			//bIsNotConnectable |= (EnumProperty != nullptr);
 			FNameProperty* NameProperty = CastField<FNameProperty>(*Prop);
 			bIsNotConnectable |= (NameProperty != nullptr);
 		}
-		uint8 MaskNotConnectable = bIsNotConnectable ? static_cast<uint8>(ETG_Access::NotConnectableFlag) : 0;
-		
+		uint8 Flags = bIsNotConnectable ? static_cast<uint8>(ETG_Access::NotConnectableFlag) : 0;
+		Flags |= bIsEnum ? AccessToUint8(EnumFlag) : 0;
+
 		if (PropertyTGType.Compare(TEXT("TG_InputParam")) == 0)
 		{
 			// Add the new argument to the signature
-			SignatureInit.Arguments.Add(FTG_Argument{ PropName, PropTypeName, { ETG_Access::InParam, MaskNotConnectable} });
+			SignatureInit.Arguments.Add(FTG_Argument{ PropName, PropTypeName, { ETG_Access::InParam, Flags} });
 		}																								   
 		else if (PropertyTGType.Compare(TEXT("TG_OutputParam")) == 0)									   
 		{																									   
 			// Add the new argument to the signature													   
-			SignatureInit.Arguments.Add(FTG_Argument{ PropName, PropTypeName, { ETG_Access::OutParam, MaskNotConnectable } });
+			SignatureInit.Arguments.Add(FTG_Argument{ PropName, PropTypeName, { ETG_Access::OutParam, Flags } });
 		}																								   
 		else if (PropertyTGType.Compare(TEXT("TG_Input")) == 0)										   
 		{													   
 			// Add the new argument to the signature													   
-			SignatureInit.Arguments.Add(FTG_Argument{ PropName, PropTypeName, { ETG_Access::In, MaskNotConnectable } });
+			SignatureInit.Arguments.Add(FTG_Argument{ PropName, PropTypeName, { ETG_Access::In, Flags } });
 		}																								   
 		else if (PropertyTGType.Compare(TEXT("TG_Output")) == 0)										   
 		{																	   
 			// Add the new argument to the signature													   
-			SignatureInit.Arguments.Add(FTG_Argument{ PropName, PropTypeName, { ETG_Access::Out, MaskNotConnectable } });
+			SignatureInit.Arguments.Add(FTG_Argument{ PropName, PropTypeName, { ETG_Access::Out, Flags } });
 		}
 		else if (PropertyTGType.Compare(TEXT("TG_Setting")) == 0)
 		{
 			// Add the new argument to the signature as a Setting													   
-			SignatureInit.Arguments.Add(FTG_Argument{ PropName, PropTypeName, { ETG_Access::InSetting, MaskNotConnectable } });
+			SignatureInit.Arguments.Add(FTG_Argument{ PropName, PropTypeName, { ETG_Access::InSetting, Flags } });
 		}
 		else if (PropertyTGType.Compare(TEXT("TG_Private")) == 0)
 		{
 			// Add the new argument to the signature as a Private													   
-			SignatureInit.Arguments.Add(FTG_Argument{ PropName, PropTypeName, { ETG_Access::Private, MaskNotConnectable } });
+			SignatureInit.Arguments.Add(FTG_Argument{ PropName, PropTypeName, { ETG_Access::Private, Flags } });
 		}
 		else
 		{
diff --git a/Engine/Plugins/TextureGraph/Source/TextureGraph/Private/TG_GraphEvaluation.cpp b/Engine/Plugins/TextureGraph/Source/TextureGraph/Private/TG_GraphEvaluation.cpp
--- a/Engine/Plugins/TextureGraph/Source/TextureGraph/Private/TG_GraphEvaluation.cpp
+++ b/Engine/Plugins/TextureGraph/Source/TextureGraph/Private/TG_GraphEvaluation.cpp
@@ -52,6 +52,7 @@ void FTG_Evaluation::TransferVarToPin(UTG_Pin* InPin, FTG_EvaluationContext* Con
 		bool bNeedsConversion = InPin->ConnectionNeedsConversion();
 		FName ArrayConverterKey = InPin->GetInputVarConverterKey();
 		FTG_Var* OutVar = InPin->EditConvertedVar(); // Generate the converted version of the connected input into the CoonvertedVar of the Pin
+		bool bIsEnum = InPin->GetArgument().ArgumentType.IsEnum();
 
 		if (!bNeedsConversion && (Var->IsArray() && !InPin->EditSelfVar()->IsArray()))
 		{
@@ -67,11 +68,13 @@ void FTG_Evaluation::TransferVarToPin(UTG_Pin* InPin, FTG_EvaluationContext* Con
 			// Else COnverter is valid and we need to run it 
 			if ((*Converter))
 			{
-				FTG_Evaluation::VarConverterInfo Info;
-				Info.Context = Context;
-				Info.InVar = Var;
-				Info.OutVar = OutVar;
-				Info.Index = Index;
+				FTG_Evaluation::VarConverterInfo Info
+				{
+					.InVar = Var,
+					.OutVar = OutVar,
+					.Index = Index,
+					.Context = Context
+				};
 
 				(*Converter)(Info);
 
@@ -89,13 +92,16 @@ void FTG_Evaluation::TransferVarToPin(UTG_Pin* InPin, FTG_EvaluationContext* Con
 
 		if (InPin->NeedsConformance())
 		{
-	 		FTG_Evaluation::VarConformerInfo Info;
-	 		Info.Context = Context;
-	 		Info.InVar = Var;
-	 		Info.OutVar = InPin->EditSelfVar();
-			Info.Index = Index;
-			bool Result = (InPin->ConformerFunctor)(Info);
-	 		if (Result)
+			FTG_Evaluation::VarConformerInfo Info
+			{
+				.InVar = Var,
+				.OutVar = InPin->EditSelfVar(),
+				.Index = Index,
+				.Context = Context
+			};
+
+			bool bResult = (InPin->ConformerFunctor)(Info);
+	 		if (bResult)
 			{
 				// Since the source var is conformed into the SelfVar
 	            // we need to pass the self var as the input for the pin
@@ -532,13 +538,17 @@ void FTG_VariantArrayToFTG_Texture_Converter(FTG_Evaluation::VarConverterInfo& I
 	FTG_VariantToFTG_Texture_Converter_Internal(Output, Input, Info);
 }
 
+void FloatToEnum_Converter(FTG_Evaluation::VarConverterInfo& Info)
+{
+	auto Input = Info.InVar->GetAs<float>();
+	Info.OutVar->EditAs<int32>() = (int32)Input;
+}
+
 //////////////////////////////////////////////////////////////////////////////////////////////////////////////////////////////////////
 #define VAR_CONVERTER(nameFrom, nameTo, function)	{ MakeConvertKey(TEXT(#nameFrom), TEXT(#nameTo)), std::bind(&function, std::placeholders::_1) }
 #define VAR_CONVERTER_DEF(nameFrom, nameTo, function)		VAR_CONVERTER(#nameFrom, #nameTo, #function)
 #define VAR_CONVERTER_NULL(nameFrom, nameTo)	{ MakeConvertKey(TEXT(#nameFrom), TEXT(#nameTo)), nullptr }
 
-
-
 //////////////////////////////////////////////////////////////////////////////////////////////////////////////////////////////////////
 FTG_Evaluation::ConverterMap FTG_Evaluation::DefaultConverters
 (
@@ -550,6 +560,7 @@ FTG_Evaluation::ConverterMap FTG_Evaluation::DefaultConverters
 		VAR_CONVERTER(float, FVector4f, FloatToFVector4f_Converter),
 		VAR_CONVERTER(float, FTG_Texture, FloatToFTG_Texture_Converter),
 		VAR_CONVERTER(float, FVector2f, FloatToFVector2f_Converter),
+		VAR_CONVERTER(float, Enum, FloatToEnum_Converter),
 
 		VAR_CONVERTER(FLinearColor, FVector4f, FLinearColorToFVector4f_Converter),
 		VAR_CONVERTER(FLinearColor, FTG_Texture, FLinearColorToFTG_Texture_Converter),
@@ -605,15 +616,24 @@ FName FTG_Evaluation::MakeConvertKey(FName From, FName To)
 	return FName(From.ToString() + TEXT("To") + To.ToString());
 }
 
+FName FTG_Evaluation::MakeConvertKey(const FTG_Argument& ArgFrom, const FTG_Argument& ArgTo)
+{
+	if (!ArgTo.ArgumentType.IsEnum())
+		return MakeConvertKey(ArgFrom.GetCPPTypeName(), ArgTo.GetCPPTypeName());
+
+	return FName(ArgFrom.GetCPPTypeName().ToString() + TEXT("ToEnum"));
+}
+
 bool FTG_Evaluation::AreArgumentsCompatible(const FTG_Argument& ArgFrom, const FTG_Argument& ArgTo, FName& ConverterKey)
 {
 	ConverterKey = FName();
 	auto AT = ArgFrom.GetCPPTypeName();
 	auto BT = ArgTo.GetCPPTypeName();
+
 	if (AT == BT)
 		return true;
 
-	FName FromToName = MakeConvertKey(AT, BT);
+	FName FromToName = MakeConvertKey(ArgFrom, ArgTo);
 	auto Found = DefaultConverters.Find(FromToName);
 	if (Found)
 	{
diff --git a/Engine/Plugins/TextureGraph/Source/TextureGraph/Private/TG_Var.cpp b/Engine/Plugins/TextureGraph/Source/TextureGraph/Private/TG_Var.cpp
--- a/Engine/Plugins/TextureGraph/Source/TextureGraph/Private/TG_Var.cpp
+++ b/Engine/Plugins/TextureGraph/Source/TextureGraph/Private/TG_Var.cpp
@@ -16,7 +16,6 @@
 
 template <> FString TG_Var_LogValue(uint8& Value)
 {
-	
 	FString LogMessage = FString::FromInt(Value);
 	return LogMessage;
 }
@@ -481,7 +480,7 @@ void VarPropertySerializer_EnumProperty(FTG_Var::VarPropertySerialInfo& Info)
 	EnumProperty->GetValue_InContainer(Info.Owner, &PropValue);
 	auto enumVal = EnumProperty->GetEnum();
 
-	const uint8 VarValue = Info.Var->GetAs<uint8>();
+	const int32 VarValue = Info.Var->GetAs<int32>();
 
 	if (Info.CopyVarToProperty)
 	{
@@ -489,7 +488,7 @@ void VarPropertySerializer_EnumProperty(FTG_Var::VarPropertySerialInfo& Info)
 	}
 	else
 	{
-		Info.Var->EditAs<uint8>() = PropValue;
+		Info.Var->EditAs<int32>() = PropValue;
 	}
 }
 
diff --git a/Engine/Plugins/TextureGraph/Source/TextureGraph/Public/TG_GraphEvaluation.h b/Engine/Plugins/TextureGraph/Source/TextureGraph/Public/TG_GraphEvaluation.h
--- a/Engine/Plugins/TextureGraph/Source/TextureGraph/Public/TG_GraphEvaluation.h
+++ b/Engine/Plugins/TextureGraph/Source/TextureGraph/Public/TG_GraphEvaluation.h
@@ -180,6 +180,7 @@ struct TEXTUREGRAPH_API FTG_Evaluation
 	typedef TMap				<FName, VarConverter> ConverterMap;
 	static ConverterMap			DefaultConverters;
 	static FName				MakeConvertKey(FName From, FName To);
+	static FName				MakeConvertKey(const FTG_Argument& ArgFrom, const FTG_Argument& ArgTo);
 
 	// Conformer functors used to conform pin values
 	struct VarConformerInfo
diff --git a/Engine/Plugins/TextureGraph/Source/TextureGraph/Public/TG_Signature.h b/Engine/Plugins/TextureGraph/Source/TextureGraph/Public/TG_Signature.h
--- a/Engine/Plugins/TextureGraph/Source/TextureGraph/Public/TG_Signature.h
+++ b/Engine/Plugins/TextureGraph/Source/TextureGraph/Public/TG_Signature.h
@@ -33,6 +33,7 @@ enum class ETG_Access : uint8
 
 	PersistentSelfVarFlag = 0x10,	// the bit mask for just the Persistent Selfvar qualifier flag
 	NotConnectableFlag = 0x20,      // The bit mask for the Not Connectable qualifier flag
+	EnumFlag = 0x40,				// the bit mask for telling whether an argument is an enum or not
 };
 
 #define AccessToUint8(AccessEnum) static_cast<uint8>(ETG_Access::AccessEnum)
@@ -63,6 +64,7 @@ struct TEXTUREGRAPH_API FTG_ArgumentType
 	void			SetPersistentSelfVar() { Flags = Flags | AccessToUint8(PersistentSelfVarFlag); }
 
 	bool			IsNotConnectable() const { return (Flags & AccessToUint8(NotConnectableFlag)) != 0; }
+	bool			IsEnum() const { return (Flags & AccessToUint8(EnumFlag)) != 0; }
 	void			SetNotConnectable() { Flags = Flags | AccessToUint8(NotConnectableFlag); }
 
 
diff --git a/Engine/Plugins/TextureGraph/Source/TextureGraph/Public/TG_Var.h b/Engine/Plugins/TextureGraph/Source/TextureGraph/Public/TG_Var.h
--- a/Engine/Plugins/TextureGraph/Source/TextureGraph/Public/TG_Var.h
+++ b/Engine/Plugins/TextureGraph/Source/TextureGraph/Public/TG_Var.h
@@ -217,6 +217,7 @@ private:
 template <> FString TG_Var_LogValue(bool& Value);
 template <> FString TG_Var_LogValue(int& Value);
 template <> FString TG_Var_LogValue(float& Value);
+template <> FString TG_Var_LogValue(uint8& Value);
 template <> FString TG_Var_LogValue(FLinearColor& Value);
 template <> FString TG_Var_LogValue(FVector4f& Value);
 template <> FString TG_Var_LogValue(FVector2f& Value);
diff --git a/Engine/Plugins/TextureGraph/Source/TextureGraphEditor/Private/EdGraph/TG_EditorGraphPanelPinFactory.cpp b/Engine/Plugins/TextureGraph/Source/TextureGraphEditor/Private/EdGraph/TG_EditorGraphPanelPinFactory.cpp
--- a/Engine/Plugins/TextureGraph/Source/TextureGraphEditor/Private/EdGraph/TG_EditorGraphPanelPinFactory.cpp
+++ b/Engine/Plugins/TextureGraph/Source/TextureGraphEditor/Private/EdGraph/TG_EditorGraphPanelPinFactory.cpp
@@ -169,10 +169,19 @@ TSharedPtr<class SGraphPin> FTG_EditorGraphPanelPinFactory::CreatePin(class UEdG
 	return nullptr;
 }
 
+class TG_GraphPinEnum : public SGraphPinEnum
+{
+public:
+	virtual EVisibility GetDefaultValueVisibility() const override
+	{
+		return EVisibility::Visible;
+	}
+};
+
 TSharedPtr<class SGraphPin> FTG_EditorGraphPanelPinFactory::CreateEnumPin(class UEdGraphPin* InPin) const
 {
-	TSharedPtr<SGraphPin> PinEnum = SNew(SGraphPinEnum, InPin);
-	
+	TSharedPtr<SGraphPin> PinEnum = SNew(TG_GraphPinEnum, InPin);
+
 	//The default desired minimum width of SGraphPinEnum is 150 and there is no control to set the width
 	//We will try to find Widgets by type in hierarchy of SGraphPinEnum and update the min desired width
 	TSharedPtr<SWidget> FoundComboButton = FindWidgetInChildren(PinEnum, "SComboButton");
