Build fix

diff --git a/Engine/Plugins/TextureGraph/Source/TextureGraph/Private/Expressions/TG_Expression.cpp b/Engine/Plugins/TextureGraph/Source/TextureGraph/Private/Expressions/TG_Expression.cpp
--- a/Engine/Plugins/TextureGraph/Source/TextureGraph/Private/Expressions/TG_Expression.cpp
+++ b/Engine/Plugins/TextureGraph/Source/TextureGraph/Private/Expressions/TG_Expression.cpp
@@ -72,6 +72,7 @@ FTG_Signature::FInit UTG_Expression::GetSignatureInitArgsFromClass() const
 		bool bIsEnum = false;
 #else
 		FString PropertyTGType;
+		bool bIsEnum = false;
 #endif
 
 		// Add NotConnectable for certain types:
@@ -80,7 +81,6 @@ FTG_Signature::FInit UTG_Expression::GetSignatureInitArgsFromClass() const
 			bIsNotConnectable |= (ByteProperty != nullptr);
 			FEnumProperty* EnumProperty = CastField<FEnumProperty>(*Prop);
 			bIsEnum = (EnumProperty != nullptr);
-			//bIsNotConnectable |= (EnumProperty != nullptr);
 			FNameProperty* NameProperty = CastField<FNameProperty>(*Prop);
 			bIsNotConnectable |= (NameProperty != nullptr);
 		}
