UE-254649 - fix

diff --git a/Engine/Plugins/TextureGraph/Source/TextureGraph/Private/Expressions/Arrays/TG_Expression_ArrayOps.cpp b/Engine/Plugins/TextureGraph/Source/TextureGraph/Private/Expressions/Arrays/TG_Expression_ArrayOps.cpp
--- a/Engine/Plugins/TextureGraph/Source/TextureGraph/Private/Expressions/Arrays/TG_Expression_ArrayOps.cpp
+++ b/Engine/Plugins/TextureGraph/Source/TextureGraph/Private/Expressions/Arrays/TG_Expression_ArrayOps.cpp
@@ -67,8 +67,8 @@ void UTG_Expression_ArraySplit::Evaluate(FTG_EvaluationContext* InContext)
 	check(StartIndex <= EndIndex);
 
 	int32 InputCount = Input.Num();
-	int32 SlicedCount = EndIndex - StartIndex;
-	int32 SplicedCount = StartIndex + (InputCount - EndIndex);
+	int32 SlicedCount = std::min(EndIndex - StartIndex, InputCount - 1);
+	int32 SplicedCount = std::max(StartIndex + (InputCount - EndIndex), 0);
 
 	if (SlicedCount > 0)
 	{
diff --git a/Engine/Plugins/TextureGraph/Source/TextureGraph/Private/Expressions/Maths/TG_Expression_Maths_TwoInputs.cpp b/Engine/Plugins/TextureGraph/Source/TextureGraph/Private/Expressions/Maths/TG_Expression_Maths_TwoInputs.cpp
--- a/Engine/Plugins/TextureGraph/Source/TextureGraph/Private/Expressions/Maths/TG_Expression_Maths_TwoInputs.cpp
+++ b/Engine/Plugins/TextureGraph/Source/TextureGraph/Private/Expressions/Maths/TG_Expression_Maths_TwoInputs.cpp
@@ -80,7 +80,8 @@ FVector4f UTG_Expression_Dot::EvaluateVector_WithValue(FTG_EvaluationContext* In
 
 FTG_Texture	UTG_Expression_Dot::EvaluateTexture(FTG_EvaluationContext* InContext)
 {
-	return T_Maths_TwoInputs::CreateDot(InContext->Cycle, Output.EditTexture().GetBufferDescriptor(), InContext->TargetId, Input1.GetTexture(InContext), Input2.GetTexture(InContext));
+	auto Desc = Output.EditTexture().GetBufferDescriptor();
+	return T_Maths_TwoInputs::CreateDot(InContext->Cycle, Desc, InContext->TargetId, Input1.GetTexture(InContext), Input2.GetTexture(InContext));
 }
 
 FTG_Variant::EType	UTG_Expression_Dot::EvalExpressionCommonOutputVariantType() const
diff --git a/Engine/Plugins/TextureGraph/Source/TextureGraph/Private/Expressions/TG_Expression.cpp b/Engine/Plugins/TextureGraph/Source/TextureGraph/Private/Expressions/TG_Expression.cpp
--- a/Engine/Plugins/TextureGraph/Source/TextureGraph/Private/Expressions/TG_Expression.cpp
+++ b/Engine/Plugins/TextureGraph/Source/TextureGraph/Private/Expressions/TG_Expression.cpp
@@ -290,12 +290,7 @@ void UTG_Expression::SetupAndEvaluate(FTG_EvaluationContext* InContext)
 	// Copy Output var values TO this expression's matching properties (Mainly to copy existing thumbs for Output FTG_Textures)
 	for (auto& OutputVar : InContext->Outputs.VarArguments)
 	{
-		if (OutputVar.Value.Argument.IsTexture())
-		{
-			FTG_Texture& Texture = OutputVar.Value.Var->GetAs<FTG_Texture>();
-			Texture.Descriptor = BaseOutputSettings;
-		}
-		else if (OutputVar.Value.Argument.IsVariant())
+		if (OutputVar.Value.Argument.IsVariant())
 		{
 			FTG_Variant& Variant = OutputVar.Value.Var->GetAs<FTG_Variant>();
 			if (Variant.IsTexture())
@@ -304,6 +299,11 @@ void UTG_Expression::SetupAndEvaluate(FTG_EvaluationContext* InContext)
 				Texture.Descriptor = BaseOutputSettings;
 			}
 		}
+		else if (OutputVar.Value.Argument.IsTexture())
+		{
+			FTG_Texture& Texture = OutputVar.Value.Var->GetAs<FTG_Texture>();
+			Texture.Descriptor = BaseOutputSettings;
+		}
 		OutputVar.Value.Var->CopyTo(this, OutputVar.Value.Argument);
 	}
 
diff --git a/Engine/Plugins/TextureGraph/Source/TextureGraphEngine/Job/Job.cpp b/Engine/Plugins/TextureGraph/Source/TextureGraphEngine/Job/Job.cpp
--- a/Engine/Plugins/TextureGraph/Source/TextureGraphEngine/Job/Job.cpp
+++ b/Engine/Plugins/TextureGraph/Source/TextureGraphEngine/Job/Job.cpp
@@ -464,7 +464,13 @@ TiledBlobPtr Job::InitResult(FString InNewName, const BufferDescriptor* InDesire
 
 	ResultDesc.bIsTransient = IsDiscard();
 	ResultDesc.Name = NewName;
-	
+
+	if ((int)ResultDesc.Width < (int)EResolution::Auto || (int)ResultDesc.Width > (int)EResolution::Resolution8192)
+		ResultDesc.Width = (int)EResolution::Auto;
+
+	if ((int)ResultDesc.Height < (int)EResolution::Auto || (int)ResultDesc.Height > (int)EResolution::Resolution8192)
+		ResultDesc.Height = (int)EResolution::Auto;
+
 	/// Update the job name as well
 	Name = NewName;
 
