Skeletal mesh crash fix

diff --git a/Engine/Plugins/TextureGraph/Source/TextureGraphEngine/3D/RenderMesh_Editor.cpp b/Engine/Plugins/TextureGraph/Source/TextureGraphEngine/3D/RenderMesh_Editor.cpp
--- a/Engine/Plugins/TextureGraph/Source/TextureGraphEngine/3D/RenderMesh_Editor.cpp
+++ b/Engine/Plugins/TextureGraph/Source/TextureGraphEngine/3D/RenderMesh_Editor.cpp
@@ -18,10 +18,11 @@ RenderMesh_Editor::RenderMesh_Editor(RenderMesh* parent, TArray<MeshInfoPtr> mes
 
 }
 
-RenderMesh_Editor::RenderMesh_Editor(UStaticMeshComponent* staticMeshComponent,UWorld* world)
+RenderMesh_Editor::RenderMesh_Editor(UStaticMeshComponent* staticMeshComponent, UWorld* world)
 {
-	_world = world;
-	_meshComponents.Add(staticMeshComponent);
+	World = world;
+	if (staticMeshComponent)
+		MeshComponents.Add(staticMeshComponent);
 }
 
 void RenderMesh_Editor::PrepareForRendering(UWorld* world, FVector scale)
@@ -37,75 +38,75 @@ AsyncActionResultPtr RenderMesh_Editor::Load()
 	// Need to figure out how we are going to support MTS and UDIM meshes.
 	_meshSplitType = MeshSplitType::Single;
 
-	//_meshActor->GetComponents<UStaticMeshComponent>(_meshComponents);
-	if (_meshComponents.Num() <= 0)
+	//_meshActor->GetComponents<UStaticMeshComponent>(MeshComponents);
+	if (MeshComponents.Num() <= 0)
 		return cti::make_exceptional_continuable<ActionResultPtr>(std::make_exception_ptr(std::runtime_error("No mesh components found on actor")));
 
 	return cti::make_continuable<ActionResultPtr>([this](auto&& promise)
 		{
 			AsyncTask(ENamedThreads::AnyBackgroundThreadNormalTask, [this, promise = std::forward<decltype(promise)>(promise)]() mutable
-			{
-				try
 				{
-					LoadInternal();
-					promise.set_value(std::make_shared<ActionResult>());
-				}
-				catch (cti::exception_t ex)
-				{
-					promise.set_exception(ex);
-				}
-			});
+					try
+					{
+						LoadInternal();
+						promise.set_value(std::make_shared<ActionResult>());
+					}
+					catch (cti::exception_t ex)
+					{
+						promise.set_exception(ex);
+					}
+				});
 		});
 }
 
 void RenderMesh_Editor::SetMaterial(UMaterialInterface* material)
 {
-	for (int meshIndex = 0; meshIndex < _meshComponents.Num(); meshIndex++)
+	for (int MeshIndex = 0; MeshIndex < MeshComponents.Num(); MeshIndex++)
 	{
-		UStaticMeshComponent* meshComponent = _meshComponents[meshIndex];
-		if (meshComponent->IsValidLowLevel())
+		UStaticMeshComponent* MeshComponent = MeshComponents[MeshIndex];
+		if (MeshComponent->IsValidLowLevel())
 		{
-			MeshInfoPtr meshInfo = _meshes[meshIndex];
-			meshComponent->SetMaterial(meshInfo->GetMaterialIndex(), material);
+			MeshInfoPtr MeshInfo = _meshes[MeshIndex];
+			MeshComponent->SetMaterial(MeshInfo->GetMaterialIndex(), material);
 		}
 	}
 }
 
 void RenderMesh_Editor::LoadInternal()
 {
-	for (UStaticMeshComponent* meshComponent : _meshComponents)
+	for (UStaticMeshComponent* MeshComponent : MeshComponents)
 	{
-		LoadSingleMeshComponent(*meshComponent);
+		LoadSingleMeshComponent(*MeshComponent);
 	}
 }
-void RenderMesh_Editor::LoadSingleMeshComponent(const UStaticMeshComponent& meshComponent)
+void RenderMesh_Editor::LoadSingleMeshComponent(const UStaticMeshComponent& MeshComponent)
 {
-	UStaticMesh* staticMesh = meshComponent.GetStaticMesh();
+	UStaticMesh* StaticMesh = MeshComponent.GetStaticMesh();
 
 	//Set bAllowCPUAccess to true as it is required to use  GetSectionFromStaticMesh
-	staticMesh->bAllowCPUAccess = true;
+	StaticMesh->bAllowCPUAccess = true;
 
 	const int32 LODindex = 0;
-	int32 numSections = staticMesh->GetNumSections(LODindex);
+	int32 numSections = StaticMesh->GetNumSections(LODindex);
 
 	// Return if RenderData is invalid
-	if (!staticMesh->GetRenderData())
+	if (!StaticMesh->GetRenderData())
 		return;
 
 	// No valid mesh data on lod 0 (shouldn't happen)
-	if (!staticMesh->GetRenderData()->LODResources.IsValidIndex(LODindex))
+	if (!StaticMesh->GetRenderData()->LODResources.IsValidIndex(LODindex))
 		return;
 
 	// load materials
-	//const TArray<UMaterialInterface*> materialInterfaces = meshComponent.GetMaterials();
-	for (int32 materialIndex = 0; materialIndex < meshComponent.GetNumMaterials(); ++materialIndex)
+	//const TArray<UMaterialInterface*> materialInterfaces = MeshComponent.GetMaterials();
+	for (int32 materialIndex = 0; materialIndex < MeshComponent.GetNumMaterials(); ++materialIndex)
 	{
-		UMaterialInterface* materialInterface = meshComponent.GetMaterial(materialIndex);
-		if (materialInterface)
+		UMaterialInterface* MaterialInterface = MeshComponent.GetMaterial(materialIndex);
+		if (MaterialInterface)
 		{
-			FString matName;
-			materialInterface->GetName(matName);
-			AddMaterialInfo(materialIndex, matName);
+			FString MatName;
+			MaterialInterface->GetName(MatName);
+			AddMaterialInfo(materialIndex, MatName);
 		}
 	}
 	_currentMaterials.Empty();
@@ -118,23 +119,23 @@ void RenderMesh_Editor::LoadSingleMeshComponent(const UStaticMeshComponent& mesh
 	TArray<FProcMeshTangent> Tangents;
 
 	_originalBounds.Init();
-	
+
 	for (int32 sectionIndex = 0; sectionIndex < numSections; sectionIndex++)
 	{
 		CoreMesh* cmesh = new CoreMesh();
 		CoreMeshPtr cmeshPtr = CoreMeshPtr(cmesh);
-		MeshInfoPtr meshInfo = std::make_shared<MeshInfo>(cmeshPtr);
+		MeshInfoPtr MeshInfo = std::make_shared<::MeshInfo>(cmeshPtr);
 
-		cmesh->bounds = staticMesh->GetBoundingBox();
-		meshComponent.GetName(cmesh->name);
+		cmesh->bounds = StaticMesh->GetBoundingBox();
+		MeshComponent.GetName(cmesh->name);
 
-		FMeshSectionInfo sectionInfo = staticMesh->GetSectionInfoMap().Get(LODindex, sectionIndex);
-		UKismetProceduralMeshLibrary::GetSectionFromStaticMesh(staticMesh, LODindex, sectionIndex, cmesh->vertices, cmesh->triangles, cmesh->normals, cmesh->uvs, cmesh->tangents);
+		FMeshSectionInfo sectionInfo = StaticMesh->GetSectionInfoMap().Get(LODindex, sectionIndex);
+		UKismetProceduralMeshLibrary::GetSectionFromStaticMesh(StaticMesh, LODindex, sectionIndex, cmesh->vertices, cmesh->triangles, cmesh->normals, cmesh->uvs, cmesh->tangents);
 		cmesh->materialIndex = sectionInfo.MaterialIndex;
-		
-		MathUtils::EncapsulateBound(_originalBounds, cmeshPtr->bounds);	
 
-		_meshes.Add(meshInfo);
+		MathUtils::EncapsulateBound(_originalBounds, cmeshPtr->bounds);
+
+		_meshes.Add(MeshInfo);
 	}
 }
 #endif // WITH_EDITOR
\ No newline at end of file
diff --git a/Engine/Plugins/TextureGraph/Source/TextureGraphEngine/3D/RenderMesh_Editor.h b/Engine/Plugins/TextureGraph/Source/TextureGraphEngine/3D/RenderMesh_Editor.h
--- a/Engine/Plugins/TextureGraph/Source/TextureGraphEngine/3D/RenderMesh_Editor.h
+++ b/Engine/Plugins/TextureGraph/Source/TextureGraphEngine/3D/RenderMesh_Editor.h
@@ -11,8 +11,8 @@ class TEXTUREGRAPHENGINE_API RenderMesh_Editor: public RenderMesh
 {
 
 protected:
-	TArray<UStaticMeshComponent*>			_meshComponents;
-	UWorld*									_world = nullptr;
+	TArray<UStaticMeshComponent*>			MeshComponents;
+	UWorld*									World = nullptr;
 	virtual void							LoadInternal() override;
 	void									LoadSingleMeshComponent(const UStaticMeshComponent& mesh);
 
@@ -21,9 +21,9 @@ public:
 											RenderMesh_Editor(RenderMesh* parent, TArray<MeshInfoPtr> meshes, MaterialInfoPtr matInfo);
 											RenderMesh_Editor(UStaticMeshComponent* staticMeshComponent, UWorld* world);
 
-	FORCEINLINE								TArray<UStaticMeshComponent*> GetMeshComponents() { return _meshComponents; }
+	FORCEINLINE								TArray<UStaticMeshComponent*> GetMeshComponents() { return MeshComponents; }
 	virtual AsyncActionResultPtr			Load() override;
-	virtual FMatrix							LocalToWorldMatrix() const override { return _meshComponents[0]->GetComponentTransform().ToMatrixWithScale(); }
+	virtual FMatrix							LocalToWorldMatrix() const override { return MeshComponents[0]->GetComponentTransform().ToMatrixWithScale(); }
 	virtual void							PrepareForRendering(UWorld* world, FVector scale) override;
 	virtual void							SetMaterial(UMaterialInterface* material) override;
 };
diff --git a/Engine/Plugins/TextureGraph/Source/TextureGraphEngine/3D/RenderMesh_EditorScene.cpp b/Engine/Plugins/TextureGraph/Source/TextureGraphEngine/3D/RenderMesh_EditorScene.cpp
--- a/Engine/Plugins/TextureGraph/Source/TextureGraphEngine/3D/RenderMesh_EditorScene.cpp
+++ b/Engine/Plugins/TextureGraph/Source/TextureGraphEngine/3D/RenderMesh_EditorScene.cpp
@@ -16,9 +16,9 @@ RenderMesh_EditorScene::RenderMesh_EditorScene(RenderMesh* parent, TArray<MeshIn
 }
 
 RenderMesh_EditorScene::RenderMesh_EditorScene(AActor* actor) :
-	_meshActor(actor)
+	MeshActor(actor)
 {
-	_meshActor->GetComponents<UStaticMeshComponent>(_meshComponents);
+	MeshActor->GetComponents<UStaticMeshComponent>(MeshComponents);
 }
 
 void RenderMesh_EditorScene::PrepareForRendering(UWorld* world, FVector scale)
@@ -35,9 +35,9 @@ AsyncActionResultPtr RenderMesh_EditorScene::Load()
 	// Need to figure out how we are going to support MTS and UDIM meshes.
 	_meshSplitType = MeshSplitType::Single;
 
-	if (_meshComponents.Num() == 0)
+	if (MeshComponents.Num() == 0)
 	{
-		_meshActor->GetComponents<UStaticMeshComponent>(_meshComponents);
+		MeshActor->GetComponents<UStaticMeshComponent>(MeshComponents);
 	}
 
 	if (GetMeshComponents().Num() <= 0)
@@ -52,7 +52,7 @@ AsyncActionResultPtr RenderMesh_EditorScene::Load()
 				{
 					LoadInternal();
 					promise.set_value(ActionResultPtr(new ActionResult(nullptr)));
-					_meshActors.Add(_meshActor);
+					_meshActors.Add(MeshActor);
 				}
 				catch (std::exception_ptr ex)
 				{
diff --git a/Engine/Plugins/TextureGraph/Source/TextureGraphEngine/3D/RenderMesh_EditorScene.h b/Engine/Plugins/TextureGraph/Source/TextureGraphEngine/3D/RenderMesh_EditorScene.h
--- a/Engine/Plugins/TextureGraph/Source/TextureGraphEngine/3D/RenderMesh_EditorScene.h
+++ b/Engine/Plugins/TextureGraph/Source/TextureGraphEngine/3D/RenderMesh_EditorScene.h
@@ -9,7 +9,7 @@ struct CoreMesh;
 class TEXTUREGRAPHENGINE_API RenderMesh_EditorScene: public RenderMesh_Editor
 {
 private:
-	AActor*									_meshActor;				/// Actor that this rendermesh is associated to
+	AActor*									MeshActor;				/// Actor that this rendermesh is associated to
 
 public:
 											RenderMesh_EditorScene();
