Fixing out of memory errors with TG.

diff --git a/Engine/Plugins/TextureGraph/Source/TextureGraphEngine/Job/Scheduler.cpp b/Engine/Plugins/TextureGraph/Source/TextureGraphEngine/Job/Scheduler.cpp
--- a/Engine/Plugins/TextureGraph/Source/TextureGraphEngine/Job/Scheduler.cpp
+++ b/Engine/Plugins/TextureGraph/Source/TextureGraphEngine/Job/Scheduler.cpp
@@ -9,6 +9,7 @@
 DECLARE_CYCLE_STAT(TEXT("Scheduler_Update"), STAT_Scheduler_Update, STATGROUP_TextureGraphEngine);
 
 const double Scheduler::IdleTimeInterval = 500.0;
+const double Scheduler::IdleBusyTimeInterval = 1000.0;
 const double Scheduler::IdleBatchTimeLimit = 100.0;
 const double Scheduler::CurrentBatchWarningLimit = 2000.0;
 
@@ -59,9 +60,20 @@ void Scheduler::Update(float dt)
 	if (TimeSinceIdle < 0.001)
 		TimeSinceIdle = Util::Time();
 
+	if (TimeSinceIdleBatchRan < 0.001)
+		TimeSinceIdleBatchRan = Util::Time();
+
 	if (!bIsRunning)
 		Start();
 
+	double Delta = Util::TimeDelta(TimeSinceIdleBatchRan);
+	if (Delta > IdleBusyTimeInterval)
+	{
+		UE_LOG(LogBatch, Warning, TEXT("Scheduler::Update"));
+		UpdateIdle(true);
+		ObserverSource->UpdateIdle();
+	}
+
 	{
 		FScopeLock Lock(&CurrentBatchMutex);
 		if (CurrentBatch)
@@ -191,10 +203,10 @@ void Scheduler::Update(float dt)
 	}
 	else
 	{
-		double Delta = Util::TimeDelta(TimeSinceIdle);
+		Delta = Util::TimeDelta(TimeSinceIdle);
 		if (Delta > IdleTimeInterval)
 		{
-			UpdateIdle();
+			UpdateIdle(false);
 
 			// Since we have time, notify the observer update
 			ObserverSource->UpdateIdle();
@@ -231,11 +243,12 @@ AsyncJobResultPtr Scheduler::UpdateIdleBatch(size_t Index)
 		});
 }
 
-void Scheduler::UpdateIdle()
+void Scheduler::UpdateIdle(bool bForce)
 {
 	if (bIsDisableIdle)
 		return;
 
+	if (!bForce)
 	{
 		FScopeLock Lock(&CurrentBatchMutex);
 
@@ -249,6 +262,7 @@ void Scheduler::UpdateIdle()
 		return;
 
 	// if no Batch to run, check if we have an idle Batch we should run
+	if (!bForce)
 	{
 		/// Shouldn't run idle services if a Batch is still being processed
 		FScopeLock Lock(&CurrentBatchMutex);
@@ -256,6 +270,7 @@ void Scheduler::UpdateIdle()
 			return;
 	}
 
+	if (!bForce)
 	{
 		/// Shouldn't run idle services if we still have batches to process
 		FScopeLock Lock(&BatchMutex);
@@ -269,6 +284,7 @@ void Scheduler::UpdateIdle()
 		return;
 
 	size_t NumIdleBatches = 0;
+	TimeSinceIdleBatchRan = Util::Time();
 
 	{
 		FScopeLock Lock(&IdlServiceMutex);
diff --git a/Engine/Plugins/TextureGraph/Source/TextureGraphEngine/Job/Scheduler.h b/Engine/Plugins/TextureGraph/Source/TextureGraphEngine/Job/Scheduler.h
--- a/Engine/Plugins/TextureGraph/Source/TextureGraphEngine/Job/Scheduler.h
+++ b/Engine/Plugins/TextureGraph/Source/TextureGraphEngine/Job/Scheduler.h
@@ -43,6 +43,7 @@ class TEXTUREGRAPHENGINE_API Scheduler
 {
 private:
 	static const double					IdleTimeInterval;				/// Start an Idle Batch after <this many> milliseconds
+	static const double					IdleBusyTimeInterval;			/// Start an Idle Batch (during busy times) after <this many> milliseconds
 	static const double					IdleBatchTimeLimit;				/// Start an Idle Batch after <this many> milliseconds
 	static const double					CurrentBatchWarningLimit;		/// Scheduler starts spewing out warnings if the current batch doesn't finish ater <this many> milliseconds
 
@@ -82,6 +83,7 @@ private:
 	JobBatchPtr							CurrentHighPriorityBatch;		/// The current high priority batch job
 
 	double								TimeSinceIdle = 0;				/// time since scheduler has been idle
+	double								TimeSinceIdleBatchRan = 0;		/// time since scheduler has been idle
 
 	bool								bCaptureNextBatch = false;		/// set the next incoming batch to get captured by RenderDoc
 	bool								bIsDisableIdle = false;			/// Disable idle services or not
@@ -93,7 +95,7 @@ private:
 	void								StopServices();
 
 	AsyncJobResultPtr					UpdateIdleBatch(size_t index);
-	void								UpdateIdle();
+	void								UpdateIdle(bool bForce);
 
 public:
 										Scheduler();
