Fix crash on material function undo

diff --git a/Engine/Plugins/TextureGraph/Source/TextureGraph/Private/TG_Graph.cpp b/Engine/Plugins/TextureGraph/Source/TextureGraph/Private/TG_Graph.cpp
--- a/Engine/Plugins/TextureGraph/Source/TextureGraph/Private/TG_Graph.cpp
+++ b/Engine/Plugins/TextureGraph/Source/TextureGraph/Private/TG_Graph.cpp
@@ -974,26 +974,30 @@ FTG_Ids UTG_Graph::GatherSourceNodes(const UTG_Node* InNode) const
 		bool AllInConnected = true;
 		bool AllInDisconnected = true;
 		auto InPinIds = InNode->GetInputPinIds();
-		for (auto opi : InPinIds) // Over all the Input pins, find the node feeding them
+
+		for (auto InPinId : InPinIds) // Over all the Input pins, find the node feeding them
 		{
-			const UTG_Pin* Pin = InNode->Pins[opi.PinIdx()];
-			if (Pin->GetEdges().IsEmpty())
-			{
-				AllInConnected = false;
-			}
-			else
+			if (InPinId.PinIdx() < InNode->Pins.Num())
 			{
-				AllInDisconnected = false;
-				auto EdgePinId = Pin->GetEdges()[0];
-				const UTG_Pin* SourcePin = GetPin(EdgePinId);
-
-				if (SourcePin)
+				const UTG_Pin* Pin = InNode->Pins[InPinId.PinIdx()];
+				if (Pin->GetEdges().IsEmpty())
 				{
-					check(SourcePin);
+					AllInConnected = false;
+				}
+				else
+				{
+					AllInDisconnected = false;
+					auto EdgePinId = Pin->GetEdges()[0];
+					const UTG_Pin* SourcePin = GetPin(EdgePinId);
 
-					// only add the source node if not already there
-					if (SourceNodes.Find(SourcePin->GetNodeId()) == INDEX_NONE)
-						SourceNodes.Emplace(SourcePin->GetNodeId());
+					if (SourcePin)
+					{
+						check(SourcePin);
+
+						// only add the source node if not already there
+						if (SourceNodes.Find(SourcePin->GetNodeId()) == INDEX_NONE)
+							SourceNodes.Emplace(SourcePin->GetNodeId());
+					}
 				}
 			}
 		}
diff --git a/Engine/Plugins/TextureGraph/Source/TextureGraph/Private/TG_Node.cpp b/Engine/Plugins/TextureGraph/Source/TextureGraph/Private/TG_Node.cpp
--- a/Engine/Plugins/TextureGraph/Source/TextureGraph/Private/TG_Node.cpp
+++ b/Engine/Plugins/TextureGraph/Source/TextureGraph/Private/TG_Node.cpp
@@ -311,7 +311,8 @@ int32 UTG_Node::ForEachOutputPins(std::function<void(const UTG_Pin* /*pin*/, uin
 	int32 Num = GetSignature().GetOutArguments().Num();
 	for (int32 i = 0; i < Num; ++i)
 	{
-		visitor(Pins[i + Offset], i);
+		if (i + Offset < Pins.Num())
+			visitor(Pins[i + Offset], i);
 	}
 	return Num;
 }
diff --git a/Engine/Plugins/TextureGraph/Source/TextureGraphEngine/Job/Scheduler.cpp b/Engine/Plugins/TextureGraph/Source/TextureGraphEngine/Job/Scheduler.cpp
--- a/Engine/Plugins/TextureGraph/Source/TextureGraphEngine/Job/Scheduler.cpp
+++ b/Engine/Plugins/TextureGraph/Source/TextureGraphEngine/Job/Scheduler.cpp
@@ -69,7 +69,6 @@ void Scheduler::Update(float dt)
 	double Delta = Util::TimeDelta(TimeSinceIdleBatchRan);
 	if (Delta > IdleBusyTimeInterval)
 	{
-		UE_LOG(LogBatch, Warning, TEXT("Scheduler::Update"));
 		UpdateIdle(true);
 		ObserverSource->UpdateIdle();
 	}
