Fix Metahumans crash

diff --git a/Engine/Plugins/TextureGraph/Source/TextureGraphEngine/Export/TextureExporter.cpp b/Engine/Plugins/TextureGraph/Source/TextureGraphEngine/Export/TextureExporter.cpp
--- a/Engine/Plugins/TextureGraph/Source/TextureGraphEngine/Export/TextureExporter.cpp
+++ b/Engine/Plugins/TextureGraph/Source/TextureGraphEngine/Export/TextureExporter.cpp
@@ -178,12 +178,25 @@ AsyncInt TextureExporter::ExportRawAsUAsset(RawBufferPtr RawObj,const FExportMap
 		}
 
 		const uint8* RawBufferData = RawObj->GetData();
+		size_t RawBufferSize = RawObj->GetUnpaddedSize();
 		UPackage* Package = CreatePackage( *PackageName);
 		Package->FullyLoad();
 		UTexture2D* NewTexture = nullptr;
 		int32 WidthToUse = Setting.Width == 0 ? RawObj->Width() : Setting.Width;
 		int32 HeightToUse = Setting.Height == 0 ? RawObj->Height() : Setting.Height;
 		
+		const int32 NumBlocksX = WidthToUse / GPixelFormats[PixelFormat].BlockSizeX;
+		const int32 NumBlocksY = HeightToUse / GPixelFormats[PixelFormat].BlockSizeY;
+		uint32 BufferSizeRequired = NumBlocksX * GPixelFormats[PixelFormat].BlockBytes * NumBlocksY;
+
+		if (BufferSizeRequired > RawBufferSize)
+		{
+			UE_LOG(LogTexture, Warning, TEXT("Export dimensions are invalid and require buffer size larger than is available. Expected to be less than: %d [Asked: %d, %dx%d]. Clamping ..."), 
+				RawBufferSize, BufferSizeRequired, WidthToUse, HeightToUse);
+			WidthToUse = RawObj->Width();
+			HeightToUse = RawObj->Height();
+		}
+
 		if ((WidthToUse % GPixelFormats[PixelFormat].BlockSizeX) == 0 &&
 			(HeightToUse % GPixelFormats[PixelFormat].BlockSizeY) == 0)
 		{
