Fix to the vector value in node preview

diff --git a/Engine/Plugins/TextureGraph/Source/TextureGraphEditor/Private/STG_NodePreview.cpp b/Engine/Plugins/TextureGraph/Source/TextureGraphEditor/Private/STG_NodePreview.cpp
--- a/Engine/Plugins/TextureGraph/Source/TextureGraphEditor/Private/STG_NodePreview.cpp
+++ b/Engine/Plugins/TextureGraph/Source/TextureGraphEditor/Private/STG_NodePreview.cpp
@@ -423,7 +423,7 @@ void STG_NodePreviewWidget::Update() const
 		}
 		else
 		{
-			NodeViewer->SetTexture(TextureHelper::GetBlack());
+			NodeViewer->SetTexture(nullptr, FLinearColor::Black);
 		}
 
 		NodeViewer->SetLabelText(GetLabelText(Variant, bValidVariant));
