Scheduler locking fix

diff --git a/Engine/Plugins/TextureGraph/Source/TextureGraphEngine/Job/DeviceTransferService.h b/Engine/Plugins/TextureGraph/Source/TextureGraphEngine/Job/DeviceTransferService.h
--- a/Engine/Plugins/TextureGraph/Source/TextureGraphEngine/Job/DeviceTransferService.h
+++ b/Engine/Plugins/TextureGraph/Source/TextureGraphEngine/Job/DeviceTransferService.h
@@ -33,6 +33,7 @@ public:
 
 	virtual AsyncJobResultPtr		Tick() override;
 	virtual void					Stop() override;
+	virtual bool					ShouldRunDuringBusyUpdate() const override { return true; }
 
 	virtual	void					QueueTransfer(BlobPtr BlobObj, Device* TargetDevice);
 	virtual void					AbortTransfer(HashType HashValue);
diff --git a/Engine/Plugins/TextureGraph/Source/TextureGraphEngine/Job/IdleService.h b/Engine/Plugins/TextureGraph/Source/TextureGraphEngine/Job/IdleService.h
--- a/Engine/Plugins/TextureGraph/Source/TextureGraphEngine/Job/IdleService.h
+++ b/Engine/Plugins/TextureGraph/Source/TextureGraphEngine/Job/IdleService.h
@@ -30,6 +30,7 @@ public:
 	virtual							~IdleService() { }
 	virtual AsyncJobResultPtr		Tick() = 0;
 	virtual void					Stop() = 0;
+	virtual bool					ShouldRunDuringBusyUpdate() const { return false; }
 
 	void							UpdateStats(double StartTime, double EndTime, bool DidOffendTimeLimit);
 
diff --git a/Engine/Plugins/TextureGraph/Source/TextureGraphEngine/Job/Job.cpp b/Engine/Plugins/TextureGraph/Source/TextureGraphEngine/Job/Job.cpp
--- a/Engine/Plugins/TextureGraph/Source/TextureGraphEngine/Job/Job.cpp
+++ b/Engine/Plugins/TextureGraph/Source/TextureGraphEngine/Job/Job.cpp
@@ -633,6 +633,12 @@ bool Job::CheckCulled(JobRunInfo InRunInfo)
 	if (!Result || Result->IsLateBound())
 		return false;
 
+	return false;
+
+	/// There is a bug in this code. It basically (incorrectly) treats a full (non-tiled) blob and tries to 
+	/// cull the job that generates it. Disabling for the time being as it's not a particular performance
+	/// deterrant. But this needs to be tweaked for these kinds of jobs (like T_Transform for instance)
+#if 0
 	check(Result->IsPromise());
 	TiledBlob_PromisePtr ResultPromise = GetResultPromise();
 	size_t NumRows = ResultPromise->Rows();
@@ -683,6 +689,7 @@ bool Job::CheckCulled(JobRunInfo InRunInfo)
 	SetPromise(0);
 
 	return bIsCulled;
+#endif 
 }
 
 AsyncPrepareResult Job::PrepareTargets(JobBatch* Batch)
diff --git a/Engine/Plugins/TextureGraph/Source/TextureGraphEngine/Job/Scheduler.cpp b/Engine/Plugins/TextureGraph/Source/TextureGraphEngine/Job/Scheduler.cpp
--- a/Engine/Plugins/TextureGraph/Source/TextureGraphEngine/Job/Scheduler.cpp
+++ b/Engine/Plugins/TextureGraph/Source/TextureGraphEngine/Job/Scheduler.cpp
@@ -69,7 +69,7 @@ void Scheduler::Update(float dt)
 	double Delta = Util::TimeDelta(TimeSinceIdleBatchRan);
 	if (Delta > IdleBusyTimeInterval)
 	{
-		UpdateIdle(true);
+		UpdateIdle(true, true);
 		ObserverSource->UpdateIdle();
 	}
 
@@ -218,11 +218,13 @@ void Scheduler::SetCaptureRenderDocNextBatch(bool bCapture /* = true */)
 	bCaptureNextBatch = bCapture;
 }
 
-AsyncJobResultPtr Scheduler::UpdateIdleBatch(size_t Index)
+AsyncJobResultPtr Scheduler::UpdateIdleBatch(size_t Index, bool bIsBusyUpdate)
 {
 	check(Index < IdleServices.size());
 	IdleServicePtr Batch = IdleServices[Index];
 	double StartTime = Util::Time();
+	if (bIsBusyUpdate && !Batch->ShouldRunDuringBusyUpdate())
+		return cti::make_ready_continuable(std::make_shared<JobResult>());
 
 	return Batch->Tick().then([this, Batch, StartTime](JobResultPtr result)
 		{
@@ -242,7 +244,7 @@ AsyncJobResultPtr Scheduler::UpdateIdleBatch(size_t Index)
 		});
 }
 
-void Scheduler::UpdateIdle(bool bForce)
+void Scheduler::UpdateIdle(bool bForce, bool bIsBusyUpdate)
 {
 	if (bIsDisableIdle)
 		return;
@@ -304,7 +306,7 @@ void Scheduler::UpdateIdle(bool bForce)
 
 	for (size_t IdleBatchIndex = 0; IdleBatchIndex < NumIdleBatches; IdleBatchIndex++)
 	{
-		auto Promise = UpdateIdleBatch(IdleBatchIndex);
+		auto Promise = UpdateIdleBatch(IdleBatchIndex, bIsBusyUpdate);
 		Promises.push_back(std::move(Promise));
 	}
 
diff --git a/Engine/Plugins/TextureGraph/Source/TextureGraphEngine/Job/Scheduler.h b/Engine/Plugins/TextureGraph/Source/TextureGraphEngine/Job/Scheduler.h
--- a/Engine/Plugins/TextureGraph/Source/TextureGraphEngine/Job/Scheduler.h
+++ b/Engine/Plugins/TextureGraph/Source/TextureGraphEngine/Job/Scheduler.h
@@ -94,8 +94,8 @@ private:
 	void								Stop();
 	void								StopServices();
 
-	AsyncJobResultPtr					UpdateIdleBatch(size_t index);
-	void								UpdateIdle(bool bForce);
+	AsyncJobResultPtr					UpdateIdleBatch(size_t index, bool bIsBusyUpdate);
+	void								UpdateIdle(bool bForce, bool bIsBusyUpdate = false);
 
 public:
 										Scheduler();
