Texture graph blob caching issue fixes

diff --git a/Engine/Plugins/TextureGraph/Source/TextureGraphEditor/Private/STG_NodePreview.cpp b/Engine/Plugins/TextureGraph/Source/TextureGraphEditor/Private/STG_NodePreview.cpp
--- a/Engine/Plugins/TextureGraph/Source/TextureGraphEditor/Private/STG_NodePreview.cpp
+++ b/Engine/Plugins/TextureGraph/Source/TextureGraphEditor/Private/STG_NodePreview.cpp
@@ -38,13 +38,11 @@ void FNodeViewer::UpdateTexture()
 		CurrentBlob->OnFinalise()
 			.then([this]()
 				{
-					UE_LOG(LogTemp, Log, TEXT("!!!Finalised. Combining!!!"));
 					const TiledBlobPtr BlobTiled = std::static_pointer_cast<TiledBlob>(CurrentBlob);
 					return BlobTiled->CombineTiles(false, false);
 				})
 			.then([this]()
 				{
-					UE_LOG(LogTemp, Log, TEXT("!!!SETTING TEXTURE!!!"));
 					FIntPoint ImageSize = FIntPoint {(int32)CurrentBlob->GetWidth(), (int32)CurrentBlob->GetHeight()};
 					SetTexture(CurrentBlob);
 				});
@@ -397,13 +395,11 @@ void STG_NodePreviewWidget::Update() const
 				Blob->OnFinalise()
 					.then([Blob]()
 						{
-							UE_LOG(LogTemp, Log, TEXT("!!!Finalised. Combining!!!"));
 							const TiledBlobPtr BlobTiled = std::static_pointer_cast<TiledBlob>(Blob);
 							return BlobTiled->CombineTiles(false, false);
 						})
 					.then([this, Blob]()
 						{
-							UE_LOG(LogTemp, Log, TEXT("!!!SETTING TEXTURE!!!"));
 							FIntPoint ImageSize = FIntPoint {(int32)Blob->GetWidth(), (int32)Blob->GetHeight()};
 							NodeViewer->SetTexture(Blob);
 							Viewport->ResetZoom(ImageSize);
diff --git a/Engine/Plugins/TextureGraph/Source/TextureGraphEngine/Job/Job.cpp b/Engine/Plugins/TextureGraph/Source/TextureGraphEngine/Job/Job.cpp
--- a/Engine/Plugins/TextureGraph/Source/TextureGraphEngine/Job/Job.cpp
+++ b/Engine/Plugins/TextureGraph/Source/TextureGraphEngine/Job/Job.cpp
@@ -399,21 +399,21 @@ bool Job::CheckCached()
 	if (ExistingResult)
 	{
 		check(ExistingResult->IsTiled());
+		BufferDescriptor Desc = ExistingResult->GetDescriptor();
+		Desc.bIsTransient = IsDiscard();
 
-		if (ResultOrg)
-			ResultOrg->FinaliseFrom(ExistingResult.get());
-		else
-		{
-			BufferDescriptor Desc = ExistingResult->GetDescriptor();
-			Desc.bIsTransient = IsDiscard();
-			
+		if (!ResultOrg)
 			ResultOrg = std::make_shared<TiledBlob_Promise>(Desc, ExistingResult->Rows(), ExistingResult->Cols(), ExistingResult->Hash());
+
+		ExistingResult->OnFinalise().then([ResultOrg = this->ResultOrg, ExistingResult]() mutable
+		{
+			//check(ExistingResult->IsFinalised());
 			ResultOrg->FinaliseFrom(ExistingResult.get());
-		}
+		});
 
 		Result = TiledBlobRef(std::static_pointer_cast<TiledBlob>(ResultOrg), true, false);
-
 		FinalJobResult = std::make_shared<JobResult>(GetResultRef(), nullptr);
+
 		bIsCulled = true;
 
 		return true;
@@ -633,17 +633,15 @@ bool Job::CheckCulled(JobRunInfo InRunInfo)
 	if (!Result || Result->IsLateBound())
 		return false;
 
-	return false;
-
 	/// There is a bug in this code. It basically (incorrectly) treats a full (non-tiled) blob and tries to 
 	/// cull the job that generates it. Disabling for the time being as it's not a particular performance
 	/// deterrant. But this needs to be tweaked for these kinds of jobs (like T_Transform for instance)
-#if 0
 	check(Result->IsPromise());
 	TiledBlob_PromisePtr ResultPromise = GetResultPromise();
 	size_t NumRows = ResultPromise->Rows();
 	size_t NumCols = ResultPromise->Cols();
 	BlobPtrTiles Tiles(NumRows, NumCols);
+	std::vector<AsyncBlobResultPtr> Promises;
 
 	for (size_t RowId = 0; RowId < ResultPromise->Rows(); RowId++)
 	{
@@ -663,6 +661,9 @@ bool Job::CheckCulled(JobRunInfo InRunInfo)
 			if (!TileBlob)
 				return false;
 
+			if (!TileBlob->IsFinalised())
+				Promises.emplace_back(TileBlob->OnFinalise());
+
 			UE_LOG(LogJob, VeryVerbose, TEXT("[Job - %s] matched tile: %d, %d"), *GetName(), (int32)RowId, (int32)ColId);
 
 			Tiles[RowId][ColId] = TileBlob;
@@ -672,9 +673,21 @@ bool Job::CheckCulled(JobRunInfo InRunInfo)
 	/// We get here then that means we have all the tiles
 	ResultPromise->SetTiles(Tiles);
 
-	/// Finalise the result if it hasn't been finalised yet
-	if (Result && !Result->IsFinalised())
-		Result->Finalise(true, 0);
+	if (!Promises.empty())
+	{
+		cti::when_all(Promises.begin(), Promises.end()).then([Result = this->Result]() mutable
+		{
+			/// Finalise the result if it hasn't been finalised yet
+			if (Result && !Result->IsFinalised())
+				Result->Finalise(true, 0);
+		});
+	}
+	else
+	{
+		/// Finalise the result if it hasn't been finalised yet
+		if (Result && !Result->IsFinalised())
+			Result->Finalise(true, 0);
+	}
 
 	UE_LOG(LogJob, Verbose, TEXT("[Job - %s] All tiles found. Marking as culled!"), *GetName());
 
@@ -689,7 +702,6 @@ bool Job::CheckCulled(JobRunInfo InRunInfo)
 	SetPromise(0);
 
 	return bIsCulled;
-#endif 
 }
 
 AsyncPrepareResult Job::PrepareTargets(JobBatch* Batch)
diff --git a/Engine/Plugins/TextureGraph/Source/TextureGraphEngine/Job/Scheduler.cpp b/Engine/Plugins/TextureGraph/Source/TextureGraphEngine/Job/Scheduler.cpp
--- a/Engine/Plugins/TextureGraph/Source/TextureGraphEngine/Job/Scheduler.cpp
+++ b/Engine/Plugins/TextureGraph/Source/TextureGraphEngine/Job/Scheduler.cpp
@@ -10,7 +10,7 @@ DECLARE_CYCLE_STAT(TEXT("Scheduler_Update"), STAT_Scheduler_Update, STATGROUP_Te
 
 const double Scheduler::IdleTimeInterval = 500.0;
 const double Scheduler::IdleBusyTimeInterval = 1000.0;
-const double Scheduler::IdleBatchTimeLimit = 100.0;
+const double Scheduler::IdleBatchTimeLimit = 500.0;
 const double Scheduler::CurrentBatchWarningLimit = 2000.0;
 
 Scheduler::Scheduler()
@@ -235,7 +235,7 @@ AsyncJobResultPtr Scheduler::UpdateIdleBatch(size_t Index, bool bIsBusyUpdate)
 			if (Duration > IdleBatchTimeLimit)
 			{
 				bDidOffendTimeLimit = true;
-				UE_LOG(LogBatch, Warning, TEXT("Idle Batch [%s] time limit offense @ %f"), *Batch->GetName(), (float)Util::Time());
+				UE_LOG(LogBatch, Verbose, TEXT("Idle Batch [%s] time limit offense @ %0.2f ms [Max: %0.2f ms]"), *Batch->GetName(), (float)Duration, IdleBatchTimeLimit);
 			}
 
 			Batch->UpdateStats(EndTime, StartTime, bDidOffendTimeLimit);
@@ -294,7 +294,8 @@ void Scheduler::UpdateIdle(bool bForce, bool bIsBusyUpdate)
 
 	std::vector<AsyncJobResultPtr> Promises;
 
-	Promises.emplace_back(TextureGraphEngine::GetBlobber()->UpdateIdle());
+	if (!bIsBusyUpdate)
+		Promises.emplace_back(TextureGraphEngine::GetBlobber()->UpdateIdle());
 
 	/// Idle update for all the devices
 	for (size_t DeviceIndex = 0; DeviceIndex < TextureGraphEngine::GetDeviceManager()->GetNumDevices(); DeviceIndex++)
