Add labels to the resize node

diff --git a/Engine/Plugins/TextureGraph/Source/TextureGraph/Public/Expressions/Utilities/TG_Expression_Utils.h b/Engine/Plugins/TextureGraph/Source/TextureGraph/Public/Expressions/Utilities/TG_Expression_Utils.h
--- a/Engine/Plugins/TextureGraph/Source/TextureGraph/Public/Expressions/Utilities/TG_Expression_Utils.h
+++ b/Engine/Plugins/TextureGraph/Source/TextureGraph/Public/Expressions/Utilities/TG_Expression_Utils.h
@@ -94,11 +94,11 @@ public:
 	FTG_Texture Input;
 
 	/// Whether to maintain aspect ratio for the output
-	UPROPERTY(meta = (TGType = "TG_Input", PinDisplayName = ""))
+	UPROPERTY(meta = (TGType = "TG_Input", PinDisplayName = "Maintain Aspect Ratio"))
 	bool bMaintainAspectRatio = true;
 
 	/// The background color to use for the output texture
-	UPROPERTY(meta = (TGType = "TG_Input", PinDisplayName = ""))
+	UPROPERTY(meta = (TGType = "TG_Input", PinDisplayName = "Background Color"))
 	FLinearColor BackgroundColor = FLinearColor(0, 0, 0, 0);
 
 	// The output texture which should be resized to the desired dimensions
