Clang CIS Fix

[RN] rnx

diff --git a/Engine/Plugins/TextureGraph/Source/TextureGraphEditor/Private/TG_Editor.cpp b/Engine/Plugins/TextureGraph/Source/TextureGraphEditor/Private/TG_Editor.cpp
--- a/Engine/Plugins/TextureGraph/Source/TextureGraphEditor/Private/TG_Editor.cpp
+++ b/Engine/Plugins/TextureGraph/Source/TextureGraphEditor/Private/TG_Editor.cpp
@@ -866,7 +866,7 @@ TSharedRef<class SGraphEditor> FTG_Editor::CreateGraphEditorWidget()
 	InEvents.OnNodeDoubleClicked = FSingleNodeEvent::CreateSP(this, &FTG_Editor::OnNodeDoubleClicked);
 	InEvents.OnTextCommitted = FOnNodeTextCommitted::CreateSP(this, &FTG_Editor::OnNodeTitleCommitted);
 	InEvents.OnVerifyTextCommit = FOnNodeVerifyTextCommit::CreateSP(this, &FTG_Editor::OnVerifyNodeTextCommit);
-	InEvents.OnCreateActionMenu = SGraphEditor::FOnCreateActionMenu::CreateSP(this, &FTG_Editor::OnCreateGraphActionMenu);
+	InEvents.OnCreateActionMenuAtLocation = SGraphEditor::FOnCreateActionMenuAtLocation::CreateSP(this, &FTG_Editor::OnCreateGraphActionMenu);
 	//InEvents.OnSpawnNodeByShortcut = SGraphEditor::FOnSpawnNodeByShortcut::CreateSP(this, &FTG_Editor::OnSpawnGraphNodeByShortcut, static_cast<UEdGraph*>(InGraph));
 
 	// Support the delete case
@@ -906,12 +906,12 @@ TSharedRef<class SGraphEditor> FTG_Editor::CreateGraphEditorWidget()
 	
 }
 
-FActionMenuContent FTG_Editor::OnCreateGraphActionMenu(UEdGraph* InGraph, const FVector2D& InNodePosition, const TArray<UEdGraphPin*>& InDraggedPins, bool bAutoExpand, SGraphEditor::FActionMenuClosed InOnMenuClosed)
+FActionMenuContent FTG_Editor::OnCreateGraphActionMenu(UEdGraph* InGraph, const FVector2f& InNodePosition, const TArray<UEdGraphPin*>& InDraggedPins, bool bAutoExpand, SGraphEditor::FActionMenuClosed InOnMenuClosed)
 {
 	TSharedRef<STG_GraphActionMenu> ActionMenu =
 		SNew(STG_GraphActionMenu)
 		.GraphObj(InGraph)
-		.NewNodePosition(InNodePosition)
+		.NewNodePosition(FDeprecateSlateVector2D(InNodePosition))
 		.DraggedFromPins(InDraggedPins)
 		.AutoExpandActionMenu(bAutoExpand)
 		.OnClosedCallback(InOnMenuClosed);
diff --git a/Engine/Plugins/TextureGraph/Source/TextureGraphEditor/Private/TG_Editor.h b/Engine/Plugins/TextureGraph/Source/TextureGraphEditor/Private/TG_Editor.h
--- a/Engine/Plugins/TextureGraph/Source/TextureGraphEditor/Private/TG_Editor.h
+++ b/Engine/Plugins/TextureGraph/Source/TextureGraphEditor/Private/TG_Editor.h
@@ -192,7 +192,7 @@ private:
 	/** Create new graph editor widget */
 	TSharedRef<class SGraphEditor>					CreateGraphEditorWidget();
 
-	FActionMenuContent								OnCreateGraphActionMenu(UEdGraph* InGraph, const FVector2D& InNodePosition, const TArray<UEdGraphPin*>& InDraggedPins, bool bAutoExpand, SGraphEditor::FActionMenuClosed InOnMenuClosed);
+	FActionMenuContent								OnCreateGraphActionMenu(UEdGraph* InGraph, const FVector2f& InNodePosition, const TArray<UEdGraphPin*>& InDraggedPins, bool bAutoExpand, SGraphEditor::FActionMenuClosed InOnMenuClosed);
 
 	/** Create Node preview widget */
 	TSharedRef<class STG_NodePreviewWidget>			CreateNodePreviewWidget();
