Fixed some skeletal meshes crashing

diff --git a/Engine/Plugins/TextureGraph/Source/TextureGraphEditor/Private/STG_EditorViewport.cpp b/Engine/Plugins/TextureGraph/Source/TextureGraphEditor/Private/STG_EditorViewport.cpp
--- a/Engine/Plugins/TextureGraph/Source/TextureGraphEditor/Private/STG_EditorViewport.cpp
+++ b/Engine/Plugins/TextureGraph/Source/TextureGraphEditor/Private/STG_EditorViewport.cpp
@@ -666,9 +666,7 @@ bool STG_EditorViewport::SetPreviewAsset(UObject* InAsset)
 			if (ComponentClass->IsChildOf(UMeshComponent::StaticClass()))
 			{
 				PreviewMeshComponent = NewObject<UMeshComponent>(GetTransientPackage(), ComponentClass, NAME_None, RF_Transient);
-
 				FComponentAssetBrokerage::AssignAssetToComponent(PreviewMeshComponent, InAsset);
-
 				PreviewPrimType = TPT_None;
 			}
 		}
@@ -687,10 +685,22 @@ bool STG_EditorViewport::SetPreviewAsset(UObject* InAsset)
 	}
 	if (IsValid(TextureGraphPtr))
 	{
-		TextureGraphPtr->SetEditorMesh(Cast<UStaticMeshComponent>(PreviewMeshComponent), PreviewScene->GetWorld()).then([this]()
+		try
 		{
-			this->InitRenderModes(TextureGraphPtr);
-		});
+			TextureGraphPtr->SetEditorMesh(Cast<UStaticMeshComponent>(PreviewMeshComponent), PreviewScene->GetWorld())
+				.then([this]()
+					{
+						this->InitRenderModes(TextureGraphPtr);
+					})
+				.fail([this]()
+					{
+						int a = 10;
+					});
+		}
+		catch (...)
+		{
+			int a = 10;
+		}
 	}
 	return (PreviewMeshComponent != nullptr);
 }
diff --git a/Engine/Plugins/TextureGraph/Source/TextureGraphEngine/3D/RenderMesh_Editor.cpp b/Engine/Plugins/TextureGraph/Source/TextureGraphEngine/3D/RenderMesh_Editor.cpp
--- a/Engine/Plugins/TextureGraph/Source/TextureGraphEngine/3D/RenderMesh_Editor.cpp
+++ b/Engine/Plugins/TextureGraph/Source/TextureGraphEngine/3D/RenderMesh_Editor.cpp
@@ -40,7 +40,7 @@ AsyncActionResultPtr RenderMesh_Editor::Load()
 
 	//_meshActor->GetComponents<UStaticMeshComponent>(MeshComponents);
 	if (MeshComponents.Num() <= 0)
-		return cti::make_exceptional_continuable<ActionResultPtr>(std::make_exception_ptr(std::runtime_error("No mesh components found on actor")));
+		return cti::make_ready_continuable(std::make_shared<ActionResult>());
 
 	return cti::make_continuable<ActionResultPtr>([this](auto&& promise)
 		{
