Thumbnail fixes

diff --git a/Engine/Plugins/TextureGraph/Source/TextureGraphEngine/Transform/Layer/T_Thumbnail.cpp b/Engine/Plugins/TextureGraph/Source/TextureGraphEngine/Transform/Layer/T_Thumbnail.cpp
--- a/Engine/Plugins/TextureGraph/Source/TextureGraphEngine/Transform/Layer/T_Thumbnail.cpp
+++ b/Engine/Plugins/TextureGraph/Source/TextureGraphEngine/Transform/Layer/T_Thumbnail.cpp
@@ -25,7 +25,7 @@ JobBatchPtr AddThumbJobToCycle(TiledBlobPtr InBlobToBind, JobUPtr JobPtr, UObjec
 
 TiledBlobRef T_Thumbnail::Bind(UMixInterface* Mix, UObject* Model, TiledBlobPtr InBlobToBind, int32 InTargetId)
 {
-	if (TextureGraphEngine::GetBlobber()->IsCacheEnabled())
+	if (!TextureGraphEngine::GetBlobber()->IsCacheEnabled())
 		return TextureHelper::GetBlack();
 
 	if (InBlobToBind.get() == nullptr)
