Update the Texture Graph Preview Viewport to use the new toolbar style.
- Add new "CreateShowSubmenu" helper function for creating consistent custom show menus.

#rnx

diff --git a/Engine/Plugins/TextureGraph/Source/TextureGraphEditor/Private/STG_EditorViewport.cpp b/Engine/Plugins/TextureGraph/Source/TextureGraphEditor/Private/STG_EditorViewport.cpp
--- a/Engine/Plugins/TextureGraph/Source/TextureGraphEditor/Private/STG_EditorViewport.cpp
+++ b/Engine/Plugins/TextureGraph/Source/TextureGraphEditor/Private/STG_EditorViewport.cpp
@@ -20,6 +20,7 @@
 #include "ISettingsModule.h"
 #include "ITG_Editor.h"
 #include "ImageUtils.h"
+#include "PreviewProfileController.h"
 #include "MaterialEditor/MaterialEditorMeshComponent.h"
 #include "Model/Mix/MixInterface.h"
 #include "Modules/ModuleManager.h"
@@ -30,6 +31,7 @@
 #include "TG_EditorCommands.h"
 #include "UnrealEdGlobals.h"
 #include "Model/Mix/MixSettings.h"
+#include "ToolMenus.h"
 #include "Widgets/Docking/SDockTab.h"
 #include "Widgets/Images/SImage.h"
 #include "Widgets/Input/SNumericEntryBox.h"
@@ -42,6 +44,7 @@
 #include "Materials/Material.h"
 #include "Model/Mix/ViewportSettings.h"
 #include "TextureGraph.h"
+#include "ViewportToolbar/UnrealEdViewportToolbar.h"
 
 
 class FTG_EditorViewportClient : public FEditorViewportClient
@@ -336,14 +339,72 @@ TSharedRef<FEditorViewportClient> STG_EditorViewport::MakeEditorViewportClient()
 	return EditorViewportClient.ToSharedRef();
 }
 
-void STG_EditorViewport::PopulateViewportOverlays(TSharedRef<class SOverlay> Overlay)
+TSharedPtr<SWidget> STG_EditorViewport::MakeViewportToolbar()
 {
-	Overlay->AddSlot()
-		.VAlign(VAlign_Top)
-		[
-			SNew(STG_EditorViewportToolBar, SharedThis(this))
-		];
+	return SNew(STG_EditorViewportToolBar, SharedThis(this));
+}
 
+TSharedPtr<SWidget> STG_EditorViewport::BuildViewportToolbar()
+{
+	const FName ToolbarName = FName("TextureGraph.ViewportToolbar");
+	
+	if (!UToolMenus::Get()->IsMenuRegistered(ToolbarName))
+	{
+		UToolMenu* Menu = UToolMenus::Get()->RegisterMenu(ToolbarName);
+		Menu->MenuType = EMultiBoxType::SlimHorizontalToolBar;
+		Menu->StyleName = "ViewportToolbar";
+		
+		Menu->AddSection("Left");
+		
+		FToolMenuSection& RightSection = Menu->AddSection("Right");
+		RightSection.Alignment = EToolMenuSectionAlign::Last;
+		
+		RightSection.AddEntry(UE::UnrealEd::CreateCameraSubmenu(UE::UnrealEd::FViewportCameraMenuOptions().ShowAll()));
+		RightSection.AddEntry(UE::UnrealEd::CreateViewModesSubmenu());
+		
+		RightSection.AddEntry(UE::UnrealEd::CreateShowSubmenu(
+			FNewToolMenuDelegate::CreateLambda([](UToolMenu* Submenu)
+			{
+				auto Commands = FTG_EditorCommands::Get();
+				FToolMenuSection& Section = Submenu->AddSection(NAME_None);
+				Section.AddMenuEntry(Commands.TogglePreviewGrid);
+				Section.AddMenuEntry(Commands.TogglePreviewBackground);
+			})
+		));
+		
+		RightSection.AddEntry(UE::UnrealEd::CreatePerformanceAndScalabilitySubmenu());
+		RightSection.AddEntry(UE::UnrealEd::CreateAssetViewerProfileSubmenu());
+	}
+	
+	FToolMenuContext Context;
+	{
+		Context.AppendCommandList(GetCommandList());
+		UUnrealEdViewportToolbarContext* ContextObject = UE::UnrealEd::CreateViewportToolbarDefaultContext(SharedThis(this));
+		ContextObject->IsViewModeSupported.BindLambda([](EViewModeIndex ViewModeIndex)
+		{
+			switch (ViewModeIndex)
+			{
+			case VMI_PrimitiveDistanceAccuracy:
+			case VMI_MeshUVDensityAccuracy:
+			case VMI_RequiredTextureResolution:
+				return false;
+			default:
+				return true;
+			}
+		});
+		Context.AddObject(ContextObject);
+	}
+	
+	return UToolMenus::Get()->GenerateWidget(ToolbarName, Context);
+}
+
+TSharedPtr<IPreviewProfileController> STG_EditorViewport::CreatePreviewProfileController()
+{
+	return MakeShared<FPreviewProfileController>();
+}
+
+void STG_EditorViewport::PopulateViewportOverlays(TSharedRef<class SOverlay> Overlay)
+{
 	Overlay->AddSlot()
 		.VAlign(VAlign_Top)
 		.HAlign(HAlign_Right)
diff --git a/Engine/Plugins/TextureGraph/Source/TextureGraphEditor/Private/STG_EditorViewport.h b/Engine/Plugins/TextureGraph/Source/TextureGraphEditor/Private/STG_EditorViewport.h
--- a/Engine/Plugins/TextureGraph/Source/TextureGraphEditor/Private/STG_EditorViewport.h
+++ b/Engine/Plugins/TextureGraph/Source/TextureGraphEditor/Private/STG_EditorViewport.h
@@ -102,6 +102,9 @@ public:
 	bool									SetPreviewAssetByName(const TCHAR* InAssetName);
 
 	virtual TSharedRef<FEditorViewportClient> MakeEditorViewportClient() override;
+	virtual	TSharedPtr<SWidget>				MakeViewportToolbar() override;
+	virtual TSharedPtr<SWidget>				BuildViewportToolbar() override;
+	virtual TSharedPtr<IPreviewProfileController> CreatePreviewProfileController() override;
 	virtual void							PopulateViewportOverlays(TSharedRef<class SOverlay> Overlay) override;
 	virtual void							BindCommands() override;
 
diff --git a/Engine/Source/Editor/UnrealEd/Private/ViewportToolbar/UnrealEdViewportToolbar.cpp b/Engine/Source/Editor/UnrealEd/Private/ViewportToolbar/UnrealEdViewportToolbar.cpp
--- a/Engine/Source/Editor/UnrealEd/Private/ViewportToolbar/UnrealEdViewportToolbar.cpp
+++ b/Engine/Source/Editor/UnrealEd/Private/ViewportToolbar/UnrealEdViewportToolbar.cpp
@@ -3908,25 +3908,25 @@ FText GetScalabilityWarningTooltip()
 	);
 }
 
-FToolMenuEntry CreateDefaultShowSubmenu()
+FToolMenuEntry CreateShowSubmenu(const FNewToolMenuChoice& InSubmenuChoice)
 {
 	FToolMenuEntry Entry = FToolMenuEntry::InitSubMenu(
-		"Show",
-		LOCTEXT("ShowSubmenuLabel", "Show"),
-		LOCTEXT("ShowSubmenuTooltip", "Show flags related to the current viewport"),
-		FNewToolMenuDelegate::CreateLambda(
-			[](UToolMenu* InMenu) -> void
-			{
-				AddDefaultShowFlags(InMenu);
-			}
-		)
-	);
+    	"Show",
+    	LOCTEXT("ShowSubmenuLabel", "Show"),
+    	LOCTEXT("ShowSubmenuTooltip", "Show flags related to the current viewport"),
+    	InSubmenuChoice
+    );
 
-	Entry.Icon = FSlateIcon(FAppStyle::GetAppStyleSetName(), "Level.VisibleHighlightIcon16x");
-	Entry.ToolBarData.LabelOverride = FText();
-	Entry.ToolBarData.ResizeParams.ClippingPriority = 800;
+    Entry.Icon = FSlateIcon(FAppStyle::GetAppStyleSetName(), "Level.VisibleHighlightIcon16x");
+    Entry.ToolBarData.LabelOverride = FText();
+    Entry.ToolBarData.ResizeParams.ClippingPriority = 800;
 
-	return Entry;
+    return Entry;
+}
+
+FToolMenuEntry CreateDefaultShowSubmenu()
+{
+	return CreateShowSubmenu(FNewToolMenuDelegate::CreateStatic(&AddDefaultShowFlags));
 }
 
 void AddDefaultShowFlags(UToolMenu* InMenu)
diff --git a/Engine/Source/Editor/UnrealEd/Public/ViewportToolbar/UnrealEdViewportToolbar.h b/Engine/Source/Editor/UnrealEd/Public/ViewportToolbar/UnrealEdViewportToolbar.h
--- a/Engine/Source/Editor/UnrealEd/Public/ViewportToolbar/UnrealEdViewportToolbar.h
+++ b/Engine/Source/Editor/UnrealEd/Public/ViewportToolbar/UnrealEdViewportToolbar.h
@@ -243,6 +243,11 @@ UNREALED_API FText GetScalabilityWarningLabel();
 UE_DEPRECATED_FORGAME(5.6, "Do not use. Will be made private.")
 UNREALED_API FText GetScalabilityWarningTooltip();
 
+/**
+ * Creates a Show submenu with custom content
+ */
+UNREALED_API FToolMenuEntry CreateShowSubmenu(const FNewToolMenuChoice& InSubmenuChoice);
+
 /**
  * Creates a Show submenu with commonly used show flags
  */
