Fix for crash when launching TextureGraph on Vulkan

diff --git a/Engine/Plugins/TextureGraph/Source/TextureGraphEngine/2D/TextureHelper.cpp b/Engine/Plugins/TextureGraph/Source/TextureGraphEngine/2D/TextureHelper.cpp
--- a/Engine/Plugins/TextureGraph/Source/TextureGraphEngine/2D/TextureHelper.cpp
+++ b/Engine/Plugins/TextureGraph/Source/TextureGraphEngine/2D/TextureHelper.cpp
@@ -145,8 +145,7 @@ void TextureHelper::InitStockTextures()
 	Promises.push_back(InitSolidTexture(&GMagenta, FLinearColor(1, 0, 1, 1), TEXT("MAGENTA"), AlbedoDesc));
 	Promises.push_back(InitSolidTexture(&GDefaultNormal, ColorUtil::DefaultNormal(), TEXT("DefaultNormal"), AlbedoDesc));
 
-	TexDescriptor MaskDesc = AlbedoDesc;
-	MaskDesc.Format = PF_G8;
+	TexDescriptor MaskDesc(1, 1, PF_G8, false, false, false);
 	MaskDesc.NumChannels = 1;
 
 	Promises.push_back(InitSolidTexture(&GWhiteMask, FLinearColor::White, TEXT("WhiteMask"), MaskDesc));
