[TextureGraph}
Fix for output attribute pin of Material expression being modified by localization of editor.
We were using the display names of the material output attributes to expose them as "stable" pin names
these display names where potentially modiufied by the language choice of the editor making the pin name unstable accross different languages.

We are now using the NON LOCALIZED display names of the TG specific EDrawMaterialAttributeTarget enum.

The new names are the same as the previous english display names.
meaning that all the english created TGs are compatible with this change.
All the non english TGs will be uncompatible and will need to be authored correctly.

diff --git a/Engine/Plugins/TextureGraph/Source/TextureGraph/Private/Expressions/Input/TG_Expression_Material.cpp b/Engine/Plugins/TextureGraph/Source/TextureGraph/Private/Expressions/Input/TG_Expression_Material.cpp
--- a/Engine/Plugins/TextureGraph/Source/TextureGraph/Private/Expressions/Input/TG_Expression_Material.cpp
+++ b/Engine/Plugins/TextureGraph/Source/TextureGraph/Private/Expressions/Input/TG_Expression_Material.cpp
@@ -211,7 +211,7 @@ EDrawMaterialAttributeTarget UTG_Expression_Material::GetRenderedAttributeId()
 	}
 	else
 	{
-		return EDrawMaterialAttributeTarget::Emissive;
+		return EDrawMaterialAttributeTarget::EmissiveColor;
 	}
 }
 
diff --git a/Engine/Plugins/TextureGraph/Source/TextureGraph/Private/Expressions/TG_Expression_MaterialBase.cpp b/Engine/Plugins/TextureGraph/Source/TextureGraph/Private/Expressions/TG_Expression_MaterialBase.cpp
--- a/Engine/Plugins/TextureGraph/Source/TextureGraph/Private/Expressions/TG_Expression_MaterialBase.cpp
+++ b/Engine/Plugins/TextureGraph/Source/TextureGraph/Private/Expressions/TG_Expression_MaterialBase.cpp
@@ -23,7 +23,7 @@ EDrawMaterialAttributeTarget UTG_Expression_MaterialBase::ConvertEMaterialProper
 	switch (InMaterialProperty)
 	{
 		case MP_EmissiveColor:
-			return EDrawMaterialAttributeTarget::Emissive;
+			return EDrawMaterialAttributeTarget::EmissiveColor;
 		case MP_Opacity:
 			return EDrawMaterialAttributeTarget::Opacity;
 		case MP_OpacityMask:
@@ -114,11 +114,8 @@ void UTG_Expression_MaterialBase::GenerateMaterialAttributeOptions()
 				if (Attribute != EDrawMaterialAttributeTarget::Count)
 				{
 					AvailableMaterialAttributeIds.Add(Attribute);
-
-
+					// Use the attribute display name but NOT localized as the output pin names
 					FName AttributeName = FName(UEnum::GetDisplayValueAsText(Attribute).ToString());
-					AttributeName = FName(FMaterialAttributeDefinitionMap::GetDisplayNameForMaterial(EMaterialProperty(i), RefMaterial).ToString());
-
 					AvailableMaterialAttributeNames.Add(AttributeName);
 				}
 			}
diff --git a/Engine/Plugins/TextureGraph/Source/TextureGraph/Public/Expressions/Input/TG_Expression_MaterialFunction.h b/Engine/Plugins/TextureGraph/Source/TextureGraph/Public/Expressions/Input/TG_Expression_MaterialFunction.h
--- a/Engine/Plugins/TextureGraph/Source/TextureGraph/Public/Expressions/Input/TG_Expression_MaterialFunction.h
+++ b/Engine/Plugins/TextureGraph/Source/TextureGraph/Public/Expressions/Input/TG_Expression_MaterialFunction.h
@@ -37,7 +37,7 @@ protected:
 	UMaterialInterface* CreateMaterialReference();
 	TObjectPtr<UMaterialInterface> ReferenceMaterial = nullptr;
 	virtual TObjectPtr<UMaterialInterface> GetMaterial() const { return ReferenceMaterial; }
-	virtual EDrawMaterialAttributeTarget GetRenderedAttributeId() { return EDrawMaterialAttributeTarget::Emissive; }
+	virtual EDrawMaterialAttributeTarget GetRenderedAttributeId() { return EDrawMaterialAttributeTarget::EmissiveColor; }
 
 private:
 	void SetMaterialFunctionInternal(UMaterialFunctionInterface* InMaterialFunction);
diff --git a/Engine/Plugins/TextureGraph/Source/TextureGraph/Public/Expressions/TG_Expression_MaterialBase.h b/Engine/Plugins/TextureGraph/Source/TextureGraph/Public/Expressions/TG_Expression_MaterialBase.h
--- a/Engine/Plugins/TextureGraph/Source/TextureGraph/Public/Expressions/TG_Expression_MaterialBase.h
+++ b/Engine/Plugins/TextureGraph/Source/TextureGraph/Public/Expressions/TG_Expression_MaterialBase.h
@@ -11,17 +11,18 @@ class RenderMaterial_BP;
 typedef std::shared_ptr<RenderMaterial_BP> RenderMaterial_BPPtr;
 
 
-// Describe the possible attributes extracted from a MAterial during a DrawMaterial call
+// Describe the possible attributes extracted from a Material during a DrawMaterial call
 // this is used as high level data.
 UENUM(BlueprintType)
 enum class EDrawMaterialAttributeTarget : uint8
 {
+	// DO NOT LOCALIZE THESE ENUM VALUES!
 	BaseColor = 0,
 	Metallic,
 	Specular,
 	Roughness,
 	Anisotropy,
-	Emissive,
+	EmissiveColor,
 	Opacity,
 	OpacityMask,
 	Normal,
@@ -106,7 +107,7 @@ protected:
 	virtual void CopyVarGeneric(const FTG_Argument& Arg, FTG_Var* InVar, bool CopyVarToArg);
 	virtual void SetMaterialInternal(UMaterialInterface* InMaterial);
 
-	virtual EDrawMaterialAttributeTarget GetRenderedAttributeId() { return EDrawMaterialAttributeTarget::Emissive;}
+	virtual EDrawMaterialAttributeTarget GetRenderedAttributeId() { return EDrawMaterialAttributeTarget::EmissiveColor;}
 	virtual TObjectPtr<UMaterialInterface> GetMaterial() const { return nullptr;}
 
 private:
