Texture Graph loses connection from material nodes when changing the editor language.

diff --git a/Engine/Plugins/TextureGraph/Source/TextureGraph/Private/Expressions/TG_Expression_MaterialBase.cpp b/Engine/Plugins/TextureGraph/Source/TextureGraph/Private/Expressions/TG_Expression_MaterialBase.cpp
--- a/Engine/Plugins/TextureGraph/Source/TextureGraph/Private/Expressions/TG_Expression_MaterialBase.cpp
+++ b/Engine/Plugins/TextureGraph/Source/TextureGraph/Private/Expressions/TG_Expression_MaterialBase.cpp
@@ -103,6 +103,19 @@ void UTG_Expression_MaterialBase::GenerateMaterialAttributeOptions()
 	AvailableMaterialAttributeIds.Empty();
 	AvailableMaterialAttributeNames.Empty();
 
+	static const FName AttribNames[static_cast<uint8>(EDrawMaterialAttributeTarget::Count)] = {
+		TEXT("Base Color"),
+		TEXT("Metallic"),
+		TEXT("Specular"),
+		TEXT("Roughness"),
+		TEXT("Anisotropy"),
+		TEXT("Emissive Color"),
+		TEXT("Opacity"),
+		TEXT("Opacity Mask"),
+		TEXT("Normal"),
+		TEXT("Tangent")
+	};
+	
 	if (GetMaterial())
 	{
 		UMaterial* RefMaterial = GetMaterial()->GetMaterial();
@@ -113,11 +126,10 @@ void UTG_Expression_MaterialBase::GenerateMaterialAttributeOptions()
 				EDrawMaterialAttributeTarget Attribute = ConvertEMaterialPropertyToEDrawMaterialAttributeTarget(EMaterialProperty(i));
 				if (Attribute != EDrawMaterialAttributeTarget::Count)
 				{
+					
 					AvailableMaterialAttributeIds.Add(Attribute);
-
-
-					FName AttributeName = FName(UEnum::GetDisplayValueAsText(Attribute).ToString());
-					AttributeName = FName(FMaterialAttributeDefinitionMap::GetDisplayNameForMaterial(EMaterialProperty(i), RefMaterial).ToString());
+					
+					FName AttributeName = AttribNames[static_cast<uint8>(Attribute)];
 
 					AvailableMaterialAttributeNames.Add(AttributeName);
 				}
