texture graph Linux compilation fixes (partial)

diff --git a/Engine/Plugins/TextureGraph/Source/TextureGraph/Private/TG_Texture.cpp b/Engine/Plugins/TextureGraph/Source/TextureGraph/Private/TG_Texture.cpp
--- a/Engine/Plugins/TextureGraph/Source/TextureGraph/Private/TG_Texture.cpp
+++ b/Engine/Plugins/TextureGraph/Source/TextureGraph/Private/TG_Texture.cpp
@@ -53,7 +53,7 @@ FArchive& operator<<(FArchive& Ar, FTG_Texture& T)
 	int32 Version = Ar.CustomVer(FTG_CustomVersion::GUID);
 		
 	if (Ar.IsSaving() ||
-		Ar.IsLoading() && Version >= FTG_CustomVersion::TGTextureAddedTexturePath)
+		(Ar.IsLoading() && Version >= FTG_CustomVersion::TGTextureAddedTexturePath))
 	{
 		Ar << T.TexturePath;
 	}
diff --git a/Engine/Plugins/TextureGraph/Source/TextureGraph/Private/TG_Variant.cpp b/Engine/Plugins/TextureGraph/Source/TextureGraph/Private/TG_Variant.cpp
--- a/Engine/Plugins/TextureGraph/Source/TextureGraph/Private/TG_Variant.cpp
+++ b/Engine/Plugins/TextureGraph/Source/TextureGraph/Private/TG_Variant.cpp
@@ -169,7 +169,7 @@ FTG_Texture FTG_Variant::GetTexture(FTG_EvaluationContext* InContext, FTG_Textur
 
 	/// Special handling for black, white and gray textures as they are quite common in the system
 	/// both in vector, color as well as scalar forms
-	if (!DesiredDesc || DesiredDesc->Width <= 1 && DesiredDesc->Height <= 1)
+	if (!DesiredDesc || (DesiredDesc->Width <= 1 && DesiredDesc->Height <= 1))
 	{
 		if (ColorUtil::IsColorBlack(Color))
 			return FTG_Texture::GetBlack();
diff --git a/Engine/Plugins/TextureGraph/Source/TextureGraph/Public/TG_Var.h b/Engine/Plugins/TextureGraph/Source/TextureGraph/Public/TG_Var.h
--- a/Engine/Plugins/TextureGraph/Source/TextureGraph/Public/TG_Var.h
+++ b/Engine/Plugins/TextureGraph/Source/TextureGraph/Public/TG_Var.h
@@ -8,6 +8,8 @@
 
 #include "TG_Var.generated.h"
 
+class UTG_Expression;
+
 template <typename T>
 void TG_Var_SetValueFromString(T& Value, const FString& StrVal)
 {
diff --git a/Engine/Plugins/TextureGraph/Source/TextureGraphEngine/2D/Tex.cpp b/Engine/Plugins/TextureGraph/Source/TextureGraphEngine/2D/Tex.cpp
--- a/Engine/Plugins/TextureGraph/Source/TextureGraphEngine/2D/Tex.cpp
+++ b/Engine/Plugins/TextureGraph/Source/TextureGraphEngine/2D/Tex.cpp
@@ -62,7 +62,7 @@ TexDescriptor::TexDescriptor(uint32 InWidth, uint32 InHeight, EPixelFormat InFor
 	: Name(Util::RandomID())
 	, Width(InWidth)
 	, Height(InHeight)
-	, Format(Format)
+	, Format(InFormat)
 	, NumChannels(TextureHelper::GetChannelsFromPixelFormat(InFormat))
 {
 }
diff --git a/Engine/Plugins/TextureGraph/Source/TextureGraphEngine/Data/Blob.h b/Engine/Plugins/TextureGraph/Source/TextureGraphEngine/Data/Blob.h
--- a/Engine/Plugins/TextureGraph/Source/TextureGraphEngine/Data/Blob.h
+++ b/Engine/Plugins/TextureGraph/Source/TextureGraphEngine/Data/Blob.h
@@ -196,5 +196,5 @@ public:
 	static CHashPtrVec				CalculateMipHashes(CHashPtr MainHash, CHashPtr ParentHash, int32 Level);
 };
 
-typedef T_BlobRef<Blob>				BlobRef;
+
 typedef T_Tiles<BlobPtr>			BlobPtrTiles;
diff --git a/Engine/Plugins/TextureGraph/Source/TextureGraphEngine/Data/BlobRef.h b/Engine/Plugins/TextureGraph/Source/TextureGraphEngine/Data/BlobRef.h
--- a/Engine/Plugins/TextureGraph/Source/TextureGraphEngine/Data/BlobRef.h
+++ b/Engine/Plugins/TextureGraph/Source/TextureGraphEngine/Data/BlobRef.h
@@ -2,6 +2,7 @@
 #pragma once
 
 #include "TextureGraphEngine.h"
+#include "Blob.h"
 #include "Device/DeviceBuffer.h"
 #include <memory>
 
@@ -212,3 +213,4 @@ public:
 	bool							IsNull() const { return !get(); }
 };
 
+typedef T_BlobRef<Blob>				BlobRef;
\ No newline at end of file
diff --git a/Engine/Plugins/TextureGraph/Source/TextureGraphEngine/Transform/Expressions/T_Filter.cpp b/Engine/Plugins/TextureGraph/Source/TextureGraphEngine/Transform/Expressions/T_Filter.cpp
--- a/Engine/Plugins/TextureGraph/Source/TextureGraphEngine/Transform/Expressions/T_Filter.cpp
+++ b/Engine/Plugins/TextureGraph/Source/TextureGraphEngine/Transform/Expressions/T_Filter.cpp
@@ -68,7 +68,7 @@ static JobUPtr CreateWarp(FString Name, MixUpdateCyclePtr InCycle, BufferDescrip
 		->AddArg(ARG_FLOAT(Intensity, "Intensity"))
 		;
 
-	return std::move(RenderJob);
+	return RenderJob;
 }
 
 TiledBlobPtr T_Filter::CreateDirectionalWarp(MixUpdateCyclePtr InCycle, BufferDescriptor DesiredDesc, TiledBlobPtr SourceTexture, TiledBlobPtr Mask, float Intensity, float AngleRad, int32 InTargetId)
diff --git a/Engine/Plugins/TextureGraph/Source/TextureGraphEngine/Transform/Mesh/T_MeshEncodedAsTexture.cpp b/Engine/Plugins/TextureGraph/Source/TextureGraphEngine/Transform/Mesh/T_MeshEncodedAsTexture.cpp
--- a/Engine/Plugins/TextureGraph/Source/TextureGraphEngine/Transform/Mesh/T_MeshEncodedAsTexture.cpp
+++ b/Engine/Plugins/TextureGraph/Source/TextureGraphEngine/Transform/Mesh/T_MeshEncodedAsTexture.cpp
@@ -67,7 +67,7 @@ static JobUPtr CreateGenericJob(MixUpdateCyclePtr Cycle, int32 TargetId, FString
 	/// this can be safely deleted 
 	Result->GetBufferRef()->SetDeviceTransferChain(DeviceBuffer::FXOnlyTransferChain);
 
-	return std::move(JobP);
+	return JobP;
 }
 
 TiledBlobPtr T_MeshEncodedAsTexture::Create_WorldPos(MixUpdateCyclePtr Cycle, int32 TargetId)
diff --git a/Engine/Plugins/TextureGraph/Source/TextureGraphEngine/Transform/Mix/T_UpdateTargets.cpp b/Engine/Plugins/TextureGraph/Source/TextureGraphEngine/Transform/Mix/T_UpdateTargets.cpp
--- a/Engine/Plugins/TextureGraph/Source/TextureGraphEngine/Transform/Mix/T_UpdateTargets.cpp
+++ b/Engine/Plugins/TextureGraph/Source/TextureGraphEngine/Transform/Mix/T_UpdateTargets.cpp
@@ -65,7 +65,7 @@ JobUPtr T_UpdateTargets::CreateJob(MixUpdateCyclePtr Cycle, int32 TargetId, bool
 	UpdateTargetsTransform->SetShouldUpdate(ShouldUpdate);
 	BlobTransformPtr Transform = std::static_pointer_cast<BlobTransform>(UpdateTargetsTransform);
 	JobUPtr job = std::make_unique<Job>(Cycle->GetMix(), TargetId, Transform);
-	return std::move(job);
+	return job;
 }
 
 void T_UpdateTargets::Create(MixUpdateCyclePtr Cycle, int32 TargetId, bool ShouldUpdate)
