[TextureGraph}
Enable support for SUbstrate Material attributes rendering
Added the shader support to extract common attributes from a Substrate Front Surface dataset.
Modifies the TextureGraphMateiralShader.usf used by MaterialNode as a custom material shader in order to add support for Substrate.
This shader is used only at Editor time.

diff --git a/Engine/Plugins/TextureGraph/Shaders/TextureGraphMaterialShader.usf b/Engine/Plugins/TextureGraph/Shaders/TextureGraphMaterialShader.usf
--- a/Engine/Plugins/TextureGraph/Shaders/TextureGraphMaterialShader.usf
+++ b/Engine/Plugins/TextureGraph/Shaders/TextureGraphMaterialShader.usf
@@ -1,11 +1,21 @@
 // Copyright Epic Games, Inc. All Rights Reserved.
 #undef SCENE_TEXTURES_DISABLED
 #define SCENE_TEXTURES_DISABLED 1
+#define SUBSTRATE_INLINE_SHADING 1
+#define SUBSTRATE_USE_FULLYSIMPLIFIED_MATERIAL 1
+
 
 #include "/Engine/Public/Platform.ush"
 #include "/Engine/Private/Common.ush"
 #include "/Engine/Generated/Material.ush"
 
+#define SHADER_MATERIAL_USES_SUBSTRATE	(SUBSTRATE_ENABLED && MATERIAL_IS_SUBSTRATE)
+#if SHADER_MATERIAL_USES_SUBSTRATE
+#define MATERIAL_SUBSTRATE_OPAQUE_PRECOMPUTED_LIGHTING 0
+#define SUBSTRATE_MATERIAL_EXPORT_REQUESTED 1
+#include "/Engine/Private/Substrate/SubstrateExport.ush"
+#endif // SHADER_MATERIAL_USES_SUBSTRATE
+
 float4 PSControl;
 
 #include "TileInfo.ush"
@@ -31,7 +41,7 @@ void TextureGraphMaterialShaderPS(
 	InUV = (InUV + float2(TileInfo_TileX, TileInfo_TileY)) / float2(TileInfo_TileCountX, TileInfo_TileCountY);
 
 	ResolvedView = ResolveView();
-	FMaterialPixelParameters Parameters = MakeInitializedMaterialPixelParameters();
+	FMaterialPixelParameters MaterialParameters = MakeInitializedMaterialPixelParameters();
 	FPixelMaterialInputs PixelMaterialInputs;
 
 	float2 ViewportUV = InUV;
@@ -39,7 +49,7 @@ void TextureGraphMaterialShaderPS(
 #if NUM_MATERIAL_TEXCOORDS
 	for(int CoordinateIndex = 0;CoordinateIndex < NUM_MATERIAL_TEXCOORDS;CoordinateIndex++)
 	{
-		Parameters.TexCoords[0] = ViewportUV;
+		MaterialParameters.TexCoords[0] = ViewportUV;
 	}
 #endif // NUM_MATERIAL_TEXCOORDS
 
@@ -47,10 +57,8 @@ void TextureGraphMaterialShaderPS(
 	SvPosition.z = max(SvPosition.z, 1e-18);
 
 	// fill out other related material parameters
-	CalcMaterialParametersPost(Parameters, PixelMaterialInputs, SvPosition, true);
-
-	Parameters.VertexColor = 1;
-
+	CalcMaterialParametersPost(MaterialParameters, PixelMaterialInputs, SvPosition, true);
+	
 	// Extract the value we want to represent
 	half3 MatBaseColor = GetMaterialBaseColor(PixelMaterialInputs);
 	half3 MatMetallic = GetMaterialMetallic(PixelMaterialInputs);
@@ -59,6 +67,48 @@ void TextureGraphMaterialShaderPS(
 	half3 MatAnisotropy = GetMaterialAnisotropy(PixelMaterialInputs);
 	half3 MatEmissive = GetMaterialEmissive(PixelMaterialInputs);
 	half3 MatOpacity = GetMaterialOpacity(PixelMaterialInputs);
+	
+	half3 MatOpacityMask = float3(0,0,0);
+#if MATERIALBLENDING_MASKED
+	MatOpacityMask = GetMaterialOpacityMaskClipValue();
+#endif // MATERIALBLENDING_TRANSLUCENT
+
+	half3 MatNormal = float3(0,0,0);
+	half3 MatTangent = float3(0,0,0);
+#if MATERIAL_TANGENTSPACENORMAL
+	MatNormal = GetMaterialNormal(MaterialParameters, PixelMaterialInputs) * 0.5 + 0.5;
+	MatTangent = GetMaterialTangent(PixelMaterialInputs) * 0.5 + 0.5;
+#endif // MATERIAL_TANGENTSPACENORMAL
+
+#if SHADER_MATERIAL_USES_SUBSTRATE
+	// Initialise a Substrate header with normal in registers
+	FSubstrateData SubstrateData = PixelMaterialInputs.GetFrontSubstrateData();
+	FSubstratePixelHeader SubstratePixelHeader = MaterialParameters.GetFrontSubstrateHeader();
+	FSubstrateIntegrationSettings Settings = InitSubstrateIntegrationSettings();
+
+	const float3 SurfaceWorldNormal = MaterialParameters.TangentToWorld[2].xyz;
+	FExportResult Export = SubstrateMaterialExportOut(
+		Settings,
+		SubstratePixelHeader,
+		SubstrateData,
+		SurfaceWorldNormal,
+		MaterialParameters.WorldPosition_CamRelative,
+		0.0f // Curvature
+	);
+	
+	MatBaseColor = Export.BaseColor;
+	MatMetallic = Export.Metallic;
+	MatSpecular = Export.Specular;
+	MatRoughness = Export.Roughness;
+	MatEmissive = Export.EmissiveLuminance;
+	MatNormal = Export.WorldNormal * 0.5 + 0.5;
+	MatOpacity = Export.Coverage;
+
+#if MATERIAL_USES_ANISOTROPY
+	MatTangent = Export.WorldTangent * 0.5 + 0.5;
+	MatAnisotropy = Export.Anisotropy;
+#endif
+#endif // SHADER_MATERIAL_USES_SUBSTRATE
 
 	half3 UV = half3(ViewportUV, 0);
 	half3 Tile = half3(float2(TileInfo_TileX, TileInfo_TileY) / float2(TileInfo_TileCountX, TileInfo_TileCountY), 0);
@@ -72,20 +122,10 @@ void TextureGraphMaterialShaderPS(
 	Color += (Attribute == 4.0) * MatAnisotropy;
 	Color += (Attribute == 5.0) * MatEmissive;
 	Color += (Attribute == 6.0) * MatOpacity;
-	
-#if MATERIALBLENDING_MASKED
-	half3 MatOpacityMask = GetMaterialOpacityMaskClipValue();
 	Color += (Attribute == 7.0) * MatOpacityMask;
-#endif // MATERIALBLENDING_TRANSLUCENT
-
-#if MATERIAL_TANGENTSPACENORMAL
-	half3 MatNormal = GetMaterialNormal(Parameters, PixelMaterialInputs) * 0.5 + 0.5;
-	half3 MatTangent = GetMaterialTangent(PixelMaterialInputs) * 0.5 + 0.5;
-	
 	Color += (Attribute == 8.0) * MatNormal;
 	Color += (Attribute >= 9.0) * MatTangent;
-#endif // MATERIAL_TANGENTSPACENORMAL
-
+	
 	Color = lerp(Color, UV, PSControl.z);
 	Color = lerp(Color, Tile, PSControl.w);
 
diff --git a/Engine/Plugins/TextureGraph/Source/TextureGraph/Private/Expressions/TG_Expression_MaterialBase.cpp b/Engine/Plugins/TextureGraph/Source/TextureGraph/Private/Expressions/TG_Expression_MaterialBase.cpp
--- a/Engine/Plugins/TextureGraph/Source/TextureGraph/Private/Expressions/TG_Expression_MaterialBase.cpp
+++ b/Engine/Plugins/TextureGraph/Source/TextureGraph/Private/Expressions/TG_Expression_MaterialBase.cpp
@@ -121,16 +121,18 @@ void UTG_Expression_MaterialBase::GenerateMaterialAttributeOptions()
 		UMaterial* RefMaterial = GetMaterial()->GetMaterial();
 		for (int i = MP_EmissiveColor; i < MP_MAX; ++i)
 		{
-			if (RefMaterial->IsPropertyConnected(EMaterialProperty(i)))
+			EMaterialProperty MP = EMaterialProperty(i);
+			if (RefMaterial->IsPropertyConnected(MP))
 			{
-				EDrawMaterialAttributeTarget Attribute = ConvertEMaterialPropertyToEDrawMaterialAttributeTarget(EMaterialProperty(i));
+				if (MP == MP_DiffuseColor && !RefMaterial->IsPropertyConnected(MP_BaseColor))
+				{
+					MP = MP_BaseColor;
+				}
+				EDrawMaterialAttributeTarget Attribute = ConvertEMaterialPropertyToEDrawMaterialAttributeTarget(MP);
 				if (Attribute != EDrawMaterialAttributeTarget::Count)
 				{
-					
 					AvailableMaterialAttributeIds.Add(Attribute);
-					
 					FName AttributeName = AttribNames[static_cast<uint8>(Attribute)];
-
 					AvailableMaterialAttributeNames.Add(AttributeName);
 				}
 			}
