Fixed circular includes in TextureGraphEngine. Also moved dll storage from types to methods/staticvars

diff --git a/Engine/Plugins/TextureGraph/Source/TextureGraphEngine/Data/Blob.h b/Engine/Plugins/TextureGraph/Source/TextureGraphEngine/Data/Blob.h
--- a/Engine/Plugins/TextureGraph/Source/TextureGraphEngine/Data/Blob.h
+++ b/Engine/Plugins/TextureGraph/Source/TextureGraphEngine/Data/Blob.h
@@ -4,10 +4,11 @@
 #include "CoreMinimal.h"
 #include "RawBuffer.h"
 #include "Device/DeviceBuffer.h"
-#include "BlobRef.h"
 #include "Helper/DataUtil.h"
 #include "2D/TextureHelper.h"
 
+#define UE_API TEXTUREGRAPHENGINE_API
+
 class Device;
 class RawBuffer;
 class BlobTransform;
@@ -32,14 +33,14 @@ public:
 };
 #endif 
 
-class TEXTUREGRAPHENGINE_API Blob : public std::enable_shared_from_this<Blob>
+class Blob : public std::enable_shared_from_this<Blob>
 {
 	friend class Blobber;
 	friend class TiledBlob;
 	friend class TiledBlob_Promise;
 
 public:
-	static const char*				LODTransformName;	
+	UE_API static const char*				LODTransformName;	
 	typedef std::vector<BlobPtrW>	OwnerList;
 	typedef std::vector<std::weak_ptr<Blob>> LinkedBlobsVec;
 
@@ -70,93 +71,93 @@ protected:
 																	/// but we need to ensure that it works correctly for all blobs
 
 #if DEBUG_BLOB_REF_KEEPING == 1
-	OwnerList						_owners;						/// The TiledBlobs that contain this blob as a tile
-	virtual void					AddTiledOwner(BlobPtrW owner);
+	OwnerList							_owners;					/// The TiledBlobs that contain this blob as a tile
+	UE_API virtual void					AddTiledOwner(BlobPtrW owner);
 #endif 
 
-									Blob();
-	virtual void					OnFinaliseInternal(BlobReadyCallback callback) const;
-	virtual void					ResetBuffer();
-	virtual void					Touch(uint64 BatchId);
-	virtual void					UpdateAccessInfo(uint64 batchId);
-	virtual void					SetHash(CHashPtr Hash);
+	UE_API 								Blob();
+	UE_API virtual void					OnFinaliseInternal(BlobReadyCallback callback) const;
+	UE_API virtual void					ResetBuffer();
+	UE_API virtual void					Touch(uint64 BatchId);
+	UE_API virtual void					UpdateAccessInfo(uint64 batchId);
+	UE_API virtual void					SetHash(CHashPtr Hash);
 
-	virtual void					UpdateLinkedBlobs(bool bDoFinalise);
-	virtual void					AddLinkedBlob(BlobPtr LinkedBlob);
+	UE_API virtual void					UpdateLinkedBlobs(bool bDoFinalise);
+	UE_API virtual void					AddLinkedBlob(BlobPtr LinkedBlob);
 
-	virtual void					FinaliseFrom(Blob* RHS);
-	virtual void					SyncWith(Blob* RHS);
+	UE_API virtual void					FinaliseFrom(Blob* RHS);
+	UE_API virtual void					SyncWith(Blob* RHS);
 
 public:
-									Blob(DeviceBufferRef InBuffer);
-									Blob(const BufferDescriptor& InDesc, CHashPtr InHash);
-									Blob(const Blob& InBlob) = delete; /// non-copyable
+	UE_API 								Blob(DeviceBufferRef InBuffer);
+	UE_API 								Blob(const BufferDescriptor& InDesc, CHashPtr InHash);
+		 								Blob(const Blob& InBlob) = delete; /// non-copyable
 
-	virtual							~Blob();
+	UE_API virtual						~Blob();
 
-	virtual AsyncBufferResultPtr	Bind(const BlobTransform* Transform, const ResourceBindInfo& BindInfo);
-	virtual AsyncBufferResultPtr	Unbind(const BlobTransform* Transform, const ResourceBindInfo& BindInfo);
-	virtual bool					IsValid() const;
-	virtual bool					IsNull() const;
-	virtual CHashPtr				Hash() const;
-	//virtual void					SetHash(CHashPtr hash);
-	virtual AsyncBufferResultPtr	Flush(const ResourceBindInfo& BindInfo);
-	virtual AsyncPrepareResult		PrepareForWrite(const ResourceBindInfo& BindInfo);
+	UE_API virtual AsyncBufferResultPtr	Bind(const BlobTransform* Transform, const ResourceBindInfo& BindInfo);
+	UE_API virtual AsyncBufferResultPtr	Unbind(const BlobTransform* Transform, const ResourceBindInfo& BindInfo);
+	UE_API virtual bool					IsValid() const;
+	UE_API virtual bool					IsNull() const;
+	UE_API virtual CHashPtr				Hash() const;
+	//UE_API virtual void				SetHash(CHashPtr hash);
+	UE_API virtual AsyncBufferResultPtr	Flush(const ResourceBindInfo& BindInfo);
+	UE_API virtual AsyncPrepareResult	PrepareForWrite(const ResourceBindInfo& BindInfo);
 
 #if 0 /// TODO /// Do we really need this?
-	virtual AsyncBufferResultPtr	PrepareForWrite(const ResourceBindInfo& BindInfo);
+	UE_API virtual AsyncBufferResultPtr	PrepareForWrite(const ResourceBindInfo& BindInfo);
 #endif 
 
-	virtual FString					DisplayName() const;
-	virtual bool					IsFinalised() const { return bIsFinalised; }
-	virtual bool					IsTiled() const { return false; }
-	virtual bool					CanCalculateHash() const { return IsFinalised(); }
+	UE_API virtual FString				DisplayName() const;
+	virtual bool						IsFinalised() const { return bIsFinalised; }
+	virtual bool						IsTiled() const { return false; }
+	virtual bool						CanCalculateHash() const { return IsFinalised(); }
 
-	virtual void					FinaliseNow(bool bNoCalcHash, CHashPtr FixedHash);
-	virtual AsyncBufferResultPtr	Finalise(bool bNoCalcHash, CHashPtr FixedHash);
+	UE_API virtual void					FinaliseNow(bool bNoCalcHash, CHashPtr FixedHash);
+	UE_API virtual AsyncBufferResultPtr	Finalise(bool bNoCalcHash, CHashPtr FixedHash);
 
-	virtual AsyncRawBufferPtr		Raw();			/// SLOW: Read doc for DeviceBuffer::Raw() in DeviceBuffer.h
-	virtual AscynCHashPtr			CalcHash();		/// SLOW: Because calls Raw() underneath. Don't use directly
+	UE_API virtual AsyncRawBufferPtr	Raw();			/// SLOW: Read doc for DeviceBuffer::Raw() in DeviceBuffer.h
+	UE_API virtual AscynCHashPtr		CalcHash();		/// SLOW: Because calls Raw() underneath. Don't use directly
 
-	virtual FString&				Name() { check(Buffer); return Buffer->GetName(); }
-	virtual const FString&			Name() const { check(Buffer); return Buffer->GetName(); }
-	virtual const BufferDescriptor& GetDescriptor() const { /*check(_buffer);*/ return Buffer->Descriptor(); }
+	virtual FString&					Name() { check(Buffer); return Buffer->GetName(); }
+	virtual const FString&				Name() const { check(Buffer); return Buffer->GetName(); }
+	virtual const BufferDescriptor&		GetDescriptor() const { /*check(_buffer);*/ return Buffer->Descriptor(); }
 	
-	virtual AsyncBlobResultPtr		OnFinalise() const;
+	UE_API virtual AsyncBlobResultPtr	OnFinalise() const;
 
-	virtual AsyncDeviceBufferRef	TransferTo(Device* TargetDevice);
-	virtual DeviceBufferRef			GetBufferRef() const;
+	UE_API virtual AsyncDeviceBufferRef	TransferTo(Device* TargetDevice);
+	UE_API virtual DeviceBufferRef		GetBufferRef() const;
 
 	//////////////////////////////////////////////////////////////////////////
 	/// Min/Max related
 	//////////////////////////////////////////////////////////////////////////
-	virtual BlobPtr					GetMinMaxBlob();
-	virtual void					SetMinMax(BlobPtr InMinMax);
-	virtual float					GetMinValue() const;
-	virtual float					GetMaxValue() const;
-	virtual bool					HasMinMax() const;
+	UE_API virtual BlobPtr				GetMinMaxBlob();
+	UE_API virtual void					SetMinMax(BlobPtr InMinMax);
+	UE_API virtual float				GetMinValue() const;
+	UE_API virtual float				GetMaxValue() const;
+	UE_API virtual bool					HasMinMax() const;
 	//////////////////////////////////////////////////////////////////////////
 	/// LOD/MipMaps related
 	//////////////////////////////////////////////////////////////////////////
-	virtual BlobPtrW				GetLODLevel(int32 Level);
-	virtual bool					HasLODLevels() const;
-	virtual bool					HasLODLevel(int32 Index) const;
-	virtual int32					NumLODLevels() const;
-	virtual void					SetLODLevel(int32 Level, BlobPtr LODBlob, BlobPtrW LODParentBlob, BlobPtrW LODSourceBlob, bool bAddToBlobber); 
-	BlobPtr							GetHistogram() const; 
+	UE_API virtual BlobPtrW				GetLODLevel(int32 Level);
+	UE_API virtual bool					HasLODLevels() const;
+	UE_API virtual bool					HasLODLevel(int32 Index) const;
+	UE_API virtual int32				NumLODLevels() const;
+	UE_API virtual void					SetLODLevel(int32 Level, BlobPtr LODBlob, BlobPtrW LODParentBlob, BlobPtrW LODSourceBlob, bool bAddToBlobber); 
+	UE_API BlobPtr						GetHistogram() const; 
 
 	//////////////////////////////////////////////////////////////////////////
 	/// Ownership related
 	//////////////////////////////////////////////////////////////////////////
 #if DEBUG_BLOB_REF_KEEPING == 1
-	virtual void					AddOwner(BlobPtrW owner);
-	virtual void					RemoveOwner(BlobPtrW owner);
-	void							RemoveOwner(Blob* owner);
-	virtual bool					HasOwner(BlobPtrW owner);
-	bool							HasOwner(Blob* owner);
-	virtual OwnerList::iterator		FindOwner(BlobPtrW owner);
-	OwnerList::iterator				FindOwner(Blob* owner);
-	virtual bool					HasBlobAsTile(Blob* blob) const { return false; };
+	UE_API virtual void					AddOwner(BlobPtrW owner);
+	UE_API virtual void					RemoveOwner(BlobPtrW owner);
+	UE_API void							RemoveOwner(Blob* owner);
+	UE_API virtual bool					HasOwner(BlobPtrW owner);
+	UE_API bool							HasOwner(Blob* owner);
+	UE_API virtual OwnerList::iterator	FindOwner(BlobPtrW owner);
+	UE_API OwnerList::iterator			FindOwner(Blob* owner);
+	virtual bool						HasBlobAsTile(Blob* blob) const { return false; };
 #endif
 
 	//////////////////////////////////////////////////////////////////////////
@@ -192,9 +193,11 @@ public:
 	/// Static functions
 	//////////////////////////////////////////////////////////////////////////
 
-	static CHashPtr					CalculateMipHash(CHashPtr MainHash, int32 Level);
-	static CHashPtrVec				CalculateMipHashes(CHashPtr MainHash, CHashPtr ParentHash, int32 Level);
+	UE_API static CHashPtr			CalculateMipHash(CHashPtr MainHash, int32 Level);
+	UE_API static CHashPtrVec		CalculateMipHashes(CHashPtr MainHash, CHashPtr ParentHash, int32 Level);
 };
 
 
 typedef T_Tiles<BlobPtr>			BlobPtrTiles;
+
+#undef UE_API
\ No newline at end of file
diff --git a/Engine/Plugins/TextureGraph/Source/TextureGraphEngine/Data/TiledBlob.h b/Engine/Plugins/TextureGraph/Source/TextureGraphEngine/Data/TiledBlob.h
--- a/Engine/Plugins/TextureGraph/Source/TextureGraphEngine/Data/TiledBlob.h
+++ b/Engine/Plugins/TextureGraph/Source/TextureGraphEngine/Data/TiledBlob.h
@@ -2,7 +2,9 @@
 #pragma once
 
 #include "CoreMinimal.h"
-#include "Blob.h"
+#include "BlobRef.h"
+
+#define UE_API TEXTUREGRAPHENGINE_API
 
 class Job;
 typedef std::weak_ptr<Job>			JobPtrW;
@@ -14,7 +16,7 @@ typedef T_Tiles<BlobPtr>			BlobPtrTiles;
 /// are always late bound i.e. they will not allocate any buffers 
 /// unless they absolutely need to. 
 //////////////////////////////////////////////////////////////////////////
-class TEXTUREGRAPHENGINE_API TiledBlob : public Blob
+class TiledBlob : public Blob
 {
 	friend class Blobber;
 	friend class TiledBlob_Promise;
@@ -28,79 +30,79 @@ protected:
 	JobPtrW							JobObj;					/// The job that is potentially generating this tiled BlobObj
 	BlobPtr							SingleBlob;			/// Single BlobObj pointer that we keep to prevent going out of Ref
 
-	void							CalcHashNow() const;
-	virtual void					ResetBuffer() override;
+	UE_API void							CalcHashNow() const;
+	UE_API virtual void					ResetBuffer() override;
 
-	virtual void					TouchTiles(uint64 BatchId);
-	virtual void					Touch(uint64 BatchId) override;
-	virtual void					MakeSingleBlob_Internal();
-	TiledBlob&						operator = (const TiledBlob& RHS);
-	virtual void					SetHash(CHashPtr Hash) override;
+	UE_API virtual void					TouchTiles(uint64 BatchId);
+	UE_API virtual void					Touch(uint64 BatchId) override;
+	UE_API virtual void					MakeSingleBlob_Internal();
+	UE_API TiledBlob&					operator = (const TiledBlob& RHS);
+	UE_API virtual void					SetHash(CHashPtr Hash) override;
 
-	virtual void					AddLinkedBlob(BlobPtr LinkedBlob) override;
-	virtual void					FinaliseFrom(Blob* RHS) override;
+	UE_API virtual void					AddLinkedBlob(BlobPtr LinkedBlob) override;
+	UE_API virtual void					FinaliseFrom(Blob* RHS) override;
 
 public:
-	static AsyncBufferResultPtr		TileBuffer(DeviceBufferRef Buffer, BlobPtrTiles& Tiles);
-	AsyncBufferResultPtr			CombineTiles(bool bTouch, bool bIsArray, uint64 BatchId = 0);
+	UE_API static AsyncBufferResultPtr	TileBuffer(DeviceBufferRef Buffer, BlobPtrTiles& Tiles);
+	UE_API AsyncBufferResultPtr			CombineTiles(bool bTouch, bool bIsArray, uint64 BatchId = 0);
 
-									TiledBlob(const BufferDescriptor& Desc, const BlobPtrTiles& Tiles);
+	UE_API 								TiledBlob(const BufferDescriptor& Desc, const BlobPtrTiles& Tiles);
 
 public:
-									TiledBlob(const BufferDescriptor& Desc, size_t NumTilesX, size_t NumTilesY, CHashPtr HashValue);
-									TiledBlob(DeviceBufferRef Buffer);
-									TiledBlob(BlobRef BlobObj);
-	virtual							~TiledBlob() override;
+	UE_API 								TiledBlob(const BufferDescriptor& Desc, size_t NumTilesX, size_t NumTilesY, CHashPtr HashValue);
+	UE_API 								TiledBlob(DeviceBufferRef Buffer);
+	UE_API 								TiledBlob(BlobRef BlobObj);
+	UE_API virtual						~TiledBlob() override;
 
-	static std::shared_ptr<TiledBlob> InitFromTiles(const BufferDescriptor& Desc, BlobPtrTiles& Tiles);
+	UE_API static std::shared_ptr<TiledBlob> InitFromTiles(const BufferDescriptor& Desc, BlobPtrTiles& Tiles);
 
-	virtual AsyncBufferResultPtr	Bind(const BlobTransform* Transform, const ResourceBindInfo& BindInfo) override;
-	virtual AsyncBufferResultPtr	Unbind(const BlobTransform* Transform, const ResourceBindInfo& BindInfo) override;
+	UE_API virtual AsyncBufferResultPtr	Bind(const BlobTransform* Transform, const ResourceBindInfo& BindInfo) override;
+	UE_API virtual AsyncBufferResultPtr	Unbind(const BlobTransform* Transform, const ResourceBindInfo& BindInfo) override;
 
-	virtual bool					IsValid() const override;
-	virtual bool					IsNull() const override;
-	virtual AscynCHashPtr			CalcHash() override;	/// SLOW: Because calls Raw() underneath. Don't use directly
+	UE_API virtual bool					IsValid() const override;
+	UE_API virtual bool					IsNull() const override;
+	UE_API virtual AscynCHashPtr		CalcHash() override;	/// SLOW: Because calls Raw() underneath. Don't use directly
 
-	virtual CHashPtr				Hash() const override;
-	virtual AsyncBufferResultPtr	Flush(const ResourceBindInfo& BindInfo) override;
-	virtual AsyncPrepareResult		PrepareForWrite(const ResourceBindInfo& BindInfo) override;
-	virtual FString					DisplayName() const override;
-	virtual bool					IsTiled() const override { return true; }
-	virtual bool					IsPromise() const { return false; }
-	virtual bool					CanCalculateHash() const override;
-	virtual void					SetTile(int32 X, int32 Y, BlobRef Tile);
-	virtual void					SetTiles(const BlobPtrTiles& InTiles);
+	UE_API virtual CHashPtr				Hash() const override;
+	UE_API virtual AsyncBufferResultPtr	Flush(const ResourceBindInfo& BindInfo) override;
+	UE_API virtual AsyncPrepareResult	PrepareForWrite(const ResourceBindInfo& BindInfo) override;
+	UE_API virtual FString				DisplayName() const override;
+	virtual bool						IsTiled() const override { return true; }
+	virtual bool						IsPromise() const { return false; }
+	UE_API virtual bool					CanCalculateHash() const override;
+	UE_API virtual void					SetTile(int32 X, int32 Y, BlobRef Tile);
+	UE_API virtual void					SetTiles(const BlobPtrTiles& InTiles);
 
-	virtual FString&				Name() override { return Desc.Name; }
-	virtual const FString&			Name() const override { return Desc.Name; }
-	virtual const BufferDescriptor& GetDescriptor() const override { return Desc; }
+	virtual FString&					Name() override { return Desc.Name; }
+	virtual const FString&				Name() const override { return Desc.Name; }
+	virtual const BufferDescriptor&		GetDescriptor() const override { return Desc; }
 
-	virtual AsyncBufferResultPtr	MakeSingleBlob();
+	UE_API virtual AsyncBufferResultPtr	MakeSingleBlob();
 
-	virtual AsyncDeviceBufferRef	TransferTo(Device* Target) override;
-	virtual DeviceBufferRef			GetBufferRef() const override;
+	UE_API virtual AsyncDeviceBufferRef	TransferTo(Device* Target) override;
+	UE_API virtual DeviceBufferRef		GetBufferRef() const override;
 
-	virtual void					ResolveLateBound(BlobPtr Ref);
-	virtual void					ResolveLateBound(const BufferDescriptor& Desc, bool bOverrideExisting = false);
-	virtual void					CopyResolveLateBound(BlobPtr RHS);
-	virtual void					SetTransient();
-	virtual void					FinaliseNow(bool bNoCalcHash, CHashPtr FixedHash) override;
+	UE_API virtual void					ResolveLateBound(BlobPtr Ref);
+	UE_API virtual void					ResolveLateBound(const BufferDescriptor& Desc, bool bOverrideExisting = false);
+	UE_API virtual void					CopyResolveLateBound(BlobPtr RHS);
+	UE_API virtual void					SetTransient();
+	UE_API virtual void					FinaliseNow(bool bNoCalcHash, CHashPtr FixedHash) override;
 
 	//////////////////////////////////////////////////////////////////////////
 	/// Min/Max related
 	//////////////////////////////////////////////////////////////////////////
-	virtual void					SetMinMax(BlobPtr MinMax_) override;
+	UE_API virtual void					SetMinMax(BlobPtr MinMax_) override;
 
 	//////////////////////////////////////////////////////////////////////////
 	/// LOD/MipMaps related
 	//////////////////////////////////////////////////////////////////////////
-	virtual void					SetLODLevel(int32 Level_, BlobPtr BlobObj, BlobPtrW LODParent_, BlobPtrW LODSource_, bool AddToBlobber) override;
+	UE_API virtual void					SetLODLevel(int32 Level_, BlobPtr BlobObj, BlobPtrW LODParent_, BlobPtrW LODSource_, bool AddToBlobber) override;
 
 	//////////////////////////////////////////////////////////////////////////
 	/// Ownership related
 	//////////////////////////////////////////////////////////////////////////
 #if DEBUG_BLOB_REF_KEEPING == 1
-	virtual bool					HasBlobAsTile(Blob* BlobObj) const override;
+	UE_API virtual bool					HasBlobAsTile(Blob* BlobObj) const override;
 #endif 
 
 	//////////////////////////////////////////////////////////////////////////
@@ -111,7 +113,7 @@ public:
 	/// it does nothing. Otherwise it creates a new TiledBlob with 1x1 tile that contains
 	/// the passed BlobObj.
 	/// Just a simple helper function, used in places that always require a TiledBlob.
-	static std::shared_ptr<TiledBlob>	AsTiledBlob(BlobPtr BlobObj);
+	UE_API static std::shared_ptr<TiledBlob>	AsTiledBlob(BlobPtr BlobObj);
 
 	//////////////////////////////////////////////////////////////////////////
 	/// Inline functions
@@ -155,41 +157,41 @@ typedef std::weak_ptr<TiledBlob>	TiledBlobPtrW;
 /// This is a promised TiledBlob that will be the result of the job.
 /// These are lazy evaluated.
 //////////////////////////////////////////////////////////////////////////
-class TEXTUREGRAPHENGINE_API TiledBlob_Promise : public TiledBlob
+class TiledBlob_Promise : public TiledBlob
 {
 protected:
 	mutable std::vector<BlobReadyCallback> Callbacks;		/// The callbacks to call when the promised tiled BlobObj is ready
 	bool							bMakeSingleBlob = false;/// Make single BlobObj on finalise
 	std::weak_ptr<TiledBlob_Promise> CachedBlob;			/// The cached blob that this is derived off
 
-	virtual void					OnFinaliseInternal(BlobReadyCallback Callback) const override;
-	TiledBlob_Promise&				operator = (const TiledBlob_Promise& RHS);
-	TiledBlob_Promise&				operator = (const TiledBlob& RHS);
-	virtual void					NotifyCallbacks();
+	UE_API virtual void					OnFinaliseInternal(BlobReadyCallback Callback) const override;
+	UE_API TiledBlob_Promise&			operator = (const TiledBlob_Promise& RHS);
+	UE_API TiledBlob_Promise&			operator = (const TiledBlob& RHS);
+	UE_API virtual void					NotifyCallbacks();
 
-	virtual void					AddLinkedBlob(BlobPtr LinkedBlob) override;
-	virtual void					FinaliseFrom(Blob* RHS) override;
+	UE_API virtual void					AddLinkedBlob(BlobPtr LinkedBlob) override;
+	UE_API virtual void					FinaliseFrom(Blob* RHS) override;
 
 public:
-									TiledBlob_Promise(TiledBlobPtr Source);
-									TiledBlob_Promise(const BufferDescriptor& Desc, size_t NumTilesX, size_t NumTilesY, CHashPtr HashValue);
-									TiledBlob_Promise(DeviceBufferRef Buffer);
-	virtual							~TiledBlob_Promise();
-
-	virtual AsyncBufferResultPtr	Bind(const BlobTransform* Transform, const ResourceBindInfo& BindInfo) override;
-	virtual AsyncBufferResultPtr	Unbind(const BlobTransform* Transform, const ResourceBindInfo& BindInfo) override;
-	virtual AsyncBufferResultPtr	Flush(const ResourceBindInfo& BindInfo) override;
-	virtual bool					IsFinalised() const override { return bIsFinalised; }
-	virtual void					SetTile(int32 RowId, int32 ColId, BlobRef Tile) override;
-	virtual AsyncBufferResultPtr	MakeSingleBlob() override;
-	virtual void					CopyResolveLateBound(BlobPtr RHS) override;
-	virtual bool					IsPromise() const override { return true; }
-
-	virtual void					FinaliseNow(bool bNoCalcHash, CHashPtr FixedHash) override;
-	virtual AsyncBufferResultPtr	Finalise(bool bNoCalcHash, CHashPtr FixedHash) override;
-	void							FinaliseFrom(TiledBlobPtr RHS);
-
-	void							ResetForReplay(); /// For debug purpose, reset the state of the tile as a promise to NOT finalised, increment the replayCount
+	UE_API 								TiledBlob_Promise(TiledBlobPtr Source);
+	UE_API 								TiledBlob_Promise(const BufferDescriptor& Desc, size_t NumTilesX, size_t NumTilesY, CHashPtr HashValue);
+	UE_API 								TiledBlob_Promise(DeviceBufferRef Buffer);
+	UE_API virtual						~TiledBlob_Promise();
+
+	UE_API virtual AsyncBufferResultPtr	Bind(const BlobTransform* Transform, const ResourceBindInfo& BindInfo) override;
+	UE_API virtual AsyncBufferResultPtr	Unbind(const BlobTransform* Transform, const ResourceBindInfo& BindInfo) override;
+	UE_API virtual AsyncBufferResultPtr	Flush(const ResourceBindInfo& BindInfo) override;
+	virtual bool						IsFinalised() const override { return bIsFinalised; }
+	UE_API virtual void					SetTile(int32 RowId, int32 ColId, BlobRef Tile) override;
+	UE_API virtual AsyncBufferResultPtr	MakeSingleBlob() override;
+	UE_API virtual void					CopyResolveLateBound(BlobPtr RHS) override;
+	virtual bool						IsPromise() const override { return true; }
+
+	UE_API virtual void					FinaliseNow(bool bNoCalcHash, CHashPtr FixedHash) override;
+	UE_API virtual AsyncBufferResultPtr	Finalise(bool bNoCalcHash, CHashPtr FixedHash) override;
+	UE_API void							FinaliseFrom(TiledBlobPtr RHS);
+
+	UE_API void							ResetForReplay(); /// For debug purpose, reset the state of the tile as a promise to NOT finalised, increment the replayCount
 
 	//////////////////////////////////////////////////////////////////////////
 	FORCEINLINE bool				IsSingleBlob() const { return bMakeSingleBlob; }
@@ -199,3 +201,5 @@ typedef std::shared_ptr<TiledBlob_Promise>	TiledBlob_PromisePtr;
 typedef std::weak_ptr<TiledBlob_Promise>	TiledBlob_PromisePtrW;
 typedef T_BlobRef<TiledBlob>				TiledBlobRef;
 typedef cti::continuable<TiledBlobRef>		AsyncTiledBlobRef;
+
+#undef UE_API
\ No newline at end of file
