Fix deprecation warnings

diff --git a/Engine/Plugins/TextureGraph/Source/TextureGraphEditor/Private/STG_EditorViewport.cpp b/Engine/Plugins/TextureGraph/Source/TextureGraphEditor/Private/STG_EditorViewport.cpp
--- a/Engine/Plugins/TextureGraph/Source/TextureGraphEditor/Private/STG_EditorViewport.cpp
+++ b/Engine/Plugins/TextureGraph/Source/TextureGraphEditor/Private/STG_EditorViewport.cpp
@@ -420,7 +420,7 @@ void STG_EditorViewport::PopulateViewportOverlays(TSharedRef<class SOverlay> Ove
 		.HAlign(HAlign_Right)
 		.Padding(5.0f)
 		[
-			BuildFeatureLevelWidget()
+			BuildShaderPlatformWidget()
 		];
 }
 
