Ran UnrealCodeFixup to fix dll storage

diff --git a/Engine/Plugins/Animation/AnimatorKit/Source/AnimatorKitSettings/Public/AnimatorKitSettings.h b/Engine/Plugins/Animation/AnimatorKit/Source/AnimatorKitSettings/Public/AnimatorKitSettings.h
--- a/Engine/Plugins/Animation/AnimatorKit/Source/AnimatorKitSettings/Public/AnimatorKitSettings.h
+++ b/Engine/Plugins/Animation/AnimatorKit/Source/AnimatorKitSettings/Public/AnimatorKitSettings.h
@@ -7,8 +7,10 @@
 
 #include "AnimatorKitSettings.generated.h"
 
-UCLASS(config=EditorPerProjectUserSettings, meta = (Experimental, DisplayName = "Animator Kit Settings"))
-class ANIMATORKITSETTINGS_API UAnimatorKitSettings : public UDeveloperSettings
+#define UE_API ANIMATORKITSETTINGS_API
+
+UCLASS(MinimalAPI, config=EditorPerProjectUserSettings, meta = (Experimental, DisplayName = "Animator Kit Settings"))
+class UAnimatorKitSettings : public UDeveloperSettings
 {
 	GENERATED_UCLASS_BODY()
 
@@ -18,8 +20,10 @@ public:
 	bool bEnableFocusMode = true;
     
 	DECLARE_MULTICAST_DELEGATE_OneParam(FOnUpdateSettings,  const UAnimatorKitSettings*);
-	static FOnUpdateSettings OnSettingsChange;
+	static UE_API FOnUpdateSettings OnSettingsChange;
+
+	UE_API virtual void PostInitProperties() override;
+	UE_API virtual void PostEditChangeProperty(FPropertyChangedEvent& PropertyChangedEvent) override;
+};
 
-	virtual void PostInitProperties() override;
-	virtual void PostEditChangeProperty(FPropertyChangedEvent& PropertyChangedEvent) override;
-};
\ No newline at end of file
+#undef UE_API
diff --git a/Engine/Plugins/Editor/EditorSysConfigAssistant/Source/EditorSysConfigAssistant/Public/EditorSysConfigAssistantSubsystem.h b/Engine/Plugins/Editor/EditorSysConfigAssistant/Source/EditorSysConfigAssistant/Public/EditorSysConfigAssistantSubsystem.h
--- a/Engine/Plugins/Editor/EditorSysConfigAssistant/Source/EditorSysConfigAssistant/Public/EditorSysConfigAssistantSubsystem.h
+++ b/Engine/Plugins/Editor/EditorSysConfigAssistant/Source/EditorSysConfigAssistant/Public/EditorSysConfigAssistantSubsystem.h
@@ -8,45 +8,47 @@
 
 #include "EditorSysConfigAssistantSubsystem.generated.h"
 
+#define UE_API EDITORSYSCONFIGASSISTANT_API
+
 class IModularFeature;
 class SNotificationItem;
 
-UCLASS()
-class EDITORSYSCONFIGASSISTANT_API UEditorSysConfigAssistantSubsystem : public UEditorSubsystem
+UCLASS(MinimalAPI)
+class UEditorSysConfigAssistantSubsystem : public UEditorSubsystem
 {
 	GENERATED_BODY()
 
 public:
-	virtual void Initialize(FSubsystemCollectionBase& Collection) override;
-	virtual void Deinitialize() override;
-	virtual bool ShouldCreateSubsystem(UObject* Outer) const override;
+	UE_API virtual void Initialize(FSubsystemCollectionBase& Collection) override;
+	UE_API virtual void Deinitialize() override;
+	UE_API virtual bool ShouldCreateSubsystem(UObject* Outer) const override;
 
 	/** Can be called on any thread */
-	void AddIssue(const FEditorSysConfigIssue& Issue);
+	UE_API void AddIssue(const FEditorSysConfigIssue& Issue);
 	/** Can be called on any thread */
-	TArray<TSharedPtr<FEditorSysConfigIssue>> GetIssues();
+	UE_API TArray<TSharedPtr<FEditorSysConfigIssue>> GetIssues();
 
 	/** Must be called on the game thread */
-	void ApplySysConfigChanges(TArrayView<const TSharedPtr<FEditorSysConfigIssue>> Issues);
+	UE_API void ApplySysConfigChanges(TArrayView<const TSharedPtr<FEditorSysConfigIssue>> Issues);
 
 	/** Must be called on the game thread */
-	void DismissSystemConfigNotification();
+	UE_API void DismissSystemConfigNotification();
 
 private:
-	void HandleModularFeatureRegistered(const FName& InFeatureName, IModularFeature* InFeature);
-	void HandleModularFeatureUnregistered(const FName& InFeatureName, IModularFeature* InFeature);
+	UE_API void HandleModularFeatureRegistered(const FName& InFeatureName, IModularFeature* InFeature);
+	UE_API void HandleModularFeatureUnregistered(const FName& InFeatureName, IModularFeature* InFeature);
 	
-	void HandleAssistantInitializationEvent();
+	UE_API void HandleAssistantInitializationEvent();
 
 	/** Must be called on the game thread */
-	static void NotifySystemConfigIssues();
+	static UE_API void NotifySystemConfigIssues();
 
 	/** Must be called on the game thread */
-	void NotifyRestart(bool bApplicationOnly);
+	UE_API void NotifyRestart(bool bApplicationOnly);
 
-	void OnApplicationRestartClicked();
-	void OnSystemRestartClicked();
-	void OnRestartDismissClicked();
+	UE_API void OnApplicationRestartClicked();
+	UE_API void OnSystemRestartClicked();
+	UE_API void OnRestartDismissClicked();
 
 	FRWLock IssuesLock;
 	TArray<TSharedPtr<FEditorSysConfigIssue>> Issues;
@@ -54,3 +56,5 @@ private:
 	TSharedPtr<SNotificationItem> IssueNotificationItem;
 	TWeakPtr<SNotificationItem> RestartNotificationItem;
 };
+
+#undef UE_API
diff --git a/Engine/Plugins/Experimental/ImageWidgets/Source/ImageWidgets/Public/SImageCatalog.h b/Engine/Plugins/Experimental/ImageWidgets/Source/ImageWidgets/Public/SImageCatalog.h
--- a/Engine/Plugins/Experimental/ImageWidgets/Source/ImageWidgets/Public/SImageCatalog.h
+++ b/Engine/Plugins/Experimental/ImageWidgets/Source/ImageWidgets/Public/SImageCatalog.h
@@ -1,4 +1,4 @@
-﻿// Copyright Epic Games, Inc. All Rights Reserved.
+// Copyright Epic Games, Inc. All Rights Reserved.
 
 #pragma once
 
@@ -6,14 +6,16 @@
 #include <Framework/Views/ITypedTableView.h>
 #include <Widgets/SCompoundWidget.h>
 
+#define UE_API IMAGEWIDGETS_API
+
 namespace UE::ImageWidgets
 {
 	/**
 	 * Contains all data for a catalog item.
 	 */
-	struct IMAGEWIDGETS_API FImageCatalogItemData
+	struct FImageCatalogItemData
 	{
-		FImageCatalogItemData(FGuid Guid, const FSlateBrush& Brush, const FText& Name, const FText& Info, const FText& ToolTip);
+		UE_API FImageCatalogItemData(FGuid Guid, const FSlateBrush& Brush, const FText& Name, const FText& Info, const FText& ToolTip);
 
 		/** Unique identifier for the catalog item */
 		FGuid Guid;
@@ -372,3 +374,5 @@ namespace UE::ImageWidgets
 		TPimplPtr<class FImpl> Impl;
 	};
 }
+
+#undef UE_API
diff --git a/Engine/Plugins/TextureGraph/Source/TextureGraph/Public/Blueprint/TG_AsyncExportTask.h b/Engine/Plugins/TextureGraph/Source/TextureGraph/Public/Blueprint/TG_AsyncExportTask.h
--- a/Engine/Plugins/TextureGraph/Source/TextureGraph/Public/Blueprint/TG_AsyncExportTask.h
+++ b/Engine/Plugins/TextureGraph/Source/TextureGraph/Public/Blueprint/TG_AsyncExportTask.h
@@ -13,21 +13,23 @@
 #include "TG_AsyncRenderTask.h"
 #include "TG_AsyncExportTask.generated.h"
 
+#define UE_API TEXTUREGRAPH_API
+
 DECLARE_DYNAMIC_MULTICAST_DELEGATE(FTSExportTaskDelegate);
 
-UCLASS()
-class TEXTUREGRAPH_API UTG_AsyncExportTask : public UTG_AsyncRenderTaskBase
+UCLASS(MinimalAPI)
+class UTG_AsyncExportTask : public UTG_AsyncRenderTaskBase
 {
 	GENERATED_UCLASS_BODY()
 
 public:
 	UFUNCTION(BlueprintCallable, Category = "TextureGraph" , meta=(DisplayName="Texture Graph Export (Async)", BlueprintInternalUseOnly = "true"))
-	static UTG_AsyncExportTask* TG_AsyncExportTask(UTextureGraphBase* TextureGraph, const bool OverwriteTextures, const bool bSave, const bool bExportAll, const bool bDisableCache = false);
+	static UE_API UTG_AsyncExportTask* TG_AsyncExportTask(UTextureGraphBase* TextureGraph, const bool OverwriteTextures, const bool bSave, const bool bExportAll, const bool bDisableCache = false);
 
-	virtual void Activate() override;
-	virtual const TArray<UTextureRenderTarget2D*>& ActivateBlocking(JobBatchPtr Batch);
+	UE_API virtual void Activate() override;
+	UE_API virtual const TArray<UTextureRenderTarget2D*>& ActivateBlocking(JobBatchPtr Batch);
 
-	virtual void FinishDestroy() override;
+	UE_API virtual void FinishDestroy() override;
 
 	UPROPERTY(BlueprintAssignable, Category = "TextureGraph")
 	FTSExportTaskDelegate OnDone;
@@ -36,10 +38,11 @@ private:
 
 	bool bBlobberCachingState;
 	UFUNCTION()
-	void OnExportDone();
+	UE_API void OnExportDone();
 	bool bOverwriteTextures = true;
 	bool bSave = false;
 	bool bExportAll = false;
 	FExportSettings TargetExportSettings;
 };
 
+#undef UE_API
diff --git a/Engine/Plugins/TextureGraph/Source/TextureGraph/Public/Blueprint/TG_AsyncRenderTask.h b/Engine/Plugins/TextureGraph/Source/TextureGraph/Public/Blueprint/TG_AsyncRenderTask.h
--- a/Engine/Plugins/TextureGraph/Source/TextureGraph/Public/Blueprint/TG_AsyncRenderTask.h
+++ b/Engine/Plugins/TextureGraph/Source/TextureGraph/Public/Blueprint/TG_AsyncRenderTask.h
@@ -14,27 +14,28 @@
 #include "Engine/World.h"
 #include "TG_AsyncRenderTask.generated.h"
 
+#define UE_API TEXTUREGRAPH_API
+
 DECLARE_DYNAMIC_MULTICAST_DELEGATE_OneParam(FTSRenderTaskDelegate, const TArray<UTextureRenderTarget2D*>&, OutputRts);
 
-UCLASS()
-class TEXTUREGRAPH_API UTG_AsyncRenderTaskBase : public UTG_AsyncTask
+UCLASS(MinimalAPI)
+class UTG_AsyncRenderTaskBase : public UTG_AsyncTask
 {
 	GENERATED_UCLASS_BODY()
 
 public:
-	virtual const TArray<UTextureRenderTarget2D*>& ActivateBlocking(JobBatchPtr Batch);
-	virtual void FinishDestroy() override;
-	virtual void SetReadyToDestroy() override;
+	UE_API virtual const TArray<UTextureRenderTarget2D*>& ActivateBlocking(JobBatchPtr Batch);
+	UE_API virtual void FinishDestroy() override;
+	UE_API virtual void SetReadyToDestroy() override;
 
-	void OnWorldCleanup(UWorld* World, bool bSessionEnded, bool CleanupResources);
+	UE_API void OnWorldCleanup(UWorld* World, bool bSessionEnded, bool CleanupResources);
 
 protected:
-	AsyncBool FinalizeAllOutputBlobs();
-	JobBatchPtr PrepareActivate(JobBatchPtr Batch, bool bIsAsync);
-	void GatherAllOutputBlobs();
-	void GatherAllRenderTargets();
-	AsyncBool GetRenderTextures();
-
+	UE_API AsyncBool FinalizeAllOutputBlobs();
+	UE_API JobBatchPtr PrepareActivate(JobBatchPtr Batch, bool bIsAsync);
+	UE_API void GatherAllOutputBlobs();
+	UE_API void GatherAllRenderTargets();
+	UE_API AsyncBool GetRenderTextures();
 
 	TArray<UTextureRenderTarget2D*> OutputRts;
 	TArray<BlobPtr> OutputBlobs;
@@ -45,18 +46,19 @@ protected:
 };
 
 /////////////////////////////////////////////////////////////////////
-UCLASS()
-class TEXTUREGRAPH_API UTG_AsyncRenderTask : public UTG_AsyncRenderTaskBase
+UCLASS(MinimalAPI)
+class UTG_AsyncRenderTask : public UTG_AsyncRenderTaskBase
 {
 	GENERATED_UCLASS_BODY()
 
 public:
 	UFUNCTION(BlueprintCallable, Category = "TextureGraph" , meta=(DisplayName="Texture Graph Render (Async)" , BlueprintInternalUseOnly = "true"))
-	static UTG_AsyncRenderTask* TG_AsyncRenderTask(UTextureGraphBase* InTextureGraph);
+	static UE_API UTG_AsyncRenderTask* TG_AsyncRenderTask(UTextureGraphBase* InTextureGraph);
 
-	virtual void Activate() override;
+	UE_API virtual void Activate() override;
 
 	UPROPERTY(BlueprintAssignable, Category = "TextureGraph")
 	FTSRenderTaskDelegate OnDone;
 };
 
+#undef UE_API
diff --git a/Engine/Plugins/TextureGraph/Source/TextureGraph/Public/Blueprint/TG_AsyncTask.h b/Engine/Plugins/TextureGraph/Source/TextureGraph/Public/Blueprint/TG_AsyncTask.h
--- a/Engine/Plugins/TextureGraph/Source/TextureGraph/Public/Blueprint/TG_AsyncTask.h
+++ b/Engine/Plugins/TextureGraph/Source/TextureGraph/Public/Blueprint/TG_AsyncTask.h
@@ -7,26 +7,29 @@
 #include "Kismet/BlueprintFunctionLibrary.h"
 #include "TG_AsyncTask.generated.h"
 
+#define UE_API TEXTUREGRAPH_API
+
 /**
 * BlueprintCallable factory functions for Async TG_Script Actions. Helps regiser with TG_AsyncTaskManager to keep Task alive and manage its life cycle 
 */
 
-UCLASS()
-class TEXTUREGRAPH_API UTG_AsyncTask : public UBlueprintAsyncActionBase
+UCLASS(MinimalAPI)
+class UTG_AsyncTask : public UBlueprintAsyncActionBase
 {
 	GENERATED_UCLASS_BODY()
 
 public:
 
-	virtual void Activate() override;
+	UE_API virtual void Activate() override;
 	
 	/**
 	 * Call to register this task with a AsyncTaskManager, it will not be destroyed until SetReadyToDestroy is called
 	 * This allows having an action stay alive until SetReadyToDestroy is manually called, allowing it to be used inside loops or if the calling BP goes away
 	 */
-	virtual void RegisterWithTGAsyncTaskManger();
+	UE_API virtual void RegisterWithTGAsyncTaskManger();
 
 	/** Call when the action is completely done, this makes the action free to delete, and will unregister it with the game instance and AsyncTaskManager  */
-	virtual void SetReadyToDestroy() override;
+	UE_API virtual void SetReadyToDestroy() override;
 };
 
+#undef UE_API
diff --git a/Engine/Plugins/TextureGraph/Source/TextureGraph/Public/Blueprint/TG_AsyncTaskManager.h b/Engine/Plugins/TextureGraph/Source/TextureGraph/Public/Blueprint/TG_AsyncTaskManager.h
--- a/Engine/Plugins/TextureGraph/Source/TextureGraph/Public/Blueprint/TG_AsyncTaskManager.h
+++ b/Engine/Plugins/TextureGraph/Source/TextureGraph/Public/Blueprint/TG_AsyncTaskManager.h
@@ -10,7 +10,7 @@
 * Keeps reference to the TG_AsyncTasks to manage their life cycle
 */
 
-class TEXTUREGRAPH_API FTG_AsyncTaskManager
+class FTG_AsyncTaskManager
 {
 private:
 	FTG_AsyncTaskManager() {}  // Private constructor to enforce singleton pattern.
diff --git a/Engine/Plugins/TextureGraph/Source/TextureGraph/Public/Expressions/Adjustment/TG_Expression_NormalFromHeightMap.h b/Engine/Plugins/TextureGraph/Source/TextureGraph/Public/Expressions/Adjustment/TG_Expression_NormalFromHeightMap.h
--- a/Engine/Plugins/TextureGraph/Source/TextureGraph/Public/Expressions/Adjustment/TG_Expression_NormalFromHeightMap.h
+++ b/Engine/Plugins/TextureGraph/Source/TextureGraph/Public/Expressions/Adjustment/TG_Expression_NormalFromHeightMap.h
@@ -9,15 +9,17 @@
 
 #include "TG_Expression_NormalFromHeightMap.generated.h"
 
-UCLASS()
-class TEXTUREGRAPH_API UTG_Expression_NormalFromHeightMap : public UTG_Expression
+#define UE_API TEXTUREGRAPH_API
+
+UCLASS(MinimalAPI)
+class UTG_Expression_NormalFromHeightMap : public UTG_Expression
 {
 	GENERATED_BODY()
 
 public:
 	TG_DECLARE_EXPRESSION(TG_Category::Adjustment);
-	virtual void						Evaluate(FTG_EvaluationContext* InContext) override;
-	virtual bool						Validate(MixUpdateCyclePtr Cycle) override;
+	UE_API virtual void						Evaluate(FTG_EvaluationContext* InContext) override;
+	UE_API virtual bool						Validate(MixUpdateCyclePtr Cycle) override;
 
 	// Normals from height are generated by taking nearby points from the texture having some offset. Using smaller offset values means we are getting samples closer to the surface from the texture. 
 	UPROPERTY(EditAnywhere, Category = NoCategory, meta = (TGType = "TG_Setting", UIMin = "0.001", ClampMin = "0.001", UIMax = "0.1", ClampMax = "0.1"))
@@ -39,3 +41,4 @@ public:
 	virtual FText						GetTooltipText() const override { return FText::FromString(TEXT("Converts the input into normal.")); }
 };
 
+#undef UE_API
diff --git a/Engine/Plugins/TextureGraph/Source/TextureGraph/Public/Expressions/Arrays/TG_Expression_Array.h b/Engine/Plugins/TextureGraph/Source/TextureGraph/Public/Expressions/Arrays/TG_Expression_Array.h
--- a/Engine/Plugins/TextureGraph/Source/TextureGraph/Public/Expressions/Arrays/TG_Expression_Array.h
+++ b/Engine/Plugins/TextureGraph/Source/TextureGraph/Public/Expressions/Arrays/TG_Expression_Array.h
@@ -7,14 +7,16 @@
 #include "TG_Texture.h"
 #include "TG_Expression_Array.generated.h"
 
-UCLASS()
-class TEXTUREGRAPH_API UTG_Expression_Array4 : public UTG_Expression
+#define UE_API TEXTUREGRAPH_API
+
+UCLASS(MinimalAPI)
+class UTG_Expression_Array4 : public UTG_Expression
 {
 	 GENERATED_BODY()
 
 public:
 	TG_DECLARE_EXPRESSION(TG_Category::Arrays);
-	virtual void					Evaluate(FTG_EvaluationContext* InContext) override;
+	UE_API virtual void					Evaluate(FTG_EvaluationContext* InContext) override;
 	
 	/// Texture at the first index of the output texture array
 	UPROPERTY(meta = (TGType = "TG_Input", PinDisplayName = ""))
@@ -40,3 +42,4 @@ public:
 	virtual FText					GetTooltipText() const override { return FText::FromString(TEXT("Combines 4 inputs of any type and constructs an array out of it.")); } 
  };
 
+#undef UE_API
diff --git a/Engine/Plugins/TextureGraph/Source/TextureGraph/Public/Expressions/Arrays/TG_Expression_ArrayGrid.h b/Engine/Plugins/TextureGraph/Source/TextureGraph/Public/Expressions/Arrays/TG_Expression_ArrayGrid.h
--- a/Engine/Plugins/TextureGraph/Source/TextureGraph/Public/Expressions/Arrays/TG_Expression_ArrayGrid.h
+++ b/Engine/Plugins/TextureGraph/Source/TextureGraph/Public/Expressions/Arrays/TG_Expression_ArrayGrid.h
@@ -7,14 +7,16 @@
 #include "TG_Texture.h"
 #include "TG_Expression_ArrayGrid.generated.h"
 
-UCLASS()
-class TEXTUREGRAPH_API UTG_Expression_ArrayGrid : public UTG_Expression
+#define UE_API TEXTUREGRAPH_API
+
+UCLASS(MinimalAPI)
+class UTG_Expression_ArrayGrid : public UTG_Expression
 {
 	 GENERATED_BODY()
 
 public:
 	TG_DECLARE_EXPRESSION(TG_Category::Arrays);
-	virtual void					Evaluate(FTG_EvaluationContext* InContext) override;
+	UE_API virtual void					Evaluate(FTG_EvaluationContext* InContext) override;
 	
 	/// The input texture array. Must be an array object of type textures
 	UPROPERTY(meta = (TGType = "TG_Input", PinDisplayName = "", HideInnerPropertiesInNode))
@@ -40,3 +42,4 @@ public:
 	virtual FText					GetTooltipText() const override { return FText::FromString(TEXT("Arranges an input texture array into an MxN sized grid.")); } 
  };
 
+#undef UE_API
diff --git a/Engine/Plugins/TextureGraph/Source/TextureGraph/Public/Expressions/Arrays/TG_Expression_ArrayOps.h b/Engine/Plugins/TextureGraph/Source/TextureGraph/Public/Expressions/Arrays/TG_Expression_ArrayOps.h
--- a/Engine/Plugins/TextureGraph/Source/TextureGraph/Public/Expressions/Arrays/TG_Expression_ArrayOps.h
+++ b/Engine/Plugins/TextureGraph/Source/TextureGraph/Public/Expressions/Arrays/TG_Expression_ArrayOps.h
@@ -7,17 +7,19 @@
 #include "TG_Texture.h"
 #include "TG_Expression_ArrayOps.generated.h"
 
+#define UE_API TEXTUREGRAPH_API
+
 /////////////////////////////////////////////////////////////
 /// Array concatenation
 /////////////////////////////////////////////////////////////
-UCLASS()
-class TEXTUREGRAPH_API UTG_Expression_ArrayConcat : public UTG_Expression
+UCLASS(MinimalAPI)
+class UTG_Expression_ArrayConcat : public UTG_Expression
 {
 	 GENERATED_BODY()
 
 public:
 	TG_DECLARE_EXPRESSION(TG_Category::Arrays);
-	virtual void					Evaluate(FTG_EvaluationContext* InContext) override;
+	UE_API virtual void					Evaluate(FTG_EvaluationContext* InContext) override;
 	
 	/// The first array as part of the concatenation process
 	UPROPERTY(meta = (TGType = "TG_Input", PinDisplayName = "Input-1"))
@@ -42,14 +44,14 @@ public:
 /////////////////////////////////////////////////////////////
 /// Split array
 /////////////////////////////////////////////////////////////
-UCLASS()
-class TEXTUREGRAPH_API UTG_Expression_ArraySplit: public UTG_Expression
+UCLASS(MinimalAPI)
+class UTG_Expression_ArraySplit: public UTG_Expression
 {
 	 GENERATED_BODY()
 
 public:
 	TG_DECLARE_EXPRESSION(TG_Category::Arrays);
-	virtual void					Evaluate(FTG_EvaluationContext* InContext) override;
+	UE_API virtual void					Evaluate(FTG_EvaluationContext* InContext) override;
 	
 	/// The array to slice
 	UPROPERTY(meta = (TGType = "TG_Input", PinDisplayName = "", MinValue))
@@ -79,3 +81,4 @@ public:
 	virtual FText					GetTooltipText() const override { return FText::FromString(TEXT("Slice a portion of an input array and put that as the output of this array. Alternatively splice an array and remove some elements from the input array")); } 
  };
 
+#undef UE_API
diff --git a/Engine/Plugins/TextureGraph/Source/TextureGraph/Public/Expressions/Channel/TG_Expression_ChannelCombiner.h b/Engine/Plugins/TextureGraph/Source/TextureGraph/Public/Expressions/Channel/TG_Expression_ChannelCombiner.h
--- a/Engine/Plugins/TextureGraph/Source/TextureGraph/Public/Expressions/Channel/TG_Expression_ChannelCombiner.h
+++ b/Engine/Plugins/TextureGraph/Source/TextureGraph/Public/Expressions/Channel/TG_Expression_ChannelCombiner.h
@@ -9,14 +9,16 @@
 
 #include "TG_Expression_ChannelCombiner.generated.h"
 
-UCLASS()
-class TEXTUREGRAPH_API UTG_Expression_ChannelCombiner : public UTG_Expression
+#define UE_API TEXTUREGRAPH_API
+
+UCLASS(MinimalAPI)
+class UTG_Expression_ChannelCombiner : public UTG_Expression
 {
 	GENERATED_BODY()
 
 public:
 	TG_DECLARE_EXPRESSION(TG_Category::Channel);
-	virtual void			Evaluate(FTG_EvaluationContext* InContext) override;
+	UE_API virtual void			Evaluate(FTG_EvaluationContext* InContext) override;
 
 	// The image that will end up in the red channel of the output image
 	UPROPERTY(meta = (TGType = "TG_Input", PinDisplayName = "R"))
@@ -41,3 +43,5 @@ public:
 	virtual FTG_Name		GetDefaultName() const override { return TEXT("CombineChannels");}
 	virtual FText			GetTooltipText() const override { return FText::FromString(TEXT("Combines individual R, G, B, A channels into an RGBA image.")); } 
 };
+
+#undef UE_API
diff --git a/Engine/Plugins/TextureGraph/Source/TextureGraph/Public/Expressions/Channel/TG_Expression_ChannelSplitter.h b/Engine/Plugins/TextureGraph/Source/TextureGraph/Public/Expressions/Channel/TG_Expression_ChannelSplitter.h
--- a/Engine/Plugins/TextureGraph/Source/TextureGraph/Public/Expressions/Channel/TG_Expression_ChannelSplitter.h
+++ b/Engine/Plugins/TextureGraph/Source/TextureGraph/Public/Expressions/Channel/TG_Expression_ChannelSplitter.h
@@ -9,14 +9,16 @@
 
 #include "TG_Expression_ChannelSplitter.generated.h"
 
-UCLASS()
-class TEXTUREGRAPH_API UTG_Expression_ChannelSplitter : public UTG_Expression
+#define UE_API TEXTUREGRAPH_API
+
+UCLASS(MinimalAPI)
+class UTG_Expression_ChannelSplitter : public UTG_Expression
 {
 	GENERATED_BODY()
 
 public:
 	TG_DECLARE_EXPRESSION(TG_Category::Channel);
-	virtual void			Evaluate(FTG_EvaluationContext* InContext) override;
+	UE_API virtual void			Evaluate(FTG_EvaluationContext* InContext) override;
 
 	// The input texture
 	UPROPERTY(meta = (TGType = "TG_Input", PinDisplayName = "RGBA"))
@@ -42,3 +44,4 @@ public:
 	virtual FText			GetTooltipText() const override { return FText::FromString(TEXT("Splits an RGBA image into individual R, G, B, A channels.")); } 
 };
 
+#undef UE_API
diff --git a/Engine/Plugins/TextureGraph/Source/TextureGraph/Public/Expressions/Channel/TG_Expression_ChannelSwizzle.h b/Engine/Plugins/TextureGraph/Source/TextureGraph/Public/Expressions/Channel/TG_Expression_ChannelSwizzle.h
--- a/Engine/Plugins/TextureGraph/Source/TextureGraph/Public/Expressions/Channel/TG_Expression_ChannelSwizzle.h
+++ b/Engine/Plugins/TextureGraph/Source/TextureGraph/Public/Expressions/Channel/TG_Expression_ChannelSwizzle.h
@@ -10,14 +10,16 @@
 
 #include "TG_Expression_ChannelSwizzle.generated.h"
 
-UCLASS()
-class TEXTUREGRAPH_API UTG_Expression_ChannelSwizzle : public UTG_Expression
+#define UE_API TEXTUREGRAPH_API
+
+UCLASS(MinimalAPI)
+class UTG_Expression_ChannelSwizzle : public UTG_Expression
 {
 	GENERATED_BODY()
 
 public:
 	TG_DECLARE_EXPRESSION(TG_Category::Channel);
-	virtual void			Evaluate(FTG_EvaluationContext* InContext) override;
+	UE_API virtual void			Evaluate(FTG_EvaluationContext* InContext) override;
 
 	// The input texture
 	UPROPERTY(meta = (TGType = "TG_Input", PinDisplayName = ""))
@@ -48,3 +50,4 @@ public:
 	virtual FText			GetTooltipText() const override { return FText::FromString(TEXT("Performs a swizzle operation.")); } 
 };
 
+#undef UE_API
diff --git a/Engine/Plugins/TextureGraph/Source/TextureGraph/Public/Expressions/Color/TG_Expression_ColorCorrection.h b/Engine/Plugins/TextureGraph/Source/TextureGraph/Public/Expressions/Color/TG_Expression_ColorCorrection.h
--- a/Engine/Plugins/TextureGraph/Source/TextureGraph/Public/Expressions/Color/TG_Expression_ColorCorrection.h
+++ b/Engine/Plugins/TextureGraph/Source/TextureGraph/Public/Expressions/Color/TG_Expression_ColorCorrection.h
@@ -8,14 +8,16 @@
 #include "TG_Texture.h"
 #include "TG_Expression_ColorCorrection.generated.h"
 
-UCLASS()
-class TEXTUREGRAPH_API UTG_Expression_ColorCorrection : public UTG_Expression
+#define UE_API TEXTUREGRAPH_API
+
+UCLASS(MinimalAPI)
+class UTG_Expression_ColorCorrection : public UTG_Expression
 {
 	GENERATED_BODY()
 
 public:
 	TG_DECLARE_EXPRESSION(TG_Category::Adjustment);
-	virtual void						Evaluate(FTG_EvaluationContext* InContext) override;
+	UE_API virtual void						Evaluate(FTG_EvaluationContext* InContext) override;
 
 	// The input image to be color corrected
 	UPROPERTY(meta = (TGType = "TG_Input", PinDisplayName = ""))
@@ -58,3 +60,4 @@ public:
 
 };
 
+#undef UE_API
diff --git a/Engine/Plugins/TextureGraph/Source/TextureGraph/Public/Expressions/Color/TG_Expression_Grayscale.h b/Engine/Plugins/TextureGraph/Source/TextureGraph/Public/Expressions/Color/TG_Expression_Grayscale.h
--- a/Engine/Plugins/TextureGraph/Source/TextureGraph/Public/Expressions/Color/TG_Expression_Grayscale.h
+++ b/Engine/Plugins/TextureGraph/Source/TextureGraph/Public/Expressions/Color/TG_Expression_Grayscale.h
@@ -9,14 +9,16 @@
 
 #include "TG_Expression_Grayscale.generated.h"
 
-UCLASS()
-class TEXTUREGRAPH_API UTG_Expression_Grayscale : public UTG_Expression
+#define UE_API TEXTUREGRAPH_API
+
+UCLASS(MinimalAPI)
+class UTG_Expression_Grayscale : public UTG_Expression
 {
 	GENERATED_BODY()
 
 public:
 	TG_DECLARE_EXPRESSION(TG_Category::Adjustment);
-	virtual void						Evaluate(FTG_EvaluationContext* InContext) override;
+	UE_API virtual void						Evaluate(FTG_EvaluationContext* InContext) override;
 
 	// The input image to be converted to grayscale
 	UPROPERTY(meta = (TGType = "TG_Input", PinDisplayName = ""))
@@ -31,3 +33,4 @@ public:
 
 };
 
+#undef UE_API
diff --git a/Engine/Plugins/TextureGraph/Source/TextureGraph/Public/Expressions/Color/TG_Expression_HSV.h b/Engine/Plugins/TextureGraph/Source/TextureGraph/Public/Expressions/Color/TG_Expression_HSV.h
--- a/Engine/Plugins/TextureGraph/Source/TextureGraph/Public/Expressions/Color/TG_Expression_HSV.h
+++ b/Engine/Plugins/TextureGraph/Source/TextureGraph/Public/Expressions/Color/TG_Expression_HSV.h
@@ -10,11 +10,13 @@
 
 #include "TG_Expression_HSV.generated.h"
 
+#define UE_API TEXTUREGRAPH_API
+
 //////////////////////////////////////////////////////////////////////////
 /// RGB2HSV Correction
 //////////////////////////////////////////////////////////////////////////
-UCLASS()
-class TEXTUREGRAPH_API UTG_Expression_RGB2HSV : public UTG_Expression_OneInput
+UCLASS(MinimalAPI)
+class UTG_Expression_RGB2HSV : public UTG_Expression_OneInput
 {
 	GENERATED_BODY()
 
@@ -26,15 +28,15 @@ public:
 protected:
 	/// This is unsupported because it only makes sense on a vector input
 	virtual float						EvaluateScalar_WithValue(FTG_EvaluationContext* InContext, const float* const Value, size_t Count) override { return 0.0f; }
-	virtual FVector4f					EvaluateVector_WithValue(FTG_EvaluationContext* InContext, const FVector4f* const Value, size_t Count);
-	virtual FTG_Texture					EvaluateTexture(FTG_EvaluationContext* InContext) override;
+	UE_API virtual FVector4f					EvaluateVector_WithValue(FTG_EvaluationContext* InContext, const FVector4f* const Value, size_t Count);
+	UE_API virtual FTG_Texture					EvaluateTexture(FTG_EvaluationContext* InContext) override;
 };
 
 //////////////////////////////////////////////////////////////////////////
 /// HSV2RGB Correction
 //////////////////////////////////////////////////////////////////////////
-UCLASS()
-class TEXTUREGRAPH_API UTG_Expression_HSV2RGB : public UTG_Expression_OneInput
+UCLASS(MinimalAPI)
+class UTG_Expression_HSV2RGB : public UTG_Expression_OneInput
 {
 	GENERATED_BODY()
 
@@ -46,15 +48,15 @@ public:
 protected:
 	/// This is unsupported because it only makes sense on a vector input
 	virtual float						EvaluateScalar_WithValue(FTG_EvaluationContext* InContext, const float* const Value, size_t Count) override { return 0.0f; }
-	virtual FVector4f					EvaluateVector_WithValue(FTG_EvaluationContext* InContext, const FVector4f* const Value, size_t Count);
-	virtual FTG_Texture					EvaluateTexture(FTG_EvaluationContext* InContext) override;
+	UE_API virtual FVector4f					EvaluateVector_WithValue(FTG_EvaluationContext* InContext, const FVector4f* const Value, size_t Count);
+	UE_API virtual FTG_Texture					EvaluateTexture(FTG_EvaluationContext* InContext) override;
 };
 
 //////////////////////////////////////////////////////////////////////////
 /// HSV Correction
 //////////////////////////////////////////////////////////////////////////
-UCLASS()
-class TEXTUREGRAPH_API UTG_Expression_HSV : public UTG_Expression_OneInput
+UCLASS(MinimalAPI)
+class UTG_Expression_HSV : public UTG_Expression_OneInput
 {
 	GENERATED_BODY()
 
@@ -78,7 +80,8 @@ public:
 protected:
 	/// This is unsupported because it only makes sense on a vector input
 	virtual float						EvaluateScalar_WithValue(FTG_EvaluationContext* InContext, const float* const ValuePtr, size_t Count) override { return 0.0f; }
-	virtual FVector4f					EvaluateVector_WithValue(FTG_EvaluationContext* InContext, const FVector4f* const ValuePtr, size_t Count);
-	virtual FTG_Texture					EvaluateTexture(FTG_EvaluationContext* InContext) override;
+	UE_API virtual FVector4f					EvaluateVector_WithValue(FTG_EvaluationContext* InContext, const FVector4f* const ValuePtr, size_t Count);
+	UE_API virtual FTG_Texture					EvaluateTexture(FTG_EvaluationContext* InContext) override;
 };
 
+#undef UE_API
diff --git a/Engine/Plugins/TextureGraph/Source/TextureGraph/Public/Expressions/Color/TG_Expression_Premult.h b/Engine/Plugins/TextureGraph/Source/TextureGraph/Public/Expressions/Color/TG_Expression_Premult.h
--- a/Engine/Plugins/TextureGraph/Source/TextureGraph/Public/Expressions/Color/TG_Expression_Premult.h
+++ b/Engine/Plugins/TextureGraph/Source/TextureGraph/Public/Expressions/Color/TG_Expression_Premult.h
@@ -10,14 +10,16 @@
 
 #include "TG_Expression_Premult.generated.h"
 
-UCLASS()
-class TEXTUREGRAPH_API UTG_Expression_Premult : public UTG_Expression
+#define UE_API TEXTUREGRAPH_API
+
+UCLASS(MinimalAPI)
+class UTG_Expression_Premult : public UTG_Expression
 {
 	GENERATED_BODY()
 
 public:
 	TG_DECLARE_EXPRESSION(TG_Category::Adjustment);
-	virtual void						Evaluate(FTG_EvaluationContext* InContext) override;
+	UE_API virtual void						Evaluate(FTG_EvaluationContext* InContext) override;
 
 	// The input image to be used
 	UPROPERTY(meta = (TGType = "TG_Input", PinDisplayName = ""))
@@ -31,3 +33,4 @@ public:
 	virtual FText						GetTooltipText() const override { return FText::FromString(TEXT("Premultiply alpha with the RGB color values at each pixel. The alpha remains unchanged.")); } 
 };
 
+#undef UE_API
diff --git a/Engine/Plugins/TextureGraph/Source/TextureGraph/Public/Expressions/CommandChange/TG_PinConnectionChange.h b/Engine/Plugins/TextureGraph/Source/TextureGraph/Public/Expressions/CommandChange/TG_PinConnectionChange.h
--- a/Engine/Plugins/TextureGraph/Source/TextureGraph/Public/Expressions/CommandChange/TG_PinConnectionChange.h
+++ b/Engine/Plugins/TextureGraph/Source/TextureGraph/Public/Expressions/CommandChange/TG_PinConnectionChange.h
@@ -6,6 +6,8 @@
 #include "TG_SystemTypes.h"  
 #include "Misc/Change.h"
 
+#define UE_API TEXTUREGRAPH_API
+
 class UTG_Graph;
 class UTG_Node;
 class UTG_Pin;
@@ -13,11 +15,11 @@ class UTG_Pin;
 /** This class helps to Undo the pin connection correctly
 as the Blob is not serialized so if the expression property is not connected and is holding 
 a reference of blob it is not rest when undoing the connection*/
-class TEXTUREGRAPH_API FTG_PinConnectionChange : public FCommandChange 
+class FTG_PinConnectionChange : public FCommandChange 
 {
 public:
-	static void StoreChange(UTG_Graph* InGraph, UTG_Node& NodeFrom, FTG_Name& PinFromName);
-	FTG_PinConnectionChange(UTG_Graph* InGraph, UTG_Node& NodeFrom, FTG_Name& PinFromName);
+	static UE_API void StoreChange(UTG_Graph* InGraph, UTG_Node& NodeFrom, FTG_Name& PinFromName);
+	UE_API FTG_PinConnectionChange(UTG_Graph* InGraph, UTG_Node& NodeFrom, FTG_Name& PinFromName);
 
 	virtual FString ToString() const override
 	{
@@ -25,7 +27,7 @@ public:
 	}
 
 	virtual void Apply(UObject* Object) override {};
-	virtual void Revert(UObject* Object) override;
+	UE_API virtual void Revert(UObject* Object) override;
 
 private:
 	UTG_Graph* Graph;
@@ -35,21 +37,23 @@ private:
 /** This class helps to Redo the pin connection Break correctly
 as the Blob is not serialized so if the expression property is not connected and is holding
 a reference of blob it is not rest when redoing the connection*/
-class TEXTUREGRAPH_API FTG_PinConnectionBreakChange : public FCommandChange
+class FTG_PinConnectionBreakChange : public FCommandChange
 {
 public:
-	static void StoreChange(UTG_Graph* InGraph, UTG_Node& NodeFrom, FTG_Name& PinFromName);
-	FTG_PinConnectionBreakChange(UTG_Graph* InGraph, UTG_Node& NodeFrom, FTG_Name& PinFromName);
+	static UE_API void StoreChange(UTG_Graph* InGraph, UTG_Node& NodeFrom, FTG_Name& PinFromName);
+	UE_API FTG_PinConnectionBreakChange(UTG_Graph* InGraph, UTG_Node& NodeFrom, FTG_Name& PinFromName);
 
 	virtual FString ToString() const override
 	{
 		return FString(TEXT("Pin Connection Break"));
 	}
 
-	virtual void Apply(UObject* Object) override;
+	UE_API virtual void Apply(UObject* Object) override;
 	virtual void Revert(UObject* Object) override {};
 
 private:
 	UTG_Graph* Graph;
 	UTG_Pin* PinTo;
 };
+
+#undef UE_API
diff --git a/Engine/Plugins/TextureGraph/Source/TextureGraph/Public/Expressions/Filter/TG_Expression_Blur.h b/Engine/Plugins/TextureGraph/Source/TextureGraph/Public/Expressions/Filter/TG_Expression_Blur.h
--- a/Engine/Plugins/TextureGraph/Source/TextureGraph/Public/Expressions/Filter/TG_Expression_Blur.h
+++ b/Engine/Plugins/TextureGraph/Source/TextureGraph/Public/Expressions/Filter/TG_Expression_Blur.h
@@ -7,6 +7,8 @@
 
 #include "TG_Expression_Blur.generated.h"
 
+#define UE_API TEXTUREGRAPH_API
+
 UENUM(BlueprintType)
 enum class EBlurType : uint8
 {
@@ -15,8 +17,8 @@ enum class EBlurType : uint8
 	Radial = 2			UMETA(DisplayName = "Radial")
 };
 
-UCLASS()
-class TEXTUREGRAPH_API UTG_Expression_Blur : public UTG_Expression
+UCLASS(MinimalAPI)
+class UTG_Expression_Blur : public UTG_Expression
 {
 	GENERATED_BODY()
 
@@ -25,7 +27,7 @@ public:
 	
 #if WITH_EDITOR
 	// Used to implement EditCondition logic for both Node UI and Details View
-	virtual bool CanEditChange(const FProperty* InProperty) const override;
+	UE_API virtual bool CanEditChange(const FProperty* InProperty) const override;
 #endif
 
 	// Controls how much an texture is blurred
@@ -53,6 +55,7 @@ public:
 	FTG_Texture Input;
 
 	virtual FText GetTooltipText() const override { return FText::FromString(TEXT("Applies a blur filter with a given strength.")); }
-	virtual void Evaluate(FTG_EvaluationContext* InContext) override;
+	UE_API virtual void Evaluate(FTG_EvaluationContext* InContext) override;
 };
 
+#undef UE_API
diff --git a/Engine/Plugins/TextureGraph/Source/TextureGraph/Public/Expressions/Filter/TG_Expression_Brightness.h b/Engine/Plugins/TextureGraph/Source/TextureGraph/Public/Expressions/Filter/TG_Expression_Brightness.h
--- a/Engine/Plugins/TextureGraph/Source/TextureGraph/Public/Expressions/Filter/TG_Expression_Brightness.h
+++ b/Engine/Plugins/TextureGraph/Source/TextureGraph/Public/Expressions/Filter/TG_Expression_Brightness.h
@@ -10,17 +10,19 @@
 
 #include "TG_Expression_Brightness.generated.h"
 
+#define UE_API TEXTUREGRAPH_API
+
 //////////////////////////////////////////////////////////////////////////
 /// Expression
 //////////////////////////////////////////////////////////////////////////
-UCLASS()
-class TEXTUREGRAPH_API UTG_Expression_Brightness : public UTG_Expression
+UCLASS(MinimalAPI)
+class UTG_Expression_Brightness : public UTG_Expression
 {
 	GENERATED_BODY()
 
 public:
 	TG_DECLARE_EXPRESSION(TG_Category::Adjustment);
-	virtual void							Evaluate(FTG_EvaluationContext* InContext) override;
+	UE_API virtual void							Evaluate(FTG_EvaluationContext* InContext) override;
 
 	// Adjusts the overall lightness or darkness of an image. Increasing brightness makes it brighter, while decreasing brightness makes it darker.
 	UPROPERTY(EditAnywhere, Category = NoCategory, meta = (TGType = "TG_Setting", UIMin = "-1", ClampMin = "-1", UIMax = "1", ClampMax = "1"))
@@ -41,3 +43,4 @@ public:
 	virtual FText							GetTooltipText() const override { return FText::FromString(TEXT("Changes the brightness and contrast of the input.")); } 
 };
 
+#undef UE_API
diff --git a/Engine/Plugins/TextureGraph/Source/TextureGraph/Public/Expressions/Filter/TG_Expression_EdgeDetect.h b/Engine/Plugins/TextureGraph/Source/TextureGraph/Public/Expressions/Filter/TG_Expression_EdgeDetect.h
--- a/Engine/Plugins/TextureGraph/Source/TextureGraph/Public/Expressions/Filter/TG_Expression_EdgeDetect.h
+++ b/Engine/Plugins/TextureGraph/Source/TextureGraph/Public/Expressions/Filter/TG_Expression_EdgeDetect.h
@@ -7,8 +7,10 @@
 
 #include "TG_Expression_EdgeDetect.generated.h"
 
-UCLASS()
-class TEXTUREGRAPH_API UTG_Expression_EdgeDetect : public UTG_Expression
+#define UE_API TEXTUREGRAPH_API
+
+UCLASS(MinimalAPI)
+class UTG_Expression_EdgeDetect : public UTG_Expression
 {
 	GENERATED_BODY()
 
@@ -28,6 +30,7 @@ public:
 	FTG_Texture Input;
 
 	virtual FText GetTooltipText() const override { return FText::FromString(TEXT("Detects edges (areas of sharp color change) from the input.")); }
-	virtual void Evaluate(FTG_EvaluationContext* InContext) override;
+	UE_API virtual void Evaluate(FTG_EvaluationContext* InContext) override;
 };
 
+#undef UE_API
diff --git a/Engine/Plugins/TextureGraph/Source/TextureGraph/Public/Expressions/Filter/TG_Expression_ErodeDilate.h b/Engine/Plugins/TextureGraph/Source/TextureGraph/Public/Expressions/Filter/TG_Expression_ErodeDilate.h
--- a/Engine/Plugins/TextureGraph/Source/TextureGraph/Public/Expressions/Filter/TG_Expression_ErodeDilate.h
+++ b/Engine/Plugins/TextureGraph/Source/TextureGraph/Public/Expressions/Filter/TG_Expression_ErodeDilate.h
@@ -10,6 +10,8 @@
 
 #include "TG_Expression_ErodeDilate.generated.h"
 
+#define UE_API TEXTUREGRAPH_API
+
 //////////////////////////////////////////////////////////////////////////
 /// Type
 //////////////////////////////////////////////////////////////////////////
@@ -34,14 +36,14 @@ enum class EErodeDilateKernelType : uint8
 //////////////////////////////////////////////////////////////////////////
 /// Type
 //////////////////////////////////////////////////////////////////////////
-UCLASS()
-class TEXTUREGRAPH_API UTG_Expression_ErodeDilate : public UTG_Expression
+UCLASS(MinimalAPI)
+class UTG_Expression_ErodeDilate : public UTG_Expression
 {
 	GENERATED_BODY()
 
 public:
 	TG_DECLARE_EXPRESSION(TG_Category::Filter);
-	virtual void						Evaluate(FTG_EvaluationContext* InContext) override;
+	UE_API virtual void						Evaluate(FTG_EvaluationContext* InContext) override;
 
 	/// The input image to 
 	UPROPERTY(meta = (TGType = "TG_Input", PinDisplayName = ""))
@@ -75,3 +77,4 @@ public:
 	virtual FText						GetTooltipText() const override { return FText::FromString(TEXT("Erode or dilate a particular image.")); } 
 };
 
+#undef UE_API
diff --git a/Engine/Plugins/TextureGraph/Source/TextureGraph/Public/Expressions/Filter/TG_Expression_Levels.h b/Engine/Plugins/TextureGraph/Source/TextureGraph/Public/Expressions/Filter/TG_Expression_Levels.h
--- a/Engine/Plugins/TextureGraph/Source/TextureGraph/Public/Expressions/Filter/TG_Expression_Levels.h
+++ b/Engine/Plugins/TextureGraph/Source/TextureGraph/Public/Expressions/Filter/TG_Expression_Levels.h
@@ -9,6 +9,8 @@
 
 #include "TG_Expression_Levels.generated.h"
 
+#define UE_API TEXTUREGRAPH_API
+
 
 
 UENUM(BlueprintType)
@@ -21,7 +23,7 @@ enum class ELevelsExpressionType : uint8
 struct FLevels;
 
 USTRUCT(BlueprintType)
-struct TEXTUREGRAPH_API FTG_LevelsSettings
+struct FTG_LevelsSettings
 {
 	GENERATED_USTRUCT_BODY()
 
@@ -38,18 +40,18 @@ struct TEXTUREGRAPH_API FTG_LevelsSettings
 	UPROPERTY(EditAnywhere, Category = "Levels", meta = (PinDisplayName = "High", UIMin = "0", ClampMin = "0", UIMax = "1", ClampMax = "1"))
 		float								High = 1;
 
-	bool SetLow(float InValue);
-	bool SetMid(float InValue);
-	bool SetHigh(float InValue);
+	UE_API bool SetLow(float InValue);
+	UE_API bool SetMid(float InValue);
+	UE_API bool SetHigh(float InValue);
 
 	// Eval Low-High range mapping on value
-	float EvalRange(float Val) const;
+	UE_API float EvalRange(float Val) const;
 	// Eval reverse Low-High range mapping on value
-	float EvalRangeInv(float Val) const;
+	UE_API float EvalRangeInv(float Val) const;
 
 	// Evaluate the MidExponent of the power curve applying the midpoint filter
-	float EvalMidExponent() const;
-	bool SetMidFromMidExponent(float InExponent);
+	UE_API float EvalMidExponent() const;
+	UE_API bool SetMidFromMidExponent(float InExponent);
 
 	void InitFromString(const FString& StrVal)
 	{
@@ -70,8 +72,8 @@ void FTG_LevelsSettings_VarPropertySerialize(FTG_Var::VarPropertySerialInfo&);
 template <> FString TG_Var_LogValue(FTG_LevelsSettings& Value);
 template <> void TG_Var_SetValueFromString(FTG_LevelsSettings& Value, const FString& StrVal);
 
-UCLASS()
-class TEXTUREGRAPH_API UTG_Expression_Levels : public UTG_Expression
+UCLASS(MinimalAPI)
+class UTG_Expression_Levels : public UTG_Expression
 {
 	GENERATED_BODY()
 
@@ -80,18 +82,18 @@ class TEXTUREGRAPH_API UTG_Expression_Levels : public UTG_Expression
 	TSharedPtr<FLevels> 				LevelsControl;
 
 public:	
-	virtual void						PostLoad() override;
+	UE_API virtual void						PostLoad() override;
 #if WITH_EDITOR
-	virtual void 						PostEditChangeProperty(FPropertyChangedEvent& PropertyChangedEvent) override;
+	UE_API virtual void 						PostEditChangeProperty(FPropertyChangedEvent& PropertyChangedEvent) override;
 
 	// Used to implement EditCondition logic for both Node UI and Details View
-	virtual bool						CanEditChange(const FProperty* InProperty) const override;
+	UE_API virtual bool						CanEditChange(const FProperty* InProperty) const override;
 #endif
 
 
 
 	TG_DECLARE_EXPRESSION(TG_Category::Adjustment);
-	virtual void						Evaluate(FTG_EvaluationContext* InContext) override;
+	UE_API virtual void						Evaluate(FTG_EvaluationContext* InContext) override;
 
 	UPROPERTY(EditAnywhere, Category = NoCategory, meta = (TGType = "TG_Setting", PinDisplayName = "Type", PinNotConnectable = true, RegenPinsOnChange))
 	ELevelsExpressionType				LevelsExpressionType = ELevelsExpressionType::LowMidHigh;
@@ -103,28 +105,28 @@ public:
 	// The Low value of the Levels adjustment, any pixel under that value is set to black. Default is 0.
 	UPROPERTY(EditAnywhere, Setter,  Category = NoCategory, meta = (TGType = "TG_Setting", PinDisplayName = "Low", UIMin = "0", ClampMin = "0", UIMax = "1", ClampMax = "1", EditConditionHides, MD_ScalarEditor))
 	float								LowValue = 0;
-	void SetLowValue(float InValue);
+	UE_API void SetLowValue(float InValue);
 
 	// The mid value of the Levels adjustment, must be in the range [Min, Max] and the Default is 0.5.
 	// The mid value determine where the smoothing curve applying the midpoint filter is crossing 0.5.
 	UPROPERTY(EditAnywhere, Setter, Category = NoCategory, meta = (TGType = "TG_Setting", PinDisplayName = "Mid", UIMin = "0", ClampMin = "0", UIMax = "1", ClampMax = "1", EditConditionHides, MD_ScalarEditor))
 	float								MidValue = 0.5f; // Midtones
-	void SetMidValue(float InValue);
+	UE_API void SetMidValue(float InValue);
 
 	// The High value of the Levels adjustment, any pixel above that value is set to white. Default is 1.
 	UPROPERTY(EditAnywhere, Setter, Category = NoCategory, meta = (TGType = "TG_Setting", PinDisplayName = "High", UIMin = "0", ClampMin = "0", UIMax = "1", ClampMax = "1", EditConditionHides, MD_ScalarEditor))
 	float								HighValue = 1;
-	void SetHighValue(float InValue);
+	UE_API void SetHighValue(float InValue);
 
 	// The black point of the output. Moving this will remap the dark point to this value. Default is 0
 	UPROPERTY(EditAnywhere, Setter, Category = NoCategory, meta = (TGType = "TG_Setting", PinDisplayName = "Out Low", UIMin = "0", ClampMin = "0", UIMax = "1", ClampMax = "1", EditConditionHides, MD_ScalarEditor))
 	float								OutLowValue = 0;
-	void SetOutLowValue(float InValue);
+	UE_API void SetOutLowValue(float InValue);
 
 	// The white point of the output. Moving this will remap the white point to this value. Default is 1
 	UPROPERTY(EditAnywhere, Setter, Category = NoCategory, meta = (TGType = "TG_Setting", PinDisplayName = "Out High", UIMin = "0", ClampMin = "0", UIMax = "1", ClampMax = "1", EditConditionHides, MD_ScalarEditor))
 	float								OutHighValue = 1;
-	void SetOutHighValue(float InValue);
+	UE_API void SetOutHighValue(float InValue);
 
 	// The High value of the Levels adjustment, any pixel above that value is set to white. Default is 1.
 	UPROPERTY(EditAnywhere, Category = NoCategory, meta = (TGType = "TG_Setting", PinDisplayName = "Mid Auto Levels", UIMin = "0", ClampMin = "0", UIMax = "1", ClampMax = "1", EditConditionHides, MD_ScalarEditor))
@@ -143,8 +145,8 @@ public:
 /////////////////////////////////////////////////////////////////
 /// Histogram scan
 /////////////////////////////////////////////////////////////////
-UCLASS()
-class TEXTUREGRAPH_API UTG_Expression_HistogramScan : public UTG_Expression
+UCLASS(MinimalAPI)
+class UTG_Expression_HistogramScan : public UTG_Expression
 {
 	GENERATED_BODY()
 
@@ -153,7 +155,7 @@ class TEXTUREGRAPH_API UTG_Expression_HistogramScan : public UTG_Expression
 public:
 
 	TG_DECLARE_EXPRESSION(TG_Category::Adjustment);
-	virtual void						Evaluate(FTG_EvaluationContext* InContext) override;
+	UE_API virtual void						Evaluate(FTG_EvaluationContext* InContext) override;
 
 	// The input image to adjust the levels for
 	UPROPERTY(EditAnywhere, Category = NoCategory, meta = (TGType = "TG_Input", PinDisplayName = "", MD_HistogramLuminance))
@@ -178,8 +180,8 @@ public:
 /////////////////////////////////////////////////////////////////
 /// Histogram range
 /////////////////////////////////////////////////////////////////
-UCLASS()
-class TEXTUREGRAPH_API UTG_Expression_HistogramRange : public UTG_Expression
+UCLASS(MinimalAPI)
+class UTG_Expression_HistogramRange : public UTG_Expression
 {
 	GENERATED_BODY()
 
@@ -188,7 +190,7 @@ class TEXTUREGRAPH_API UTG_Expression_HistogramRange : public UTG_Expression
 public:
 
 	TG_DECLARE_EXPRESSION(TG_Category::Adjustment);
-	virtual void						Evaluate(FTG_EvaluationContext* InContext) override;
+	UE_API virtual void						Evaluate(FTG_EvaluationContext* InContext) override;
 
 	// The input image to adjust the levels for
 	UPROPERTY(EditAnywhere, Category = NoCategory, meta = (TGType = "TG_Input", PinDisplayName = "", MD_HistogramLuminance))
@@ -215,8 +217,8 @@ public:
 /// Histogram select
 /////////////////////////////////////////////////////////////////
 #if 0 /// TODO
-UCLASS()
-class TEXTUREGRAPH_API UTG_Expression_HistogramSelect : public UTG_Expression
+UCLASS(MinimalAPI)
+class UTG_Expression_HistogramSelect : public UTG_Expression
 {
 	GENERATED_BODY()
 
@@ -225,7 +227,7 @@ class TEXTUREGRAPH_API UTG_Expression_HistogramSelect : public UTG_Expression
 public:
 
 	TG_DECLARE_EXPRESSION(TG_Category::Adjustment);
-	virtual void						Evaluate(FTG_EvaluationContext* InContext) override;
+	UE_API virtual void						Evaluate(FTG_EvaluationContext* InContext) override;
 
 	// The input image to adjust the levels for
 	UPROPERTY(EditAnywhere, Category = NoCategory, meta = (TGType = "TG_Input", PinDisplayName = "", MD_HistogramLuminance))
@@ -251,3 +253,5 @@ public:
 
 };
 #endif 
+
+#undef UE_API
diff --git a/Engine/Plugins/TextureGraph/Source/TextureGraph/Public/Expressions/Filter/TG_Expression_Threshold.h b/Engine/Plugins/TextureGraph/Source/TextureGraph/Public/Expressions/Filter/TG_Expression_Threshold.h
--- a/Engine/Plugins/TextureGraph/Source/TextureGraph/Public/Expressions/Filter/TG_Expression_Threshold.h
+++ b/Engine/Plugins/TextureGraph/Source/TextureGraph/Public/Expressions/Filter/TG_Expression_Threshold.h
@@ -9,14 +9,16 @@
 
 #include "TG_Expression_Threshold.generated.h"
 
-UCLASS()
-class TEXTUREGRAPH_API UTG_Expression_Threshold : public UTG_Expression
+#define UE_API TEXTUREGRAPH_API
+
+UCLASS(MinimalAPI)
+class UTG_Expression_Threshold : public UTG_Expression
 {
 	GENERATED_BODY()
 
 public:
 	TG_DECLARE_EXPRESSION(TG_Category::Adjustment);
-	virtual void						Evaluate(FTG_EvaluationContext* InContext) override;
+	UE_API virtual void						Evaluate(FTG_EvaluationContext* InContext) override;
 
 	// The input image to adjust the levels for
 	UPROPERTY(EditAnywhere, Category = NoCategory, meta = (TGType = "TG_Input", PinDisplayName = "", MD_HistogramLuminance))
@@ -33,3 +35,4 @@ public:
 	virtual FText						GetTooltipText() const override { return FText::FromString(TEXT("Apply a threshold filter to the input image, output is white if the luminance of the pixel is greater or equal of the threshold, black otherwise")); } 
 };
 
+#undef UE_API
diff --git a/Engine/Plugins/TextureGraph/Source/TextureGraph/Public/Expressions/Filter/TG_Expression_Warp.h b/Engine/Plugins/TextureGraph/Source/TextureGraph/Public/Expressions/Filter/TG_Expression_Warp.h
--- a/Engine/Plugins/TextureGraph/Source/TextureGraph/Public/Expressions/Filter/TG_Expression_Warp.h
+++ b/Engine/Plugins/TextureGraph/Source/TextureGraph/Public/Expressions/Filter/TG_Expression_Warp.h
@@ -8,11 +8,13 @@
 
 #include "TG_Expression_Warp.generated.h"
 
+#define UE_API TEXTUREGRAPH_API
+
 //////////////////////////////////////////////////////////////////////////
 /// Base warp class with common options
 //////////////////////////////////////////////////////////////////////////
-UCLASS()
-class TEXTUREGRAPH_API UTG_Expression_Warp : public UTG_Expression
+UCLASS(MinimalAPI)
+class UTG_Expression_Warp : public UTG_Expression
 {
 	GENERATED_BODY()
 
@@ -54,10 +56,11 @@ public:
 public:
 #if WITH_EDITOR
 	// Used to implement EditCondition logic for both Node UI and Details View
-	virtual bool						CanEditChange(const FProperty* InProperty) const override;
+	UE_API virtual bool						CanEditChange(const FProperty* InProperty) const override;
 #endif
 
 	virtual FText GetTooltipText() const override { return FText::FromString(TEXT("Evalute different kinds of warp depending on the Mask and Intensity settings.")); }
-	virtual void Evaluate(FTG_EvaluationContext* InContext) override;
+	UE_API virtual void Evaluate(FTG_EvaluationContext* InContext) override;
 };
 
+#undef UE_API
diff --git a/Engine/Plugins/TextureGraph/Source/TextureGraph/Public/Expressions/Input/TG_Expression_Bool.h b/Engine/Plugins/TextureGraph/Source/TextureGraph/Public/Expressions/Input/TG_Expression_Bool.h
--- a/Engine/Plugins/TextureGraph/Source/TextureGraph/Public/Expressions/Input/TG_Expression_Bool.h
+++ b/Engine/Plugins/TextureGraph/Source/TextureGraph/Public/Expressions/Input/TG_Expression_Bool.h
@@ -6,15 +6,17 @@
 #include "TG_Expression_InputParam.h"
 #include "TG_Expression_Bool.generated.h"
 
-UCLASS()
-class TEXTUREGRAPH_API UTG_Expression_Bool : public UTG_Expression_InputParam
+#define UE_API TEXTUREGRAPH_API
+
+UCLASS(MinimalAPI)
+class UTG_Expression_Bool : public UTG_Expression_InputParam
 {
 	GENERATED_BODY()
 	TG_DECLARE_INPUT_PARAM_EXPRESSION(TG_Category::Input);
 
 public:
 
-	virtual void Evaluate(FTG_EvaluationContext* InContext) override;
+	UE_API virtual void Evaluate(FTG_EvaluationContext* InContext) override;
 
 	UPROPERTY(EditAnywhere, Category = NoCategory, meta = (TGType = "TG_InputParam", PinDisplayName = ""))
     bool Bool = false;
@@ -24,7 +26,9 @@ public:
 	bool ValueOut = false;
 	
 	virtual FTG_Name GetDefaultName() const override { return TEXT("Bool"); }
-	virtual void SetTitleName(FName NewName) override;
-	virtual FName GetTitleName() const override;
+	UE_API virtual void SetTitleName(FName NewName) override;
+	UE_API virtual FName GetTitleName() const override;
 	virtual FText GetTooltipText() const override { return FText::FromString(TEXT("Makes a boolean input parameter. It is automatically exposed as a graph parameter.")); } 
 };
+
+#undef UE_API
diff --git a/Engine/Plugins/TextureGraph/Source/TextureGraph/Public/Expressions/Input/TG_Expression_Color.h b/Engine/Plugins/TextureGraph/Source/TextureGraph/Public/Expressions/Input/TG_Expression_Color.h
--- a/Engine/Plugins/TextureGraph/Source/TextureGraph/Public/Expressions/Input/TG_Expression_Color.h
+++ b/Engine/Plugins/TextureGraph/Source/TextureGraph/Public/Expressions/Input/TG_Expression_Color.h
@@ -8,15 +8,17 @@
 
 #include "TG_Expression_Color.generated.h"
 
-UCLASS()
-class TEXTUREGRAPH_API UTG_Expression_Color : public UTG_Expression_InputParam
+#define UE_API TEXTUREGRAPH_API
+
+UCLASS(MinimalAPI)
+class UTG_Expression_Color : public UTG_Expression_InputParam
 {
 	GENERATED_BODY()
 	TG_DECLARE_INPUT_PARAM_EXPRESSION(TG_Category::Input);
 
 public:
 
-	virtual void Evaluate(FTG_EvaluationContext* InContext) override;
+	UE_API virtual void Evaluate(FTG_EvaluationContext* InContext) override;
 
 	// The color value
 	UPROPERTY(EditAnywhere, Category = NoCategory, meta = (TGType = "TG_InputParam", PinDisplayName = ""))
@@ -27,10 +29,11 @@ public:
 	FLinearColor ValueOut = FLinearColor::Black;
 
 
-	virtual void SetTitleName(FName NewName) override;
-	virtual FName GetTitleName() const override;
+	UE_API virtual void SetTitleName(FName NewName) override;
+	UE_API virtual FName GetTitleName() const override;
 	virtual FTG_Name GetDefaultName() const override { return TEXT("Color"); }
 	
 	virtual FText GetTooltipText() const override { return FText::FromString(TEXT("Makes an RGBA color available. It is automatically exposed as a graph input parameter.")); } 
 };
 
+#undef UE_API
diff --git a/Engine/Plugins/TextureGraph/Source/TextureGraph/Public/Expressions/Input/TG_Expression_Graph.h b/Engine/Plugins/TextureGraph/Source/TextureGraph/Public/Expressions/Input/TG_Expression_Graph.h
--- a/Engine/Plugins/TextureGraph/Source/TextureGraph/Public/Expressions/Input/TG_Expression_Graph.h
+++ b/Engine/Plugins/TextureGraph/Source/TextureGraph/Public/Expressions/Input/TG_Expression_Graph.h
@@ -8,19 +8,21 @@
 
 #include "TG_Expression_Graph.generated.h"
 
+#define UE_API TEXTUREGRAPH_API
+
 class UTextureGraph;
 
 //////////////////////////////////////////////////////////////////////////
 //// Generic subgraph
 //////////////////////////////////////////////////////////////////////////
 
-UCLASS(Hidden)
-class TEXTUREGRAPH_API UTG_Expression_Graph : public UTG_Expression
+UCLASS(MinimalAPI, Hidden)
+class UTG_Expression_Graph : public UTG_Expression
 {
 	GENERATED_BODY()
 public:
 	TG_DECLARE_DYNAMIC_EXPRESSION(TG_Category::Input);
-	virtual void Evaluate(FTG_EvaluationContext* InContext) override;
+	UE_API virtual void Evaluate(FTG_EvaluationContext* InContext) override;
 
 protected:
 	
@@ -28,23 +30,23 @@ protected:
 	mutable TArray<FTG_Id> OutParamIds;
 
 	virtual UTG_Graph* GetGraph() const { return nullptr; }
-	void NotifyGraphChanged();
+	UE_API void NotifyGraphChanged();
 
-	virtual void SetupAndEvaluate(FTG_EvaluationContext* InContext) override;
+	UE_API virtual void SetupAndEvaluate(FTG_EvaluationContext* InContext) override;
 };
 
 //////////////////////////////////////////////////////////////////////////
 // TextureGraph Expression
 //////////////////////////////////////////////////////////////////////////
 
-UCLASS()
-class TEXTUREGRAPH_API UTG_Expression_TextureGraph : public UTG_Expression_Graph
+UCLASS(MinimalAPI)
+class UTG_Expression_TextureGraph : public UTG_Expression_Graph
 {
 	GENERATED_BODY()
 public:
 
-	UTG_Expression_TextureGraph();
-	virtual ~UTG_Expression_TextureGraph();
+	UE_API UTG_Expression_TextureGraph();
+	UE_API virtual ~UTG_Expression_TextureGraph();
 	
 
 #if WITH_EDITOR
@@ -52,40 +54,42 @@ protected:
 	// Only in editor mode, in the case the TextureGraph is assigned from Details panel we need to remember the previous version
 	TWeakObjectPtr<UTextureGraph>	PreEditTextureGraph; 
 public:
-	virtual void PreEditChange(FProperty* PropertyAboutToChange) override;
-	virtual void PostEditChangeProperty(FPropertyChangedEvent& PropertyChangedEvent) override;
-	virtual void PostEditUndo() override;
+	UE_API virtual void PreEditChange(FProperty* PropertyAboutToChange) override;
+	UE_API virtual void PostEditChangeProperty(FPropertyChangedEvent& PropertyChangedEvent) override;
+	UE_API virtual void PostEditUndo() override;
 #endif
 
-	virtual void Evaluate(FTG_EvaluationContext* InContext) override;
-	virtual bool Validate(MixUpdateCyclePtr	Cycle) override;
+	UE_API virtual void Evaluate(FTG_EvaluationContext* InContext) override;
+	UE_API virtual bool Validate(MixUpdateCyclePtr	Cycle) override;
 
 	UPROPERTY(EditAnywhere, Setter, Category = NoCategory, meta = (TGType = "TG_Setting"))
 	TObjectPtr<UTextureGraph>	TextureGraph;				// The TextureGraph to use for the sub-graph
 
-	void SetTextureGraph(UTextureGraph* InTextureGraph);
+	UE_API void SetTextureGraph(UTextureGraph* InTextureGraph);
 
-	virtual bool CanHandleAsset(UObject* Asset) override;
-	virtual void SetAsset(UObject* Asset) override;
+	UE_API virtual bool CanHandleAsset(UObject* Asset) override;
+	UE_API virtual void SetAsset(UObject* Asset) override;
 	
 protected:
-	virtual void Initialize() override;
-	virtual void CopyVarGeneric(const FTG_Argument& Arg, FTG_Var* InVar, bool CopyVarToArg) override;
+	UE_API virtual void Initialize() override;
+	UE_API virtual void CopyVarGeneric(const FTG_Argument& Arg, FTG_Var* InVar, bool CopyVarToArg) override;
 
-	virtual UTG_Graph* GetGraph() const override;
+	UE_API virtual UTG_Graph* GetGraph() const override;
 
-	void OnTextureGraphPreSave(UObject* Object, FObjectPreSaveContext SaveContext);
+	UE_API void OnTextureGraphPreSave(UObject* Object, FObjectPreSaveContext SaveContext);
 private:
 	UPROPERTY(Transient)
 	TObjectPtr<UTG_Graph>	RuntimeGraph;		// The runtime graph instance to use for the sub-graph
 
 	FDelegateHandle PreSaveHandle;				// delegate handle for global ObjectPreSave event 
 	
-	bool CheckDependencies(const UTextureGraph* InTextureGraph) const;
-	void SetTextureGraphInternal(UTextureGraph* InTextureGraph);
+	UE_API bool CheckDependencies(const UTextureGraph* InTextureGraph) const;
+	UE_API void SetTextureGraphInternal(UTextureGraph* InTextureGraph);
 
 public:
 
 	virtual FText GetTooltipText() const override { return FText::FromString(TEXT("Encapsulates another graph as a node with the graph's input and output parameters exposed as the node input and output pins.")); }
 
-};
\ No newline at end of file
+};
+
+#undef UE_API
diff --git a/Engine/Plugins/TextureGraph/Source/TextureGraph/Public/Expressions/Input/TG_Expression_InputParam.h b/Engine/Plugins/TextureGraph/Source/TextureGraph/Public/Expressions/Input/TG_Expression_InputParam.h
--- a/Engine/Plugins/TextureGraph/Source/TextureGraph/Public/Expressions/Input/TG_Expression_InputParam.h
+++ b/Engine/Plugins/TextureGraph/Source/TextureGraph/Public/Expressions/Input/TG_Expression_InputParam.h
@@ -8,27 +8,29 @@
 
 #include "TG_Expression_InputParam.generated.h"
 
-UCLASS(Abstract)
-class TEXTUREGRAPH_API UTG_Expression_InputParam : public UTG_Expression
+#define UE_API TEXTUREGRAPH_API
+
+UCLASS(MinimalAPI, Abstract)
+class UTG_Expression_InputParam : public UTG_Expression
 {
 	GENERATED_BODY()
 
 protected:
 	// these 2 methods need to be implemented to provide concrete signatures for derived InputParam expression classes
-	virtual FTG_SignaturePtr BuildInputParameterSignature() const;
-	virtual FTG_SignaturePtr BuildInputConstantSignature() const;
+	UE_API virtual FTG_SignaturePtr BuildInputParameterSignature() const;
+	UE_API virtual FTG_SignaturePtr BuildInputConstantSignature() const;
 	virtual bool ShouldShowSettings() const override { return false; }
 public:
 
 #if WITH_EDITOR
-	void PostEditChangeProperty(FPropertyChangedEvent& PropertyChangedEvent) override;
+	UE_API void PostEditChangeProperty(FPropertyChangedEvent& PropertyChangedEvent) override;
 #endif
 
 	// Node can be converted from Param to Constant
 	UPROPERTY(Setter)
 	bool bIsConstant = false;
-	void SetbIsConstant(bool InIsConstant);
-	void ToggleIsConstant();
+	UE_API void SetbIsConstant(bool InIsConstant);
+	UE_API void ToggleIsConstant();
 
 };
 
@@ -40,3 +42,5 @@ public:
 			static FTG_SignaturePtr ParameterSignature = BuildInputParameterSignature(); return ParameterSignature; \
 		} } \
 	public: virtual FName GetCategory() const override { return Category; }
+
+#undef UE_API
diff --git a/Engine/Plugins/TextureGraph/Source/TextureGraph/Public/Expressions/Input/TG_Expression_Material.h b/Engine/Plugins/TextureGraph/Source/TextureGraph/Public/Expressions/Input/TG_Expression_Material.h
--- a/Engine/Plugins/TextureGraph/Source/TextureGraph/Public/Expressions/Input/TG_Expression_Material.h
+++ b/Engine/Plugins/TextureGraph/Source/TextureGraph/Public/Expressions/Input/TG_Expression_Material.h
@@ -11,22 +11,24 @@
 
 #include "TG_Expression_Material.generated.h"
 
+#define UE_API TEXTUREGRAPH_API
+
 typedef std::weak_ptr<class Job>		JobPtrW;
 
-UCLASS()
-class TEXTUREGRAPH_API UTG_Expression_Material : public UTG_Expression_MaterialBase
+UCLASS(MinimalAPI)
+class UTG_Expression_Material : public UTG_Expression_MaterialBase
 {
 	GENERATED_BODY()
 
 public:
-	UTG_Expression_Material();
-	virtual ~UTG_Expression_Material();
+	UE_API UTG_Expression_Material();
+	UE_API virtual ~UTG_Expression_Material();
 #if WITH_EDITOR
-	virtual void PostEditChangeProperty(FPropertyChangedEvent& PropertyChangedEvent) override;
-	virtual void PostEditUndo() override;
+	UE_API virtual void PostEditChangeProperty(FPropertyChangedEvent& PropertyChangedEvent) override;
+	UE_API virtual void PostEditUndo() override;
 #endif
 
-	virtual void PostLoad() override;
+	UE_API virtual void PostLoad() override;
 
 #if WITH_EDITORONLY_DATA
 	UE_DEPRECATED(5.6, "Material field has been deprecated.")
@@ -37,19 +39,19 @@ public:
 	// The input material referenced by this Material node
 	UPROPERTY(EditAnywhere, Setter, Category = NoCategory, meta = (TGType = "TG_InputParam", DisplayName = "Material", PinDisplayName = "Material"))
 	FTG_Material InputMaterial;
-	void SetInputMaterial(const FTG_Material& InMaterial);
+	UE_API void SetInputMaterial(const FTG_Material& InMaterial);
 
 	// The Material attribute identifier among all the attributes of the material that is rendered in the output
 	UPROPERTY(EditAnywhere, BlueprintReadWrite, Setter, Category = NoCategory, meta = (TGType = "TG_Setting", GetOptions = "GetRenderAttributeOptions"))
 	FName RenderedAttribute;
-	void SetRenderedAttribute(FName InRenderedAttribute);
+	UE_API void SetRenderedAttribute(FName InRenderedAttribute);
 
 	// THe list of Rendered attribute options available 
 	UFUNCTION(CallInEditor)
-	TArray<FName> GetRenderAttributeOptions() const;
+	UE_API TArray<FName> GetRenderAttributeOptions() const;
 
-	virtual bool CanHandleAsset(UObject* Asset) override;
-	virtual void SetAsset(UObject* Asset) override;
+	UE_API virtual bool CanHandleAsset(UObject* Asset) override;
+	UE_API virtual void SetAsset(UObject* Asset) override;
 	virtual FText GetTooltipText() const override { return FText::FromString(TEXT("Renders a material into a quad and makes it available. It is automatically exposed as a graph input parameter.")); } 
 
 protected:
@@ -59,30 +61,31 @@ protected:
 
 	// As the "Material" or the "InputMaterial" properties change
 	// we update the current active "MaterialCopy" and pass on internally as the active "GetMaterial()"
-	virtual void SetMaterialInternal(UMaterialInterface* InMaterial) override;
+	UE_API virtual void SetMaterialInternal(UMaterialInterface* InMaterial) override;
 
 	// Title name is still used mostly for legacy but not exposed anymore in the details.
 	// This is changed on the node itself and then call SetTitleName and rename the InnputMaterial Alias name
 	UPROPERTY()
 	FName TitleName = TEXT("Material");
 
-	virtual void Initialize() override;
+	UE_API virtual void Initialize() override;
 
 public:
 	
-	virtual void SetTitleName(FName NewName) override;
-	virtual FName GetTitleName() const override;
+	UE_API virtual void SetTitleName(FName NewName) override;
+	UE_API virtual FName GetTitleName() const override;
 
 	virtual FName GetCategory() const override { return TG_Category::Input;}
 	
 protected:
 	virtual TObjectPtr<UMaterialInterface> GetMaterial() const override { return MaterialCopy;};
-	virtual EDrawMaterialAttributeTarget GetRenderedAttributeId()  override;
+	UE_API virtual EDrawMaterialAttributeTarget GetRenderedAttributeId()  override;
 
 #if WITH_EDITOR // Listener for referenced material being saved to update the integration in TG 
 	FDelegateHandle PreSaveHandle; 
-	void OnReferencedObjectPreSave(UObject* Object, FObjectPreSaveContext SaveContext);
+	UE_API void OnReferencedObjectPreSave(UObject* Object, FObjectPreSaveContext SaveContext);
 #endif
 
 };
 
+#undef UE_API
diff --git a/Engine/Plugins/TextureGraph/Source/TextureGraph/Public/Expressions/Input/TG_Expression_MaterialFunction.h b/Engine/Plugins/TextureGraph/Source/TextureGraph/Public/Expressions/Input/TG_Expression_MaterialFunction.h
--- a/Engine/Plugins/TextureGraph/Source/TextureGraph/Public/Expressions/Input/TG_Expression_MaterialFunction.h
+++ b/Engine/Plugins/TextureGraph/Source/TextureGraph/Public/Expressions/Input/TG_Expression_MaterialFunction.h
@@ -10,44 +10,47 @@
 
 #include "TG_Expression_MaterialFunction.generated.h"
 
+#define UE_API TEXTUREGRAPH_API
 
-UCLASS()
-class TEXTUREGRAPH_API UTG_Expression_MaterialFunction : public UTG_Expression_MaterialBase
+
+UCLASS(MinimalAPI)
+class UTG_Expression_MaterialFunction : public UTG_Expression_MaterialBase
 {
 	GENERATED_BODY()
 public:
 #if WITH_EDITOR
-	virtual void PostEditChangeProperty(FPropertyChangedEvent& PropertyChangedEvent) override;
-	virtual void PostEditUndo() override;
+	UE_API virtual void PostEditChangeProperty(FPropertyChangedEvent& PropertyChangedEvent) override;
+	UE_API virtual void PostEditUndo() override;
 #endif
 	
 	// The material function to employ for rendering
 	UPROPERTY(EditAnywhere, BlueprintReadWrite, Setter, Category = NoCategory, meta = (TGType = "TG_Setting", TGPinNotConnectable))
 	TObjectPtr<UMaterialFunctionInterface> MaterialFunction;
-	void SetMaterialFunction(UMaterialFunctionInterface* InMaterialFunction);
+	UE_API void SetMaterialFunction(UMaterialFunctionInterface* InMaterialFunction);
 
-	virtual bool CanHandleAsset(UObject* Asset) override;
-	virtual void SetAsset(UObject* Asset) override;
+	UE_API virtual bool CanHandleAsset(UObject* Asset) override;
+	UE_API virtual void SetAsset(UObject* Asset) override;
 
 	virtual FText GetTooltipText() const override { return FText::FromString(TEXT("Renders a material function output into a quad and makes it available for the texture graph.")); } 
 
 protected:
-	virtual void Initialize() override;
+	UE_API virtual void Initialize() override;
 
-	UMaterialInterface* CreateMaterialReference();
+	UE_API UMaterialInterface* CreateMaterialReference();
 	TObjectPtr<UMaterialInterface> ReferenceMaterial = nullptr;
 	virtual TObjectPtr<UMaterialInterface> GetMaterial() const { return ReferenceMaterial; }
 	virtual EDrawMaterialAttributeTarget GetRenderedAttributeId() { return EDrawMaterialAttributeTarget::Emissive; }
 
 private:
-	void SetMaterialFunctionInternal(UMaterialFunctionInterface* InMaterialFunction);
+	UE_API void SetMaterialFunctionInternal(UMaterialFunctionInterface* InMaterialFunction);
 
 public:
 	
 	virtual FTG_Name GetDefaultName() const override { return TEXT("MaterialFunction");}
-	virtual bool CanRenameTitle() const override;
+	UE_API virtual bool CanRenameTitle() const override;
 
 	virtual FName GetCategory() const override { return TG_Category::Input; }
 
 };
 
+#undef UE_API
diff --git a/Engine/Plugins/TextureGraph/Source/TextureGraph/Public/Expressions/Input/TG_Expression_OutputSettings.h b/Engine/Plugins/TextureGraph/Source/TextureGraph/Public/Expressions/Input/TG_Expression_OutputSettings.h
--- a/Engine/Plugins/TextureGraph/Source/TextureGraph/Public/Expressions/Input/TG_Expression_OutputSettings.h
+++ b/Engine/Plugins/TextureGraph/Source/TextureGraph/Public/Expressions/Input/TG_Expression_OutputSettings.h
@@ -9,16 +9,18 @@
 
 #include "TG_Expression_OutputSettings.generated.h"
 
-UCLASS()
-class TEXTUREGRAPH_API UTG_Expression_OutputSettings : public UTG_Expression
+#define UE_API TEXTUREGRAPH_API
+
+UCLASS(MinimalAPI)
+class UTG_Expression_OutputSettings : public UTG_Expression
 {
 	GENERATED_BODY()
 
 public:
 	TG_DECLARE_EXPRESSION(TG_Category::Input);
-	virtual void Evaluate(FTG_EvaluationContext* InContext) override;
+	UE_API virtual void Evaluate(FTG_EvaluationContext* InContext) override;
 
-	virtual bool Validate(MixUpdateCyclePtr	Cycle) override;
+	UE_API virtual bool Validate(MixUpdateCyclePtr	Cycle) override;
 	
 	// The output of the node
 	UPROPERTY(meta = (TGType = "TG_Output", PinDisplayName = ""))
@@ -34,11 +36,12 @@ public:
 
 	class ULayerChannel* Channel;
 	virtual FTG_Name GetDefaultName() const override { return TEXT("Settings"); }
-	virtual void SetTitleName(FName NewName) override;
-	virtual FName GetTitleName() const override;
+	UE_API virtual void SetTitleName(FName NewName) override;
+	UE_API virtual FName GetTitleName() const override;
 	virtual FText GetTooltipText() const override { return FText::FromString(TEXT("Makes the settings for the output.")); }
 
 private:
 	FTG_OutputSettings PreviousInput;
 };
 
+#undef UE_API
diff --git a/Engine/Plugins/TextureGraph/Source/TextureGraph/Public/Expressions/Input/TG_Expression_Scalar.h b/Engine/Plugins/TextureGraph/Source/TextureGraph/Public/Expressions/Input/TG_Expression_Scalar.h
--- a/Engine/Plugins/TextureGraph/Source/TextureGraph/Public/Expressions/Input/TG_Expression_Scalar.h
+++ b/Engine/Plugins/TextureGraph/Source/TextureGraph/Public/Expressions/Input/TG_Expression_Scalar.h
@@ -8,15 +8,17 @@
 
 #include "TG_Expression_Scalar.generated.h"
 
-UCLASS()
-class TEXTUREGRAPH_API UTG_Expression_Scalar : public UTG_Expression_InputParam
+#define UE_API TEXTUREGRAPH_API
+
+UCLASS(MinimalAPI)
+class UTG_Expression_Scalar : public UTG_Expression_InputParam
 {
 	GENERATED_BODY()
 	TG_DECLARE_INPUT_PARAM_EXPRESSION(TG_Category::Input);
 
 public:
 
-	virtual void Evaluate(FTG_EvaluationContext* InContext) override;
+	UE_API virtual void Evaluate(FTG_EvaluationContext* InContext) override;
 
 	// The floating point constant
 	UPROPERTY(EditAnywhere, Category = NoCategory, meta = (TGType = "TG_InputParam", PinDisplayName = ""))
@@ -27,8 +29,9 @@ public:
 		float ValueOut = 1.0f;
 
 	virtual FTG_Name GetDefaultName() const override { return TEXT("Scalar"); }
-	virtual void SetTitleName(FName NewName) override;
-	virtual FName GetTitleName() const override;
+	UE_API virtual void SetTitleName(FName NewName) override;
+	UE_API virtual FName GetTitleName() const override;
 	virtual FText GetTooltipText() const override { return FText::FromString(TEXT("Makes a single floating point value available. It is automatically exposed as a graph input parameter.")); } 
 };
 
+#undef UE_API
diff --git a/Engine/Plugins/TextureGraph/Source/TextureGraph/Public/Expressions/Input/TG_Expression_String.h b/Engine/Plugins/TextureGraph/Source/TextureGraph/Public/Expressions/Input/TG_Expression_String.h
--- a/Engine/Plugins/TextureGraph/Source/TextureGraph/Public/Expressions/Input/TG_Expression_String.h
+++ b/Engine/Plugins/TextureGraph/Source/TextureGraph/Public/Expressions/Input/TG_Expression_String.h
@@ -6,15 +6,17 @@
 #include "TG_Expression_InputParam.h"
 #include "TG_Expression_String.generated.h"
 
-UCLASS()
-class TEXTUREGRAPH_API UTG_Expression_String : public UTG_Expression_InputParam
+#define UE_API TEXTUREGRAPH_API
+
+UCLASS(MinimalAPI)
+class UTG_Expression_String : public UTG_Expression_InputParam
 {
 	GENERATED_BODY()
 	TG_DECLARE_INPUT_PARAM_EXPRESSION(TG_Category::Input);
 
 public:
 
-	virtual void Evaluate(FTG_EvaluationContext* InContext) override;
+	UE_API virtual void Evaluate(FTG_EvaluationContext* InContext) override;
 
 	// The string constant
 	UPROPERTY(EditAnywhere, Category = NoCategory, meta = (TGType = "TG_InputParam", PinDisplayName = ""))
@@ -25,8 +27,9 @@ public:
 		FString ValueOut;
 
 	virtual FTG_Name GetDefaultName() const override { return TEXT("String"); }
-	virtual void SetTitleName(FName NewName) override;
-	virtual FName GetTitleName() const override;
+	UE_API virtual void SetTitleName(FName NewName) override;
+	UE_API virtual FName GetTitleName() const override;
 	virtual FText GetTooltipText() const override { return FText::FromString(TEXT("String parameter that can be used as input in various places.")); } 
 };
 
+#undef UE_API
diff --git a/Engine/Plugins/TextureGraph/Source/TextureGraph/Public/Expressions/Input/TG_Expression_Texture.h b/Engine/Plugins/TextureGraph/Source/TextureGraph/Public/Expressions/Input/TG_Expression_Texture.h
--- a/Engine/Plugins/TextureGraph/Source/TextureGraph/Public/Expressions/Input/TG_Expression_Texture.h
+++ b/Engine/Plugins/TextureGraph/Source/TextureGraph/Public/Expressions/Input/TG_Expression_Texture.h
@@ -10,18 +10,20 @@
 
 #include "TG_Expression_Texture.generated.h"
 
-UCLASS()
-class TEXTUREGRAPH_API UTG_Expression_Texture : public UTG_Expression_InputParam
+#define UE_API TEXTUREGRAPH_API
+
+UCLASS(MinimalAPI)
+class UTG_Expression_Texture : public UTG_Expression_InputParam
 {
 	GENERATED_BODY()
 	TG_DECLARE_INPUT_PARAM_EXPRESSION(TG_Category::Input);
 
 public:
 
-	virtual void Evaluate(FTG_EvaluationContext* InContext) override;
+	UE_API virtual void Evaluate(FTG_EvaluationContext* InContext) override;
 
-	virtual bool Validate(MixUpdateCyclePtr	Cycle) override;
-	void SetSourceInternal();
+	UE_API virtual bool Validate(MixUpdateCyclePtr	Cycle) override;
+	UE_API void SetSourceInternal();
 
 	// The output of the node, which is the loaded texture asset
 	UPROPERTY(meta = (TGType = "TG_Output", PinDisplayName = "", HideInnerPropertiesInNode))
@@ -36,21 +38,22 @@ public:
 	UPROPERTY(EditAnywhere, Setter, Category = NoCategory, meta = (TGType = "TG_InputParam"))
 	FTG_Texture Texture;
 
-	void SetSource(UTexture* InSource);
-	void SetTextureInternal();
-	void SetTexture(const FTG_Texture& InTexture);
+	UE_API void SetSource(UTexture* InSource);
+	UE_API void SetTextureInternal();
+	UE_API void SetTexture(const FTG_Texture& InTexture);
 	
 #if WITH_EDITOR
-	void PostEditChangeProperty(FPropertyChangedEvent& PropertyChangedEvent) override;
+	UE_API void PostEditChangeProperty(FPropertyChangedEvent& PropertyChangedEvent) override;
 #endif
 	class ULayerChannel* Channel;
 	virtual FTG_Name GetDefaultName() const override { return TEXT("Texture"); }
-	virtual void SetTitleName(FName NewName) override;
-	virtual FName GetTitleName() const override;
+	UE_API virtual void SetTitleName(FName NewName) override;
+	UE_API virtual FName GetTitleName() const override;
 	virtual FText GetTooltipText() const override { return FText::FromString(TEXT("Makes an existing texture asset available. It is automatically exposed as a graph input parameter.")); }
 	virtual bool IgnoreInputTextureOnUndo() const override { return false; }
 
-	virtual void SetAsset(UObject* Asset) override;
-	virtual bool CanHandleAsset(UObject* Asset) override;
+	UE_API virtual void SetAsset(UObject* Asset) override;
+	UE_API virtual bool CanHandleAsset(UObject* Asset) override;
 };
 
+#undef UE_API
diff --git a/Engine/Plugins/TextureGraph/Source/TextureGraph/Public/Expressions/Input/TG_Expression_TextureDescriptor.h b/Engine/Plugins/TextureGraph/Source/TextureGraph/Public/Expressions/Input/TG_Expression_TextureDescriptor.h
--- a/Engine/Plugins/TextureGraph/Source/TextureGraph/Public/Expressions/Input/TG_Expression_TextureDescriptor.h
+++ b/Engine/Plugins/TextureGraph/Source/TextureGraph/Public/Expressions/Input/TG_Expression_TextureDescriptor.h
@@ -6,15 +6,17 @@
 #include "TG_Expression_InputParam.h"
 #include "TG_Expression_TextureDescriptor.generated.h"
 
-UCLASS()
-class TEXTUREGRAPH_API UTG_Expression_TextureDescriptor : public UTG_Expression_InputParam
+#define UE_API TEXTUREGRAPH_API
+
+UCLASS(MinimalAPI)
+class UTG_Expression_TextureDescriptor : public UTG_Expression_InputParam
 {
 	GENERATED_BODY()
 	TG_DECLARE_INPUT_PARAM_EXPRESSION(TG_Category::Input);
 
 public:
 
-	virtual void Evaluate(FTG_EvaluationContext* InContext) override;
+	UE_API virtual void Evaluate(FTG_EvaluationContext* InContext) override;
 
 	/// The width of the texture. 0 and negative values mean "Auto"
 	UPROPERTY(EditAnywhere, Category = NoCategory, meta = (TGType = "TG_Input"))
@@ -40,3 +42,4 @@ public:
 	virtual FText GetTooltipText() const override { return FText::FromString(TEXT("Allows the user to customize texture settings like resolution, format etc.")); } 
 };
 
+#undef UE_API
diff --git a/Engine/Plugins/TextureGraph/Source/TextureGraph/Public/Expressions/Input/TG_Expression_TexturePath.h b/Engine/Plugins/TextureGraph/Source/TextureGraph/Public/Expressions/Input/TG_Expression_TexturePath.h
--- a/Engine/Plugins/TextureGraph/Source/TextureGraph/Public/Expressions/Input/TG_Expression_TexturePath.h
+++ b/Engine/Plugins/TextureGraph/Source/TextureGraph/Public/Expressions/Input/TG_Expression_TexturePath.h
@@ -10,8 +10,10 @@
 
 #include "TG_Expression_TexturePath.generated.h"
 
-UCLASS()
-class TEXTUREGRAPH_API UTG_Expression_TexturePath : public UTG_Expression_InputParam
+#define UE_API TEXTUREGRAPH_API
+
+UCLASS(MinimalAPI)
+class UTG_Expression_TexturePath : public UTG_Expression_InputParam
 {
 	GENERATED_BODY()
 	TG_DECLARE_INPUT_PARAM_EXPRESSION(TG_Category::Input);
@@ -19,18 +21,18 @@ class TEXTUREGRAPH_API UTG_Expression_TexturePath : public UTG_Expression_InputP
 protected:
 	// Special case for TexturePath Constant signature, we want to keep the Path Input connectable in that case
 	// so do this in the override version of BuildInputConstantSignature()
-	virtual FTG_SignaturePtr BuildInputConstantSignature() const override;
+	UE_API virtual FTG_SignaturePtr BuildInputConstantSignature() const override;
 
 	// Validate the input path, returning the actual path to use
 	// empty if the input path is NOT valid
-	bool ValidateInputPath(FString& ValidatedPath) const;
+	UE_API bool ValidateInputPath(FString& ValidatedPath) const;
 
-	void ReportError(MixUpdateCyclePtr Cycle);
-	static TiledBlobPtr LoadStaticImage(FTG_EvaluationContext* InContext, const FString& Filename, BufferDescriptor* DesiredDesc);
+	UE_API void ReportError(MixUpdateCyclePtr Cycle);
+	static UE_API TiledBlobPtr LoadStaticImage(FTG_EvaluationContext* InContext, const FString& Filename, BufferDescriptor* DesiredDesc);
 public:
 
-	virtual void Evaluate(FTG_EvaluationContext* InContext) override;
-	virtual bool Validate(MixUpdateCyclePtr	Cycle) override;
+	UE_API virtual void Evaluate(FTG_EvaluationContext* InContext) override;
+	UE_API virtual bool Validate(MixUpdateCyclePtr	Cycle) override;
 	
 	// All the output textures from this node. Note that it outputs an array of textures.
 	UPROPERTY(meta = (TGType = "TG_Output", PinDisplayName = "", HideInnerPropertiesInNode))
@@ -43,8 +45,9 @@ public:
 
 	class ULayerChannel* Channel;
 	virtual FTG_Name GetDefaultName() const override { return TEXT("TexturePath"); }
-	virtual void SetTitleName(FName NewName) override;
-	virtual FName GetTitleName() const override;
+	UE_API virtual void SetTitleName(FName NewName) override;
+	UE_API virtual FName GetTitleName() const override;
 	virtual FText GetTooltipText() const override { return FText::FromString(TEXT("Loads texture(s) from a path. If the path is a single file name then one texture is loaded. If it's a directory then all textures are loaded from that directory.")); }
 };
 
+#undef UE_API
diff --git a/Engine/Plugins/TextureGraph/Source/TextureGraph/Public/Expressions/Input/TG_Expression_Vector.h b/Engine/Plugins/TextureGraph/Source/TextureGraph/Public/Expressions/Input/TG_Expression_Vector.h
--- a/Engine/Plugins/TextureGraph/Source/TextureGraph/Public/Expressions/Input/TG_Expression_Vector.h
+++ b/Engine/Plugins/TextureGraph/Source/TextureGraph/Public/Expressions/Input/TG_Expression_Vector.h
@@ -8,15 +8,17 @@
 
 #include "TG_Expression_Vector.generated.h"
 
-UCLASS()
-class TEXTUREGRAPH_API UTG_Expression_Vector : public UTG_Expression_InputParam
+#define UE_API TEXTUREGRAPH_API
+
+UCLASS(MinimalAPI)
+class UTG_Expression_Vector : public UTG_Expression_InputParam
 {
 	GENERATED_BODY()
 	TG_DECLARE_INPUT_PARAM_EXPRESSION(TG_Category::Input);
 
 public:
 
-	virtual void Evaluate(FTG_EvaluationContext* InContext) override;
+	UE_API virtual void Evaluate(FTG_EvaluationContext* InContext) override;
 
 	// The 4 dimensional vector X, Y, Z and W.
 	UPROPERTY(EditAnywhere, Category = NoCategory, meta = (TGType = "TG_InputParam", PinDisplayName = ""))
@@ -27,8 +29,9 @@ public:
 	FVector4f ValueOut = FVector4f(0);
 	
 	virtual FTG_Name GetDefaultName() const override { return TEXT("Vector"); }
-	virtual void SetTitleName(FName NewName) override;
-	virtual FName GetTitleName() const override;
+	UE_API virtual void SetTitleName(FName NewName) override;
+	UE_API virtual FName GetTitleName() const override;
 	virtual FText GetTooltipText() const override { return FText::FromString(TEXT("Makes an XYZ vector available. It is automatically exposed as a graph input parameter.")); } 
 };
 
+#undef UE_API
diff --git a/Engine/Plugins/TextureGraph/Source/TextureGraph/Public/Expressions/Maths/TG_Expression_Blend.h b/Engine/Plugins/TextureGraph/Source/TextureGraph/Public/Expressions/Maths/TG_Expression_Blend.h
--- a/Engine/Plugins/TextureGraph/Source/TextureGraph/Public/Expressions/Maths/TG_Expression_Blend.h
+++ b/Engine/Plugins/TextureGraph/Source/TextureGraph/Public/Expressions/Maths/TG_Expression_Blend.h
@@ -10,13 +10,15 @@
 #include "2D/BlendModes.h"
 #include "TG_Expression_Blend.generated.h"
 
-UCLASS()
-class TEXTUREGRAPH_API UTG_Expression_Blend : public UTG_Expression
+#define UE_API TEXTUREGRAPH_API
+
+UCLASS(MinimalAPI)
+class UTG_Expression_Blend : public UTG_Expression
 {
 	GENERATED_BODY()
 public:
 	TG_DECLARE_EXPRESSION(TG_Category::Maths)
-	virtual void Evaluate(FTG_EvaluationContext* InContext) override;
+	UE_API virtual void Evaluate(FTG_EvaluationContext* InContext) override;
 
 	// Blend mode determine how inputs texture mix together.
 	UPROPERTY(EditAnywhere, BlueprintReadWrite, Category = NoCategory, meta = (TGType = "TG_Setting"))
@@ -52,3 +54,4 @@ public:
 	virtual FText GetTooltipText() const override { return FText::FromString(TEXT("Combines two inputs based on a chosen blend type. The second input is multiplied with the Opacity, and if available the Mask as well, before applying the blend.")); } 
 };
 
+#undef UE_API
diff --git a/Engine/Plugins/TextureGraph/Source/TextureGraph/Public/Expressions/Maths/TG_Expression_Clamp.h b/Engine/Plugins/TextureGraph/Source/TextureGraph/Public/Expressions/Maths/TG_Expression_Clamp.h
--- a/Engine/Plugins/TextureGraph/Source/TextureGraph/Public/Expressions/Maths/TG_Expression_Clamp.h
+++ b/Engine/Plugins/TextureGraph/Source/TextureGraph/Public/Expressions/Maths/TG_Expression_Clamp.h
@@ -9,11 +9,13 @@
 
 #include "TG_Expression_Clamp.generated.h"
 
+#define UE_API TEXTUREGRAPH_API
+
 //////////////////////////////////////////////////////////////////////////
 /// Clamp
 //////////////////////////////////////////////////////////////////////////
-UCLASS()
-class TEXTUREGRAPH_API UTG_Expression_Clamp : public UTG_Expression_Variant
+UCLASS(MinimalAPI)
+class UTG_Expression_Clamp : public UTG_Expression_Variant
 {
 	GENERATED_BODY()
 
@@ -35,16 +37,16 @@ public:
 
 protected:
 	virtual std::vector<FTG_Variant>	GetEvaluateArgs() override { return std::vector<FTG_Variant>({ Input, MinValue, MaxValue }); }
-	virtual float						EvaluateScalar_WithValue(FTG_EvaluationContext* InContext, const float* const Value, size_t Count) override;
-	virtual FTG_Texture					EvaluateTexture(FTG_EvaluationContext* InContext) override;
+	UE_API virtual float						EvaluateScalar_WithValue(FTG_EvaluationContext* InContext, const float* const Value, size_t Count) override;
+	UE_API virtual FTG_Texture					EvaluateTexture(FTG_EvaluationContext* InContext) override;
 	virtual bool						ErrorCheckInputTextures() const { return false; }
 };
 
 //////////////////////////////////////////////////////////////////////////
 /// Smooth step
 //////////////////////////////////////////////////////////////////////////
-UCLASS()
-class TEXTUREGRAPH_API UTG_Expression_SmoothStep : public UTG_Expression_Variant
+UCLASS(MinimalAPI)
+class UTG_Expression_SmoothStep : public UTG_Expression_Variant
 {
 	GENERATED_BODY()
 
@@ -65,16 +67,16 @@ public:
 
 protected:
 	virtual std::vector<FTG_Variant>	GetEvaluateArgs() override { return std::vector<FTG_Variant>({ Input, Min, Max}); }
-	virtual float						EvaluateScalar_WithValue(FTG_EvaluationContext* InContext, const float* const Value, size_t Count) override;
-	virtual FTG_Texture					EvaluateTexture(FTG_EvaluationContext* InContext) override;
+	UE_API virtual float						EvaluateScalar_WithValue(FTG_EvaluationContext* InContext, const float* const Value, size_t Count) override;
+	UE_API virtual FTG_Texture					EvaluateTexture(FTG_EvaluationContext* InContext) override;
 	virtual bool						ErrorCheckInputTextures() const { return false; }
 };
 
 //////////////////////////////////////////////////////////////////////////
 /// Min
 //////////////////////////////////////////////////////////////////////////
-UCLASS()
-class TEXTUREGRAPH_API UTG_Expression_Min : public UTG_Expression_Variant
+UCLASS(MinimalAPI)
+class UTG_Expression_Min : public UTG_Expression_Variant
 {
 	GENERATED_BODY()
 
@@ -92,15 +94,15 @@ public:
 
 protected:
 	virtual std::vector<FTG_Variant>	GetEvaluateArgs() override { return std::vector<FTG_Variant>({ Input1, Input2 }); }
-	virtual float						EvaluateScalar_WithValue(FTG_EvaluationContext* InContext, const float* const Value, size_t Count) override;
-	virtual FTG_Texture					EvaluateTexture(FTG_EvaluationContext* InContext) override;
+	UE_API virtual float						EvaluateScalar_WithValue(FTG_EvaluationContext* InContext, const float* const Value, size_t Count) override;
+	UE_API virtual FTG_Texture					EvaluateTexture(FTG_EvaluationContext* InContext) override;
 };
 
 //////////////////////////////////////////////////////////////////////////
 /// Max
 //////////////////////////////////////////////////////////////////////////
-UCLASS()
-class TEXTUREGRAPH_API UTG_Expression_Max : public UTG_Expression_Variant
+UCLASS(MinimalAPI)
+class UTG_Expression_Max : public UTG_Expression_Variant
 {
 	GENERATED_BODY()
 
@@ -118,7 +120,8 @@ public:
 
 protected:
 	virtual std::vector<FTG_Variant>	GetEvaluateArgs() override { return std::vector<FTG_Variant>({ Input1, Input2 }); }
-	virtual float						EvaluateScalar_WithValue(FTG_EvaluationContext* InContext, const float* const Value, size_t Count) override;
-	virtual FTG_Texture					EvaluateTexture(FTG_EvaluationContext* InContext) override;
+	UE_API virtual float						EvaluateScalar_WithValue(FTG_EvaluationContext* InContext, const float* const Value, size_t Count) override;
+	UE_API virtual FTG_Texture					EvaluateTexture(FTG_EvaluationContext* InContext) override;
 };
 
+#undef UE_API
diff --git a/Engine/Plugins/TextureGraph/Source/TextureGraph/Public/Expressions/Maths/TG_Expression_IfThenElse.h b/Engine/Plugins/TextureGraph/Source/TextureGraph/Public/Expressions/Maths/TG_Expression_IfThenElse.h
--- a/Engine/Plugins/TextureGraph/Source/TextureGraph/Public/Expressions/Maths/TG_Expression_IfThenElse.h
+++ b/Engine/Plugins/TextureGraph/Source/TextureGraph/Public/Expressions/Maths/TG_Expression_IfThenElse.h
@@ -10,11 +10,13 @@
 
 #include "TG_Expression_IfThenElse.generated.h"
 
+#define UE_API TEXTUREGRAPH_API
+
 //////////////////////////////////////////////////////////////////////////
 /// Expression
 //////////////////////////////////////////////////////////////////////////
-UCLASS()
-class TEXTUREGRAPH_API UTG_Expression_IfThenElse : public UTG_Expression_Variant
+UCLASS(MinimalAPI)
+class UTG_Expression_IfThenElse : public UTG_Expression_Variant
 {
 	GENERATED_BODY()
 
@@ -48,8 +50,9 @@ public:
 
 protected:
 	virtual std::vector<FTG_Variant>	GetEvaluateArgs() override { return std::vector<FTG_Variant>({ LHS, RHS, Then, Else }); }
-	virtual float						EvaluateScalar_WithValue(FTG_EvaluationContext* InContext, const float* const Value, size_t Count) override;
-	virtual FVector4f					EvaluateVector_WithValue(FTG_EvaluationContext* InContext, const FVector4f* const Value, size_t Count);
-	virtual FTG_Texture					EvaluateTexture(FTG_EvaluationContext* InContext) override;
+	UE_API virtual float						EvaluateScalar_WithValue(FTG_EvaluationContext* InContext, const float* const Value, size_t Count) override;
+	UE_API virtual FVector4f					EvaluateVector_WithValue(FTG_EvaluationContext* InContext, const FVector4f* const Value, size_t Count);
+	UE_API virtual FTG_Texture					EvaluateTexture(FTG_EvaluationContext* InContext) override;
 };
 
+#undef UE_API
diff --git a/Engine/Plugins/TextureGraph/Source/TextureGraph/Public/Expressions/Maths/TG_Expression_Invert.h b/Engine/Plugins/TextureGraph/Source/TextureGraph/Public/Expressions/Maths/TG_Expression_Invert.h
--- a/Engine/Plugins/TextureGraph/Source/TextureGraph/Public/Expressions/Maths/TG_Expression_Invert.h
+++ b/Engine/Plugins/TextureGraph/Source/TextureGraph/Public/Expressions/Maths/TG_Expression_Invert.h
@@ -10,17 +10,19 @@
 
 #include "TG_Expression_Invert.generated.h"
 
+#define UE_API TEXTUREGRAPH_API
+
 //////////////////////////////////////////////////////////////////////////
 /// Expression
 //////////////////////////////////////////////////////////////////////////
-UCLASS()
-class TEXTUREGRAPH_API UTG_Expression_Invert : public UTG_Expression
+UCLASS(MinimalAPI)
+class UTG_Expression_Invert : public UTG_Expression
 {
 	GENERATED_BODY()
 
 public:
 	TG_DECLARE_EXPRESSION(TG_Category::Maths);
-	virtual void						Evaluate(FTG_EvaluationContext* InContext) override;
+	UE_API virtual void						Evaluate(FTG_EvaluationContext* InContext) override;
 
 	// The max value to use. This is used as MaxValue - Input for every pixel in every channel to calculate the output
 	// NB: Hiding this property for the time being as requested
@@ -46,3 +48,4 @@ public:
 	virtual FText						GetTooltipText() const override { return FText::FromString(TEXT("Inverts the input, i.e. computes one minus the pixel value. Alpha is only inverted if the flag is checked.")); } 
 };
 
+#undef UE_API
