Ran UnrealCodeFixup to fix dll storage

diff --git a/Engine/Plugins/TextureGraph/Source/TextureGraph/Public/Expressions/Maths/TG_Expression_Maths_OneInput.h b/Engine/Plugins/TextureGraph/Source/TextureGraph/Public/Expressions/Maths/TG_Expression_Maths_OneInput.h
--- a/Engine/Plugins/TextureGraph/Source/TextureGraph/Public/Expressions/Maths/TG_Expression_Maths_OneInput.h
+++ b/Engine/Plugins/TextureGraph/Source/TextureGraph/Public/Expressions/Maths/TG_Expression_Maths_OneInput.h
@@ -9,17 +9,19 @@
 
 #include "TG_Expression_Maths_OneInput.generated.h"
 
+#define UE_API TEXTUREGRAPH_API
+
 //////////////////////////////////////////////////////////////////////////
 /// Trigonometry functions
 //////////////////////////////////////////////////////////////////////////
-UCLASS()
-class TEXTUREGRAPH_API UTG_Expression_Trigonometry : public UTG_Expression_Variant
+UCLASS(MinimalAPI)
+class UTG_Expression_Trigonometry : public UTG_Expression_Variant
 {
 	GENERATED_BODY()
 
 protected:
-	virtual float						EvaluateScalar_WithValue(FTG_EvaluationContext* InContext, const float* const Value, size_t Count) override;
-	virtual FTG_Texture					EvaluateTexture(FTG_EvaluationContext* InContext) override;
+	UE_API virtual float						EvaluateScalar_WithValue(FTG_EvaluationContext* InContext, const float* const Value, size_t Count) override;
+	UE_API virtual FTG_Texture					EvaluateTexture(FTG_EvaluationContext* InContext) override;
 	virtual std::vector<FTG_Variant>	GetEvaluateArgs() override { return std::vector<FTG_Variant>({ Angle }); }
 
 public:
@@ -38,8 +40,8 @@ public:
 //////////////////////////////////////////////////////////////////////////
 /// Generic single input math function
 //////////////////////////////////////////////////////////////////////////
-UCLASS(Abstract)
-class TEXTUREGRAPH_API UTG_Expression_OneInput : public UTG_Expression_Variant
+UCLASS(MinimalAPI, Abstract)
+class UTG_Expression_OneInput : public UTG_Expression_Variant
 {
 	GENERATED_BODY()
 
@@ -57,14 +59,14 @@ public:
 //////////////////////////////////////////////////////////////////////////
 /// Absolute value of a floating point number
 //////////////////////////////////////////////////////////////////////////
-UCLASS()
-class TEXTUREGRAPH_API UTG_Expression_Abs : public UTG_Expression_OneInput
+UCLASS(MinimalAPI)
+class UTG_Expression_Abs : public UTG_Expression_OneInput
 {
 	GENERATED_BODY()
 
 protected:
-	virtual float						EvaluateScalar_WithValue(FTG_EvaluationContext* InContext, const float* const Value, size_t Count) override;
-	virtual FTG_Texture					EvaluateTexture(FTG_EvaluationContext* InContext) override;
+	UE_API virtual float						EvaluateScalar_WithValue(FTG_EvaluationContext* InContext, const float* const Value, size_t Count) override;
+	UE_API virtual FTG_Texture					EvaluateTexture(FTG_EvaluationContext* InContext) override;
 
 public:
 	virtual FText						GetTooltipText() const override { return FText::FromString(TEXT("Absolute value of a floating point number.")); } 
@@ -73,14 +75,14 @@ public:
 //////////////////////////////////////////////////////////////////////////
 /// Square root of a floating point number
 //////////////////////////////////////////////////////////////////////////
-UCLASS()
-class TEXTUREGRAPH_API UTG_Expression_Sqrt : public UTG_Expression_OneInput
+UCLASS(MinimalAPI)
+class UTG_Expression_Sqrt : public UTG_Expression_OneInput
 {
 	GENERATED_BODY()
 
 protected:
-	virtual float						EvaluateScalar_WithValue(FTG_EvaluationContext* InContext, const float* const Value, size_t Count) override;
-	virtual FTG_Texture					EvaluateTexture(FTG_EvaluationContext* InContext) override;
+	UE_API virtual float						EvaluateScalar_WithValue(FTG_EvaluationContext* InContext, const float* const Value, size_t Count) override;
+	UE_API virtual FTG_Texture					EvaluateTexture(FTG_EvaluationContext* InContext) override;
 
 public:
 	virtual FText						GetTooltipText() const override { return FText::FromString(TEXT("Square root of a number.")); } 
@@ -89,14 +91,14 @@ public:
 //////////////////////////////////////////////////////////////////////////
 /// Square of a floating point number
 //////////////////////////////////////////////////////////////////////////
-UCLASS()
-class TEXTUREGRAPH_API UTG_Expression_Square : public UTG_Expression_OneInput
+UCLASS(MinimalAPI)
+class UTG_Expression_Square : public UTG_Expression_OneInput
 {
 	GENERATED_BODY()
 
 protected:
-	virtual float						EvaluateScalar_WithValue(FTG_EvaluationContext* InContext, const float* const Value, size_t Count) override;
-	virtual FTG_Texture					EvaluateTexture(FTG_EvaluationContext* InContext) override;
+	UE_API virtual float						EvaluateScalar_WithValue(FTG_EvaluationContext* InContext, const float* const Value, size_t Count) override;
+	UE_API virtual FTG_Texture					EvaluateTexture(FTG_EvaluationContext* InContext) override;
 
 public:
 	virtual FText						GetTooltipText() const override { return FText::FromString(TEXT("Square of a number.")); } 
@@ -105,14 +107,14 @@ public:
 //////////////////////////////////////////////////////////////////////////
 /// Cube of a floating point number
 //////////////////////////////////////////////////////////////////////////
-UCLASS()
-class TEXTUREGRAPH_API UTG_Expression_Cube : public UTG_Expression_OneInput
+UCLASS(MinimalAPI)
+class UTG_Expression_Cube : public UTG_Expression_OneInput
 {
 	GENERATED_BODY()
 
 protected:
-	virtual float						EvaluateScalar_WithValue(FTG_EvaluationContext* InContext, const float* const Value, size_t Count) override;
-	virtual FTG_Texture					EvaluateTexture(FTG_EvaluationContext* InContext) override;
+	UE_API virtual float						EvaluateScalar_WithValue(FTG_EvaluationContext* InContext, const float* const Value, size_t Count) override;
+	UE_API virtual FTG_Texture					EvaluateTexture(FTG_EvaluationContext* InContext) override;
 
 public:
 	virtual FText						GetTooltipText() const override { return FText::FromString(TEXT("Cube of a number.")); } 
@@ -121,14 +123,14 @@ public:
 //////////////////////////////////////////////////////////////////////////
 /// Cube root of a floating point number
 //////////////////////////////////////////////////////////////////////////
-UCLASS()
-class TEXTUREGRAPH_API UTG_Expression_Cbrt : public UTG_Expression_OneInput
+UCLASS(MinimalAPI)
+class UTG_Expression_Cbrt : public UTG_Expression_OneInput
 {
 	GENERATED_BODY()
 
 protected:
-	virtual float						EvaluateScalar_WithValue(FTG_EvaluationContext* InContext, const float* const Value, size_t Count) override;
-	virtual FTG_Texture					EvaluateTexture(FTG_EvaluationContext* InContext) override;
+	UE_API virtual float						EvaluateScalar_WithValue(FTG_EvaluationContext* InContext, const float* const Value, size_t Count) override;
+	UE_API virtual FTG_Texture					EvaluateTexture(FTG_EvaluationContext* InContext) override;
 
 public:
 	virtual FText						GetTooltipText() const override { return FText::FromString(TEXT("Cubic root of a number.")); } 
@@ -137,14 +139,14 @@ public:
 //////////////////////////////////////////////////////////////////////////
 /// Base-e exponential function
 //////////////////////////////////////////////////////////////////////////
-UCLASS()
-class TEXTUREGRAPH_API UTG_Expression_Exp : public UTG_Expression_OneInput
+UCLASS(MinimalAPI)
+class UTG_Expression_Exp : public UTG_Expression_OneInput
 {
 	GENERATED_BODY()
 
 protected:
-	virtual float						EvaluateScalar_WithValue(FTG_EvaluationContext* InContext, const float* const Value, size_t Count) override;
-	virtual FTG_Texture					EvaluateTexture(FTG_EvaluationContext* InContext) override;
+	UE_API virtual float						EvaluateScalar_WithValue(FTG_EvaluationContext* InContext, const float* const Value, size_t Count) override;
+	UE_API virtual FTG_Texture					EvaluateTexture(FTG_EvaluationContext* InContext) override;
 
 public:
 	virtual FText						GetTooltipText() const override { return FText::FromString(TEXT("Base-e exponential function [Exp(Input)].")); } 
@@ -153,14 +155,14 @@ public:
 //////////////////////////////////////////////////////////////////////////
 /// Log base-2
 //////////////////////////////////////////////////////////////////////////
-UCLASS()
-class TEXTUREGRAPH_API UTG_Expression_Log2 : public UTG_Expression_OneInput
+UCLASS(MinimalAPI)
+class UTG_Expression_Log2 : public UTG_Expression_OneInput
 {
 	GENERATED_BODY()
 
 protected:
-	virtual float						EvaluateScalar_WithValue(FTG_EvaluationContext* InContext, const float* const Value, size_t Count) override;
-	virtual FTG_Texture					EvaluateTexture(FTG_EvaluationContext* InContext) override;
+	UE_API virtual float						EvaluateScalar_WithValue(FTG_EvaluationContext* InContext, const float* const Value, size_t Count) override;
+	UE_API virtual FTG_Texture					EvaluateTexture(FTG_EvaluationContext* InContext) override;
 
 public:
 	virtual FText						GetTooltipText() const override { return FText::FromString(TEXT("Log base-2 of input [Log2(Input)].")); } 
@@ -169,14 +171,14 @@ public:
 //////////////////////////////////////////////////////////////////////////
 /// Log base-10
 //////////////////////////////////////////////////////////////////////////
-UCLASS()
-class TEXTUREGRAPH_API UTG_Expression_Log10 : public UTG_Expression_OneInput
+UCLASS(MinimalAPI)
+class UTG_Expression_Log10 : public UTG_Expression_OneInput
 {
 	GENERATED_BODY()
 
 protected:
-	virtual float						EvaluateScalar_WithValue(FTG_EvaluationContext* InContext, const float* const Value, size_t Count) override;
-	virtual FTG_Texture					EvaluateTexture(FTG_EvaluationContext* InContext) override;
+	UE_API virtual float						EvaluateScalar_WithValue(FTG_EvaluationContext* InContext, const float* const Value, size_t Count) override;
+	UE_API virtual FTG_Texture					EvaluateTexture(FTG_EvaluationContext* InContext) override;
 
 public:
 	virtual FText						GetTooltipText() const override { return FText::FromString(TEXT("Log base-10 of input [Log10(Input)].")); } 
@@ -185,14 +187,14 @@ public:
 //////////////////////////////////////////////////////////////////////////
 /// Log base-e
 //////////////////////////////////////////////////////////////////////////
-UCLASS()
-class TEXTUREGRAPH_API UTG_Expression_Log : public UTG_Expression_OneInput
+UCLASS(MinimalAPI)
+class UTG_Expression_Log : public UTG_Expression_OneInput
 {
 	GENERATED_BODY()
 
 protected:
-	virtual float						EvaluateScalar_WithValue(FTG_EvaluationContext* InContext, const float* const Value, size_t Count) override;
-	virtual FTG_Texture					EvaluateTexture(FTG_EvaluationContext* InContext) override;
+	UE_API virtual float						EvaluateScalar_WithValue(FTG_EvaluationContext* InContext, const float* const Value, size_t Count) override;
+	UE_API virtual FTG_Texture					EvaluateTexture(FTG_EvaluationContext* InContext) override;
 
 public:
 	virtual FText						GetTooltipText() const override { return FText::FromString(TEXT("Log natural (Log base-e) of input [Log(Input)].")); } 
@@ -201,14 +203,14 @@ public:
 //////////////////////////////////////////////////////////////////////////
 /// Floor
 //////////////////////////////////////////////////////////////////////////
-UCLASS()
-class TEXTUREGRAPH_API UTG_Expression_Floor : public UTG_Expression_OneInput
+UCLASS(MinimalAPI)
+class UTG_Expression_Floor : public UTG_Expression_OneInput
 {
 	GENERATED_BODY()
 
 protected:
-	virtual float						EvaluateScalar_WithValue(FTG_EvaluationContext* InContext, const float* const Value, size_t Count) override;
-	virtual FTG_Texture					EvaluateTexture(FTG_EvaluationContext* InContext) override;
+	UE_API virtual float						EvaluateScalar_WithValue(FTG_EvaluationContext* InContext, const float* const Value, size_t Count) override;
+	UE_API virtual FTG_Texture					EvaluateTexture(FTG_EvaluationContext* InContext) override;
 
 public:
 	virtual FText						GetTooltipText() const override { return FText::FromString(TEXT("Floor of a number.")); } 
@@ -217,14 +219,14 @@ public:
 //////////////////////////////////////////////////////////////////////////
 /// Ceiling
 //////////////////////////////////////////////////////////////////////////
-UCLASS()
-class TEXTUREGRAPH_API UTG_Expression_Ceil : public UTG_Expression_OneInput
+UCLASS(MinimalAPI)
+class UTG_Expression_Ceil : public UTG_Expression_OneInput
 {
 	GENERATED_BODY()
 
 protected:
-	virtual float						EvaluateScalar_WithValue(FTG_EvaluationContext* InContext, const float* const Value, size_t Count) override;
-	virtual FTG_Texture					EvaluateTexture(FTG_EvaluationContext* InContext) override;
+	UE_API virtual float						EvaluateScalar_WithValue(FTG_EvaluationContext* InContext, const float* const Value, size_t Count) override;
+	UE_API virtual FTG_Texture					EvaluateTexture(FTG_EvaluationContext* InContext) override;
 
 public:
 	virtual FText						GetTooltipText() const override { return FText::FromString(TEXT("Ceiling of a number.")); } 
@@ -233,16 +235,17 @@ public:
 //////////////////////////////////////////////////////////////////////////
 /// Round
 //////////////////////////////////////////////////////////////////////////
-UCLASS()
-class TEXTUREGRAPH_API UTG_Expression_Round : public UTG_Expression_OneInput
+UCLASS(MinimalAPI)
+class UTG_Expression_Round : public UTG_Expression_OneInput
 {
 	GENERATED_BODY()
 
 protected:
-	virtual float						EvaluateScalar_WithValue(FTG_EvaluationContext* InContext, const float* const Value, size_t Count) override;
-	virtual FTG_Texture					EvaluateTexture(FTG_EvaluationContext* InContext) override;
+	UE_API virtual float						EvaluateScalar_WithValue(FTG_EvaluationContext* InContext, const float* const Value, size_t Count) override;
+	UE_API virtual FTG_Texture					EvaluateTexture(FTG_EvaluationContext* InContext) override;
 
 public:
 	virtual FText						GetTooltipText() const override { return FText::FromString(TEXT("Round a number.")); } 
 };
 
+#undef UE_API
diff --git a/Engine/Plugins/TextureGraph/Source/TextureGraph/Public/Expressions/Maths/TG_Expression_Maths_ThreeInputs.h b/Engine/Plugins/TextureGraph/Source/TextureGraph/Public/Expressions/Maths/TG_Expression_Maths_ThreeInputs.h
--- a/Engine/Plugins/TextureGraph/Source/TextureGraph/Public/Expressions/Maths/TG_Expression_Maths_ThreeInputs.h
+++ b/Engine/Plugins/TextureGraph/Source/TextureGraph/Public/Expressions/Maths/TG_Expression_Maths_ThreeInputs.h
@@ -8,11 +8,13 @@
 
 #include "TG_Expression_Maths_ThreeInputs.generated.h"
 
+#define UE_API TEXTUREGRAPH_API
+
 //////////////////////////////////////////////////////////////////////////
 /// MAD
 //////////////////////////////////////////////////////////////////////////
-UCLASS()
-class TEXTUREGRAPH_API UTG_Expression_MAD : public UTG_Expression_Variant
+UCLASS(MinimalAPI)
+class UTG_Expression_MAD : public UTG_Expression_Variant
 {
 public:
 	GENERATED_BODY()
@@ -35,16 +37,16 @@ public:
 	virtual FText						GetTooltipText() const override { return FText::FromString(TEXT("Multiplies the first two inputs and then adds the third input, A \u2192 B \u002b C.")); } 
 
 protected:
-	virtual float						EvaluateScalar_WithValue(FTG_EvaluationContext* InContext, const float* const Value, size_t Count) override;
-	virtual FTG_Texture					EvaluateTexture(FTG_EvaluationContext* InContext) override;
+	UE_API virtual float						EvaluateScalar_WithValue(FTG_EvaluationContext* InContext, const float* const Value, size_t Count) override;
+	UE_API virtual FTG_Texture					EvaluateTexture(FTG_EvaluationContext* InContext) override;
 	virtual std::vector<FTG_Variant>	GetEvaluateArgs() override { return std::vector<FTG_Variant>({ Input1, Input2, Input3 }); }
 };
 
 //////////////////////////////////////////////////////////////////////////
 /// Lerp
 //////////////////////////////////////////////////////////////////////////
-UCLASS()
-class TEXTUREGRAPH_API UTG_Expression_Lerp : public UTG_Expression_Variant
+UCLASS(MinimalAPI)
+class UTG_Expression_Lerp : public UTG_Expression_Variant
 {
 public:
 	GENERATED_BODY()
@@ -67,7 +69,8 @@ public:
 
 protected:
 	virtual std::vector<FTG_Variant>	GetEvaluateArgs() override { return std::vector<FTG_Variant>({ Input1, Input2, LerpValue }); }
-	virtual float						EvaluateScalar_WithValue(FTG_EvaluationContext* InContext, const float* const Value, size_t Count) override;
-	virtual FTG_Texture					EvaluateTexture(FTG_EvaluationContext* InContext) override;
+	UE_API virtual float						EvaluateScalar_WithValue(FTG_EvaluationContext* InContext, const float* const Value, size_t Count) override;
+	UE_API virtual FTG_Texture					EvaluateTexture(FTG_EvaluationContext* InContext) override;
 };
 
+#undef UE_API
diff --git a/Engine/Plugins/TextureGraph/Source/TextureGraph/Public/Expressions/Maths/TG_Expression_Maths_TwoInputs.h b/Engine/Plugins/TextureGraph/Source/TextureGraph/Public/Expressions/Maths/TG_Expression_Maths_TwoInputs.h
--- a/Engine/Plugins/TextureGraph/Source/TextureGraph/Public/Expressions/Maths/TG_Expression_Maths_TwoInputs.h
+++ b/Engine/Plugins/TextureGraph/Source/TextureGraph/Public/Expressions/Maths/TG_Expression_Maths_TwoInputs.h
@@ -8,11 +8,13 @@
 
 #include "TG_Expression_Maths_TwoInputs.generated.h"
 
+#define UE_API TEXTUREGRAPH_API
+
 //////////////////////////////////////////////////////////////////////////
 /// Basic math op expression
 //////////////////////////////////////////////////////////////////////////
-UCLASS(Abstract)
-class TEXTUREGRAPH_API UTG_Expression_BasicMath_Op : public UTG_Expression_Variant
+UCLASS(MinimalAPI, Abstract)
+class UTG_Expression_BasicMath_Op : public UTG_Expression_Variant
 {
 	GENERATED_BODY()
 
@@ -34,8 +36,8 @@ protected:
 //////////////////////////////////////////////////////////////////////////
 /// Multiply
 //////////////////////////////////////////////////////////////////////////
-UCLASS()
-class TEXTUREGRAPH_API UTG_Expression_Multiply : public UTG_Expression_BasicMath_Op
+UCLASS(MinimalAPI)
+class UTG_Expression_Multiply : public UTG_Expression_BasicMath_Op
 {
 	GENERATED_BODY()
 	
@@ -43,15 +45,15 @@ public:
 	virtual FText						GetTooltipText() const override { return FText::FromString(TEXT("Multiplies the two inputs, A \u2192 B.")); } 
 
 protected:
-	virtual float						EvaluateScalar_WithValue(FTG_EvaluationContext* InContext, const float* const Value, size_t Count) override;
-	virtual FTG_Texture					EvaluateTexture(FTG_EvaluationContext* InContext) override;
+	UE_API virtual float						EvaluateScalar_WithValue(FTG_EvaluationContext* InContext, const float* const Value, size_t Count) override;
+	UE_API virtual FTG_Texture					EvaluateTexture(FTG_EvaluationContext* InContext) override;
 };
 
 //////////////////////////////////////////////////////////////////////////
 /// Divide
 //////////////////////////////////////////////////////////////////////////
-UCLASS()
-class TEXTUREGRAPH_API UTG_Expression_Divide : public UTG_Expression_BasicMath_Op
+UCLASS(MinimalAPI)
+class UTG_Expression_Divide : public UTG_Expression_BasicMath_Op
 {
 	GENERATED_BODY()
 	
@@ -59,15 +61,15 @@ public:
 	virtual FText						GetTooltipText() const override { return FText::FromString(TEXT("Divides the two inputs, A \u00f7 B.")); } 
 
 protected:
-	virtual float						EvaluateScalar_WithValue(FTG_EvaluationContext* InContext, const float* const Value, size_t Count) override;
-	virtual FTG_Texture					EvaluateTexture(FTG_EvaluationContext* InContext) override;
+	UE_API virtual float						EvaluateScalar_WithValue(FTG_EvaluationContext* InContext, const float* const Value, size_t Count) override;
+	UE_API virtual FTG_Texture					EvaluateTexture(FTG_EvaluationContext* InContext) override;
 };
 
 //////////////////////////////////////////////////////////////////////////
 /// Add
 //////////////////////////////////////////////////////////////////////////
-UCLASS()
-class TEXTUREGRAPH_API UTG_Expression_Add : public UTG_Expression_BasicMath_Op
+UCLASS(MinimalAPI)
+class UTG_Expression_Add : public UTG_Expression_BasicMath_Op
 {
 	GENERATED_BODY()
 	
@@ -75,15 +77,15 @@ public:
 	virtual FText						GetTooltipText() const override { return FText::FromString(TEXT("Adds up the two inputs, A \u002b B.")); } 
 
 protected:
-	virtual float						EvaluateScalar_WithValue(FTG_EvaluationContext* InContext, const float* const Value, size_t Count) override;
-	virtual FTG_Texture					EvaluateTexture(FTG_EvaluationContext* InContext) override;
+	UE_API virtual float						EvaluateScalar_WithValue(FTG_EvaluationContext* InContext, const float* const Value, size_t Count) override;
+	UE_API virtual FTG_Texture					EvaluateTexture(FTG_EvaluationContext* InContext) override;
 };
 
 //////////////////////////////////////////////////////////////////////////
 /// Subtract
 //////////////////////////////////////////////////////////////////////////
-UCLASS()
-class TEXTUREGRAPH_API UTG_Expression_Subtract : public UTG_Expression_BasicMath_Op
+UCLASS(MinimalAPI)
+class UTG_Expression_Subtract : public UTG_Expression_BasicMath_Op
 {
 	GENERATED_BODY()
 	
@@ -91,15 +93,15 @@ public:
 	virtual FText						GetTooltipText() const override { return FText::FromString(TEXT("Subtracts the two inputs, A \u2212 B.")); } 
 
 protected:
-	virtual float						EvaluateScalar_WithValue(FTG_EvaluationContext* InContext, const float* const Value, size_t Count) override;
-	virtual FTG_Texture					EvaluateTexture(FTG_EvaluationContext* InContext) override;
+	UE_API virtual float						EvaluateScalar_WithValue(FTG_EvaluationContext* InContext, const float* const Value, size_t Count) override;
+	UE_API virtual FTG_Texture					EvaluateTexture(FTG_EvaluationContext* InContext) override;
 };
 
 //////////////////////////////////////////////////////////////////////////
 /// Dot product of two vectors
 //////////////////////////////////////////////////////////////////////////
-UCLASS()
-class TEXTUREGRAPH_API UTG_Expression_Dot : public UTG_Expression_BasicMath_Op
+UCLASS(MinimalAPI)
+class UTG_Expression_Dot : public UTG_Expression_BasicMath_Op
 {
 	GENERATED_BODY()
 	
@@ -107,20 +109,20 @@ public:
 	virtual FText						GetTooltipText() const override { return FText::FromString(TEXT("Dot product of inputs. A <dot> B.")); } 
 
 protected:
-	virtual float						EvaluateScalar_WithValue(FTG_EvaluationContext* InContext, const float* const Value, size_t Count) override;
-	virtual FVector4f					EvaluateVector_WithValue(FTG_EvaluationContext* InContext, const FVector4f* const Value, size_t Count) override;
-	virtual FTG_Texture					EvaluateTexture(FTG_EvaluationContext* InContext) override;
+	UE_API virtual float						EvaluateScalar_WithValue(FTG_EvaluationContext* InContext, const float* const Value, size_t Count) override;
+	UE_API virtual FVector4f					EvaluateVector_WithValue(FTG_EvaluationContext* InContext, const FVector4f* const Value, size_t Count) override;
+	UE_API virtual FTG_Texture					EvaluateTexture(FTG_EvaluationContext* InContext) override;
 
-	virtual void						Evaluate(FTG_EvaluationContext* InContext) override;
+	UE_API virtual void						Evaluate(FTG_EvaluationContext* InContext) override;
 
-	virtual FTG_Variant::EType			EvalExpressionCommonOutputVariantType() const override; // specify the logic to decide the output variant type
+	UE_API virtual FTG_Variant::EType			EvalExpressionCommonOutputVariantType() const override; // specify the logic to decide the output variant type
 };
 
 //////////////////////////////////////////////////////////////////////////
 /// Cross product of two vectors
 //////////////////////////////////////////////////////////////////////////
-UCLASS()
-class TEXTUREGRAPH_API UTG_Expression_Cross : public UTG_Expression_BasicMath_Op
+UCLASS(MinimalAPI)
+class UTG_Expression_Cross : public UTG_Expression_BasicMath_Op
 {
 	GENERATED_BODY()
 	
@@ -128,16 +130,16 @@ public:
 	virtual FText						GetTooltipText() const override { return FText::FromString(TEXT("Cross product of inputs. A <cross> B.")); } 
 
 protected:
-	virtual float						EvaluateScalar_WithValue(FTG_EvaluationContext* InContext, const float* const Value, size_t Count) override;
-	virtual FVector4f					EvaluateVector_WithValue(FTG_EvaluationContext* InContext, const FVector4f* const Value, size_t Count) override;
-	virtual FTG_Texture					EvaluateTexture(FTG_EvaluationContext* InContext) override;
+	UE_API virtual float						EvaluateScalar_WithValue(FTG_EvaluationContext* InContext, const float* const Value, size_t Count) override;
+	UE_API virtual FVector4f					EvaluateVector_WithValue(FTG_EvaluationContext* InContext, const FVector4f* const Value, size_t Count) override;
+	UE_API virtual FTG_Texture					EvaluateTexture(FTG_EvaluationContext* InContext) override;
 };
 
 //////////////////////////////////////////////////////////////////////////
 /// Step function: Usual HLSL step style function
 //////////////////////////////////////////////////////////////////////////
-UCLASS()
-class TEXTUREGRAPH_API UTG_Expression_Step : public UTG_Expression_BasicMath_Op
+UCLASS(MinimalAPI)
+class UTG_Expression_Step : public UTG_Expression_BasicMath_Op
 {
 	GENERATED_BODY()
 	
@@ -145,15 +147,15 @@ public:
 	virtual FText						GetTooltipText() const override { return FText::FromString(TEXT("step(y, x) returns 1 if the x parameter is greater than or equal to the y parameter; otherwise, 0.")); } 
 
 protected:
-	virtual float						EvaluateScalar_WithValue(FTG_EvaluationContext* InContext, const float* const Value, size_t Count) override;
-	virtual FTG_Texture					EvaluateTexture(FTG_EvaluationContext* InContext) override;
+	UE_API virtual float						EvaluateScalar_WithValue(FTG_EvaluationContext* InContext, const float* const Value, size_t Count) override;
+	UE_API virtual FTG_Texture					EvaluateTexture(FTG_EvaluationContext* InContext) override;
 };
 
 //////////////////////////////////////////////////////////////////////////
 /// Power
 //////////////////////////////////////////////////////////////////////////
-UCLASS()
-class TEXTUREGRAPH_API UTG_Expression_Pow : public UTG_Expression_Variant
+UCLASS(MinimalAPI)
+class UTG_Expression_Pow : public UTG_Expression_Variant
 {
 	GENERATED_BODY()
 	
@@ -171,7 +173,8 @@ public:
 
 protected:
 	virtual std::vector<FTG_Variant>	GetEvaluateArgs() override { return std::vector<FTG_Variant>({ Base, Exponent }); }
-	virtual float						EvaluateScalar_WithValue(FTG_EvaluationContext* InContext, const float* const Value, size_t Count) override;
-	virtual FTG_Texture					EvaluateTexture(FTG_EvaluationContext* InContext) override;
+	UE_API virtual float						EvaluateScalar_WithValue(FTG_EvaluationContext* InContext, const float* const Value, size_t Count) override;
+	UE_API virtual FTG_Texture					EvaluateTexture(FTG_EvaluationContext* InContext) override;
 };
 
+#undef UE_API
diff --git a/Engine/Plugins/TextureGraph/Source/TextureGraph/Public/Expressions/Output/TG_Expression_Output.h b/Engine/Plugins/TextureGraph/Source/TextureGraph/Public/Expressions/Output/TG_Expression_Output.h
--- a/Engine/Plugins/TextureGraph/Source/TextureGraph/Public/Expressions/Output/TG_Expression_Output.h
+++ b/Engine/Plugins/TextureGraph/Source/TextureGraph/Public/Expressions/Output/TG_Expression_Output.h
@@ -10,16 +10,18 @@
 
 #include "TG_Expression_Output.generated.h"
 
-UCLASS()
-class TEXTUREGRAPH_API UTG_Expression_Output : public UTG_Expression
+#define UE_API TEXTUREGRAPH_API
+
+UCLASS(MinimalAPI)
+class UTG_Expression_Output : public UTG_Expression
 {
 	GENERATED_BODY()
 
 public:
 	TG_DECLARE_VARIANT_EXPRESSION(TG_Category::Output);
-	virtual void Evaluate(FTG_EvaluationContext* InContext) override;
+	UE_API virtual void Evaluate(FTG_EvaluationContext* InContext) override;
 
-	virtual bool Validate(MixUpdateCyclePtr Cycle) override;
+	UE_API virtual bool Validate(MixUpdateCyclePtr Cycle) override;
 	
 	// The final generated output
 	UPROPERTY(meta = (TGType = "TG_Input", PinDisplayName = ""))
@@ -37,15 +39,17 @@ public:
 	bool bShouldExport = true;
 
 	virtual FTG_Name GetDefaultName() const override { return TEXT("Output");}
-	virtual void SetTitleName(FName NewName) override;
-	virtual FName GetTitleName() const override;
+	UE_API virtual void SetTitleName(FName NewName) override;
+	UE_API virtual FName GetTitleName() const override;
 	virtual FText GetTooltipText() const override { return FText::FromString(TEXT("Uses the node input as a graph output. It can be exported as a texture asset and is automatically exposed as a graph output parameter. ")); } 
 	virtual bool ShouldShowSettings() const { return false; }
 
-	void UpdateBufferDescriptorValues();
-	void InitializeOutputSettings();
+	UE_API void UpdateBufferDescriptorValues();
+	UE_API void InitializeOutputSettings();
 
-	void SetShouldExport(bool InShouldExport);
+	UE_API void SetShouldExport(bool InShouldExport);
 
 	bool GetShouldExport() const { return OutputSettings.bShouldExport; }
 };
+
+#undef UE_API
diff --git a/Engine/Plugins/TextureGraph/Source/TextureGraph/Public/Expressions/Procedural/TG_Expression_Gradient.h b/Engine/Plugins/TextureGraph/Source/TextureGraph/Public/Expressions/Procedural/TG_Expression_Gradient.h
--- a/Engine/Plugins/TextureGraph/Source/TextureGraph/Public/Expressions/Procedural/TG_Expression_Gradient.h
+++ b/Engine/Plugins/TextureGraph/Source/TextureGraph/Public/Expressions/Procedural/TG_Expression_Gradient.h
@@ -11,19 +11,21 @@
 
 #include "TG_Expression_Gradient.generated.h"
 
-UCLASS()
-class TEXTUREGRAPH_API UTG_Expression_Gradient : public UTG_Expression
+#define UE_API TEXTUREGRAPH_API
+
+UCLASS(MinimalAPI)
+class UTG_Expression_Gradient : public UTG_Expression
 {
 	GENERATED_BODY()
 
 public:
 	TG_DECLARE_EXPRESSION(TG_Category::Procedural)
 	
-	virtual void						Evaluate(FTG_EvaluationContext* InContext) override;
+	UE_API virtual void						Evaluate(FTG_EvaluationContext* InContext) override;
 
 #if WITH_EDITOR
 	// Used to implement EditCondition logic for both Node UI and Details View
-	virtual bool						CanEditChange(const FProperty* InProperty) const override;
+	UE_API virtual bool						CanEditChange(const FProperty* InProperty) const override;
 #endif
 
 
@@ -65,3 +67,5 @@ public:
 	
 	virtual FText						GetTooltipText() const override { return FText::FromString(TEXT("Generates different types of gradients.")); } 
 };
+
+#undef UE_API
diff --git a/Engine/Plugins/TextureGraph/Source/TextureGraph/Public/Expressions/Procedural/TG_Expression_Noise.h b/Engine/Plugins/TextureGraph/Source/TextureGraph/Public/Expressions/Procedural/TG_Expression_Noise.h
--- a/Engine/Plugins/TextureGraph/Source/TextureGraph/Public/Expressions/Procedural/TG_Expression_Noise.h
+++ b/Engine/Plugins/TextureGraph/Source/TextureGraph/Public/Expressions/Procedural/TG_Expression_Noise.h
@@ -10,6 +10,8 @@
 
 #include "TG_Expression_Noise.generated.h"
 
+#define UE_API TEXTUREGRAPH_API
+
 UENUM(BlueprintType)
 enum ENoiseType : int
 {
@@ -20,8 +22,8 @@ enum ENoiseType : int
 	NOISETYPE_Worley3					UMETA(DisplayName = "Worley3"),
 };
 
-UCLASS()
-class TEXTUREGRAPH_API UTG_Expression_Noise : public UTG_Expression
+UCLASS(MinimalAPI)
+class UTG_Expression_Noise : public UTG_Expression
 {
 	GENERATED_BODY()
 
@@ -29,7 +31,7 @@ class TEXTUREGRAPH_API UTG_Expression_Noise : public UTG_Expression
 public:
 	TG_DECLARE_EXPRESSION(TG_Category::Procedural)
 	
-	virtual void						Evaluate(FTG_EvaluationContext* InContext) override;
+	UE_API virtual void						Evaluate(FTG_EvaluationContext* InContext) override;
 
 	// The type of the noise function
 	UPROPERTY(EditAnywhere, BlueprintReadWrite, Category = NoCategory, meta = (TGType = "TG_Setting"))
@@ -66,3 +68,4 @@ public:
 	virtual FText						GetTooltipText() const override { return FText::FromString(TEXT("Generates different types of image noise.")); } 
 };
 
+#undef UE_API
diff --git a/Engine/Plugins/TextureGraph/Source/TextureGraph/Public/Expressions/Procedural/TG_Expression_Pattern.h b/Engine/Plugins/TextureGraph/Source/TextureGraph/Public/Expressions/Procedural/TG_Expression_Pattern.h
--- a/Engine/Plugins/TextureGraph/Source/TextureGraph/Public/Expressions/Procedural/TG_Expression_Pattern.h
+++ b/Engine/Plugins/TextureGraph/Source/TextureGraph/Public/Expressions/Procedural/TG_Expression_Pattern.h
@@ -11,6 +11,8 @@
 
 #include "TG_Expression_Pattern.generated.h"
 
+#define UE_API TEXTUREGRAPH_API
+
 UENUM(BlueprintType)
 enum EPatternType : int 
 {
@@ -21,7 +23,7 @@ enum EPatternType : int
 };
 
 USTRUCT(BlueprintType)
-struct TEXTUREGRAPH_API FPatternMaskPlacement_TS
+struct FPatternMaskPlacement_TS
 {
 	GENERATED_BODY()
 	
@@ -68,7 +70,7 @@ struct TEXTUREGRAPH_API FPatternMaskPlacement_TS
 };
 
 USTRUCT(BlueprintType)
-struct TEXTUREGRAPH_API FPatternMaskJitter_TS
+struct FPatternMaskJitter_TS
 {
 	GENERATED_BODY()
 
@@ -142,7 +144,7 @@ struct TEXTUREGRAPH_API FPatternMaskJitter_TS
 };
 
 USTRUCT(BlueprintType)
-struct TEXTUREGRAPH_API FPatternMaskBevel_TS
+struct FPatternMaskBevel_TS
 {
 	GENERATED_BODY()
 	
@@ -169,7 +171,7 @@ struct TEXTUREGRAPH_API FPatternMaskBevel_TS
 };
 
 USTRUCT(BlueprintType)
-struct TEXTUREGRAPH_API FPatternMaskCutout_TS
+struct FPatternMaskCutout_TS
 {
 	GENERATED_BODY()
 	
@@ -196,7 +198,7 @@ struct TEXTUREGRAPH_API FPatternMaskCutout_TS
 };
 
 USTRUCT(BlueprintType)
-struct TEXTUREGRAPH_API FGradientDir_TS
+struct FGradientDir_TS
 {
 	GENERATED_BODY()
 	
@@ -224,8 +226,8 @@ struct TEXTUREGRAPH_API FGradientDir_TS
 	}
 };
 
-UCLASS()
-class TEXTUREGRAPH_API UTG_Expression_Pattern : public UTG_Expression
+UCLASS(MinimalAPI)
+class UTG_Expression_Pattern : public UTG_Expression
 {
 	GENERATED_BODY()
 
@@ -233,7 +235,7 @@ class TEXTUREGRAPH_API UTG_Expression_Pattern : public UTG_Expression
 public:
 	TG_DECLARE_EXPRESSION(TG_Category::Procedural)
 	
-	virtual void						Evaluate(FTG_EvaluationContext* InContext) override;
+	UE_API virtual void						Evaluate(FTG_EvaluationContext* InContext) override;
 
 	// The type of pattern
 	UPROPERTY(EditAnywhere, BlueprintReadWrite, Category = NoCategory, meta = (TGType = "TG_Setting"))
@@ -266,3 +268,4 @@ public:
 	virtual FText						GetTooltipText() const override { return FText::FromString(TEXT("Generates different types of geometric patterns.")); } 
 };
 
+#undef UE_API
diff --git a/Engine/Plugins/TextureGraph/Source/TextureGraph/Public/Expressions/Procedural/TG_Expression_Shape.h b/Engine/Plugins/TextureGraph/Source/TextureGraph/Public/Expressions/Procedural/TG_Expression_Shape.h
--- a/Engine/Plugins/TextureGraph/Source/TextureGraph/Public/Expressions/Procedural/TG_Expression_Shape.h
+++ b/Engine/Plugins/TextureGraph/Source/TextureGraph/Public/Expressions/Procedural/TG_Expression_Shape.h
@@ -10,6 +10,8 @@
 
 #include "TG_Expression_Shape.generated.h"
 
+#define UE_API TEXTUREGRAPH_API
+
 UENUM(BlueprintType)
 enum class EShapeType : uint8
 {
@@ -23,8 +25,8 @@ enum class EShapeType : uint8
 	Polygon				UMETA(DisplayName = "Polygon"),
 };
 
-UCLASS()
-class TEXTUREGRAPH_API UTG_Expression_Shape : public UTG_Expression
+UCLASS(MinimalAPI)
+class UTG_Expression_Shape : public UTG_Expression
 {
 	GENERATED_BODY()
 
@@ -32,11 +34,11 @@ class TEXTUREGRAPH_API UTG_Expression_Shape : public UTG_Expression
 public:
 	TG_DECLARE_EXPRESSION(TG_Category::Procedural)
 	
-	virtual void						Evaluate(FTG_EvaluationContext* InContext) override;
+	UE_API virtual void						Evaluate(FTG_EvaluationContext* InContext) override;
 
 #if WITH_EDITOR
 	// Used to implement EditCondition logic for both Node UI and Details View
-	virtual bool						CanEditChange(const FProperty* InProperty) const override;
+	UE_API virtual bool						CanEditChange(const FProperty* InProperty) const override;
 #endif
 	
 	// Select the geometric primitive shape to generate
@@ -84,3 +86,4 @@ public:
 
 
 
+#undef UE_API
diff --git a/Engine/Plugins/TextureGraph/Source/TextureGraph/Public/Expressions/Procedural/TG_Expression_Transform.h b/Engine/Plugins/TextureGraph/Source/TextureGraph/Public/Expressions/Procedural/TG_Expression_Transform.h
--- a/Engine/Plugins/TextureGraph/Source/TextureGraph/Public/Expressions/Procedural/TG_Expression_Transform.h
+++ b/Engine/Plugins/TextureGraph/Source/TextureGraph/Public/Expressions/Procedural/TG_Expression_Transform.h
@@ -9,18 +9,20 @@
 
 #include "TG_Expression_Transform.generated.h"
 
+#define UE_API TEXTUREGRAPH_API
+
 
 //////////////////////////////////////////////////////////////////////////
 /// Expression
 //////////////////////////////////////////////////////////////////////////
-UCLASS()
-class TEXTUREGRAPH_API UTG_Expression_Transform : public UTG_Expression
+UCLASS(MinimalAPI)
+class UTG_Expression_Transform : public UTG_Expression
 {
 	GENERATED_BODY()
 
 public:
 	TG_DECLARE_EXPRESSION(TG_Category::Procedural);
-	virtual void							Evaluate(FTG_EvaluationContext* InContext) override;
+	UE_API virtual void							Evaluate(FTG_EvaluationContext* InContext) override;
 
 	// Output image
 	UPROPERTY(EditAnywhere, Category = NoCategory, meta = (TGType = "TG_Output") )
@@ -96,3 +98,4 @@ public:
 	virtual FText							GetTooltipText() const override { return FText::FromString(TEXT("Performs translation, rotation and repetition of an image")); } 
 };
 
+#undef UE_API
diff --git a/Engine/Plugins/TextureGraph/Source/TextureGraph/Public/Expressions/TG_Expression.h b/Engine/Plugins/TextureGraph/Source/TextureGraph/Public/Expressions/TG_Expression.h
--- a/Engine/Plugins/TextureGraph/Source/TextureGraph/Public/Expressions/TG_Expression.h
+++ b/Engine/Plugins/TextureGraph/Source/TextureGraph/Public/Expressions/TG_Expression.h
@@ -10,29 +10,31 @@
 
 #include "TG_Expression.generated.h"
 
+#define UE_API TEXTUREGRAPH_API
+
 #if WITH_EDITOR
 DECLARE_MULTICAST_DELEGATE_OneParam(FOnTSExpressionChanged, UTG_Expression*);
 #endif
 
-class TEXTUREGRAPH_API TG_Category
+class TG_Category
 {
 public:
-	static const FName Default;
-	static const FName Output;
-	static const FName Input;
-	static const FName Adjustment;
-	static const FName Channel;
-	static const FName DevOnly;
-	static const FName Procedural;
-	static const FName Maths;
-	static const FName Utilities;
-	static const FName Filter;
-	static const FName Arrays;
-	static const FName Custom;
+	static UE_API const FName Default;
+	static UE_API const FName Output;
+	static UE_API const FName Input;
+	static UE_API const FName Adjustment;
+	static UE_API const FName Channel;
+	static UE_API const FName DevOnly;
+	static UE_API const FName Procedural;
+	static UE_API const FName Maths;
+	static UE_API const FName Utilities;
+	static UE_API const FName Filter;
+	static UE_API const FName Arrays;
+	static UE_API const FName Custom;
 };
 
-UCLASS(abstract)
-class TEXTUREGRAPH_API UTG_Expression : public UObject
+UCLASS(MinimalAPI, abstract)
+class UTG_Expression : public UObject
 {
 	GENERATED_BODY()
 
@@ -52,17 +54,17 @@ public:
 	FTG_TextureDescriptor BaseOutputSettings;
 
 	// Override serialization to save Class Serialization version
-	virtual	void		Serialize(FArchive& Ar) override;
+	UE_API virtual	void		Serialize(FArchive& Ar) override;
 
 
-	virtual FTG_Name GetDefaultName() const;
+	UE_API virtual FTG_Name GetDefaultName() const;
 	virtual FName GetCategory() const { return TG_Category::Default; }
 	virtual FTG_SignaturePtr GetSignature() const { return nullptr; }
 	virtual FText GetTooltipText() const { return FText::FromString(TEXT("Texture Graph Node")); } 
 
 	virtual void SetTitleName(FName NewName) {}
 	virtual FName GetTitleName() const { return GetDefaultName(); }
-	virtual bool CanRenameTitle() const;
+	UE_API virtual bool CanRenameTitle() const;
 
 	virtual bool CanHandleAsset(UObject* Asset) { return false; };
 	virtual void SetAsset(UObject* Asset) { check(CanHandleAsset((Asset))); };
@@ -71,18 +73,18 @@ public:
 	virtual void Evaluate(FTG_EvaluationContext* InContext) {}
 
 #if WITH_EDITOR
-	void PropertyChangeTriggered(FProperty* Property, EPropertyChangeType::Type ChangeType);
-	virtual bool CanEditChange( const FProperty* InProperty ) const override;	
+	UE_API void PropertyChangeTriggered(FProperty* Property, EPropertyChangeType::Type ChangeType);
+	UE_API virtual bool CanEditChange( const FProperty* InProperty ) const override;	
 #endif
 	
 	// Validate internal checks, warnings and errors
 	virtual bool Validate(MixUpdateCyclePtr	Cycle) { return true; }
 
 	// Validate that the expression conforms to a conformant function (e.g. Clamp etc.)
-	void ValidateGenerateConformer(UTG_Pin* InPin);
+	UE_API void ValidateGenerateConformer(UTG_Pin* InPin);
 
 	// Expression notifies its parent node on key events
-	virtual	UTG_Node* GetParentNode() const final;
+	UE_API virtual	UTG_Node* GetParentNode() const final;
 protected:
 	friend class UTG_Graph;
 	friend class UTG_Node;
@@ -97,21 +99,21 @@ protected:
 	// This is where the Var values are copied over to the matching Expression's properties
 	// This is NOT the function you want to derive, unless you know exactly what you are doing.
 	// Instead overwrite the Execute function.
-	virtual void SetupAndEvaluate(FTG_EvaluationContext* InContext);
+	UE_API virtual void SetupAndEvaluate(FTG_EvaluationContext* InContext);
 
-	virtual void CopyVarToExpressionArgument(const FTG_Argument& Arg, FTG_Var* InVar);
-	virtual void CopyVarFromExpressionArgument(const FTG_Argument& Arg, FTG_Var* InVar);
+	UE_API virtual void CopyVarToExpressionArgument(const FTG_Argument& Arg, FTG_Var* InVar);
+	UE_API virtual void CopyVarFromExpressionArgument(const FTG_Argument& Arg, FTG_Var* InVar);
 	virtual void CopyVarGeneric(const FTG_Argument& Arg, FTG_Var* InVar, bool CopyVarToArg) {}
 
 	// Log the actual values for vars and the expression evaluation
 	// Called from SetupAndEvaluate call if Context ask for it
-	virtual void LogEvaluation(FTG_EvaluationContext* InContext);
+	UE_API virtual void LogEvaluation(FTG_EvaluationContext* InContext);
 
-	virtual FTG_Signature::FInit GetSignatureInitArgsFromClass() const final;
+	UE_API virtual FTG_Signature::FInit GetSignatureInitArgsFromClass() const final;
 
 
 	// Build the signature of the expression by collecting the FTG_ExpressionXXX UProperties of the class
-	virtual FTG_SignaturePtr BuildSignatureFromClass() const final;
+	UE_API virtual FTG_SignaturePtr BuildSignatureFromClass() const final;
 
 	// Build Signature in derived classes dynamically
 	virtual FTG_SignaturePtr BuildSignatureDynamically() const { return nullptr; }
@@ -121,12 +123,12 @@ public:
 	// If some state has changed in the expression that affects its representation
 	// triggered when a property has changed and needs to be copied over to its corresponding Var
 	// NB: required to be public for calling from TG_Variant customization
-	virtual void NotifyExpressionChanged(const FPropertyChangedEvent& PropertyChangedEvent) const final;
+	UE_API virtual void NotifyExpressionChanged(const FPropertyChangedEvent& PropertyChangedEvent) const final;
 protected:
 
 	// If the signature changes and the node need to regenerate its own signature.
 	// Only concrete implementation for Dynamic Expression
-	virtual void NotifySignatureChanged() const;
+	UE_API virtual void NotifySignatureChanged() const;
 
 	// Variant expression API: specialized when using the macro TG_DECLARE_VARIANT_EXPRESSION: 
 	// When input connection are changing, the graph / node calls ResetCommonInputVariantType in order for the expression
@@ -145,13 +147,13 @@ protected:
 	// Default implementation: The common variant type is found as the type that superseeds all the variant input arguments
 	//           currently fed from other nodes 
 	// This call can be modified in a specific class to get a different behavior.
-	virtual FTG_Variant::EType EvalExpressionCommonInputVariantType() const;
+	UE_API virtual FTG_Variant::EType EvalExpressionCommonInputVariantType() const;
 
 	// Variant expression API:
 	// Similarly, this call eval the common output variant type used by the expression (only if it is variant)
 	// Default implementation: the common output variant is the same as the common input variant.
 	// this call can be modified in a specific class to get a differnet behavior, see Expression_Dot for an example.
-	virtual FTG_Variant::EType EvalExpressionCommonOutputVariantType() const;
+	UE_API virtual FTG_Variant::EType EvalExpressionCommonOutputVariantType() const;
 
 	virtual bool IgnoreInputTextureOnUndo() const { return true; }
 
@@ -172,8 +174,8 @@ protected:
 	}
 
 #if WITH_EDITOR
-	virtual void PostEditChangeProperty(FPropertyChangedEvent& PropertyChangedEvent) override;
-	virtual void PostEditUndo() override;
+	UE_API virtual void PostEditChangeProperty(FPropertyChangedEvent& PropertyChangedEvent) override;
+	UE_API virtual void PostEditUndo() override;
 #endif
 };
 
@@ -223,8 +225,8 @@ protected:
 	
 	
 
-UCLASS(Hidden)
-class TEXTUREGRAPH_API UTG_Expression_Null : public UTG_Expression
+UCLASS(MinimalAPI, Hidden)
+class UTG_Expression_Null : public UTG_Expression
 {
 	GENERATED_BODY()
 
@@ -232,3 +234,4 @@ protected:
 	TG_DECLARE_EXPRESSION(TG_Category::Default);
 };
 
+#undef UE_API
diff --git a/Engine/Plugins/TextureGraph/Source/TextureGraph/Public/Expressions/TG_Expression_MaterialBase.h b/Engine/Plugins/TextureGraph/Source/TextureGraph/Public/Expressions/TG_Expression_MaterialBase.h
--- a/Engine/Plugins/TextureGraph/Source/TextureGraph/Public/Expressions/TG_Expression_MaterialBase.h
+++ b/Engine/Plugins/TextureGraph/Source/TextureGraph/Public/Expressions/TG_Expression_MaterialBase.h
@@ -7,6 +7,8 @@
 #include "TG_Texture.h"
 #include "TG_Expression_MaterialBase.generated.h"
 
+#define UE_API TEXTUREGRAPH_API
+
 class RenderMaterial_BP;
 typedef std::shared_ptr<RenderMaterial_BP> RenderMaterial_BPPtr;
 
@@ -33,8 +35,8 @@ enum class EDrawMaterialAttributeTarget : uint8
 
 struct FMaterialInstanceStaticSwitchPermutationMap;
 
-UCLASS(abstract)
-class TEXTUREGRAPH_API UTG_Expression_MaterialBase : public UTG_Expression
+UCLASS(MinimalAPI, abstract)
+class UTG_Expression_MaterialBase : public UTG_Expression
 {
 	GENERATED_BODY()
 
@@ -55,16 +57,16 @@ public:
 
 	mutable TArray<FTG_Texture> Outputs;
 
-	virtual void Evaluate(FTG_EvaluationContext* InContext) override;
+	UE_API virtual void Evaluate(FTG_EvaluationContext* InContext) override;
 
 	// Validate internal checks, warnings and errors
-	virtual bool Validate(MixUpdateCyclePtr	Cycle) override;
+	UE_API virtual bool Validate(MixUpdateCyclePtr	Cycle) override;
 
 	// Access the list of attributes available for rendering from the current Material
 	const TArray<EDrawMaterialAttributeTarget>& GetAvailableMaterialAttributeIds() const { return AvailableMaterialAttributeIds; }
 	const TArray<FName>& GetAvailableMaterialAttributeNames() const { return AvailableMaterialAttributeNames; }
 
-	static FName CPPTypeNameFromMaterialParamType(EMaterialParameterType InMatType);
+	static UE_API FName CPPTypeNameFromMaterialParamType(EMaterialParameterType InMatType);
 
 protected:
 	// A local Material Instance Constant is recreated from the reference material assigned through SetMaterialInternal
@@ -92,27 +94,29 @@ protected:
 
 
 
-	virtual void Initialize() override;
+	UE_API virtual void Initialize() override;
 
-	void GenerateMaterialAttributeOptions(); // Based on the current Material, list of material attributes available
+	UE_API void GenerateMaterialAttributeOptions(); // Based on the current Material, list of material attributes available
 
 	// During evaluation, pick the MIC from the input static switch(es) combination
-	UMaterialInstanceConstant* GetEvaluationMaterialInstanceConstant(FTG_EvaluationContext* InContext);
+	UE_API UMaterialInstanceConstant* GetEvaluationMaterialInstanceConstant(FTG_EvaluationContext* InContext);
 
-	TiledBlobPtr CreateRenderMaterialJob(FTG_EvaluationContext* InContext, const FString& InName, const FString& InMaterialPath, const BufferDescriptor& InDescriptor, EDrawMaterialAttributeTarget InDrawMaterialAttributeTarget);
-	TiledBlobPtr CreateRenderMaterialJob(FTG_EvaluationContext* InContext, const RenderMaterial_BPPtr& InRenderMaterial, const BufferDescriptor& InDescriptor, EDrawMaterialAttributeTarget InDrawMaterialAttributeTarget);
-	void LinkMaterialParameters(FTG_EvaluationContext* InContext, JobUPtr& InMaterialJob, const UMaterialInterface* InMaterial, BufferDescriptor InDescriptor);
+	UE_API TiledBlobPtr CreateRenderMaterialJob(FTG_EvaluationContext* InContext, const FString& InName, const FString& InMaterialPath, const BufferDescriptor& InDescriptor, EDrawMaterialAttributeTarget InDrawMaterialAttributeTarget);
+	UE_API TiledBlobPtr CreateRenderMaterialJob(FTG_EvaluationContext* InContext, const RenderMaterial_BPPtr& InRenderMaterial, const BufferDescriptor& InDescriptor, EDrawMaterialAttributeTarget InDrawMaterialAttributeTarget);
+	UE_API void LinkMaterialParameters(FTG_EvaluationContext* InContext, JobUPtr& InMaterialJob, const UMaterialInterface* InMaterial, BufferDescriptor InDescriptor);
 
-	virtual void CopyVarGeneric(const FTG_Argument& Arg, FTG_Var* InVar, bool CopyVarToArg);
-	virtual void SetMaterialInternal(UMaterialInterface* InMaterial);
+	UE_API virtual void CopyVarGeneric(const FTG_Argument& Arg, FTG_Var* InVar, bool CopyVarToArg);
+	UE_API virtual void SetMaterialInternal(UMaterialInterface* InMaterial);
 
 	virtual EDrawMaterialAttributeTarget GetRenderedAttributeId() { return EDrawMaterialAttributeTarget::Emissive;}
 	virtual TObjectPtr<UMaterialInterface> GetMaterial() const { return nullptr;}
 
 private:
-	void AddSignatureParam(const TArray<FMaterialParameterInfo>& OutParameterInfo, const TArray<FGuid>& OutParameterIds, EMaterialParameterType MatType, FTG_Signature::FInit& SignatureInit) const;
+	UE_API void AddSignatureParam(const TArray<FMaterialParameterInfo>& OutParameterInfo, const TArray<FGuid>& OutParameterIds, EMaterialParameterType MatType, FTG_Signature::FInit& SignatureInit) const;
 
-	static EDrawMaterialAttributeTarget ConvertEMaterialPropertyToEDrawMaterialAttributeTarget(EMaterialProperty InMaterialProperty);
+	static UE_API EDrawMaterialAttributeTarget ConvertEMaterialPropertyToEDrawMaterialAttributeTarget(EMaterialProperty InMaterialProperty);
 };
 
- 
\ No newline at end of file
+ 
+
+#undef UE_API
diff --git a/Engine/Plugins/TextureGraph/Source/TextureGraph/Public/Expressions/TG_Expression_Variant.h b/Engine/Plugins/TextureGraph/Source/TextureGraph/Public/Expressions/TG_Expression_Variant.h
--- a/Engine/Plugins/TextureGraph/Source/TextureGraph/Public/Expressions/TG_Expression_Variant.h
+++ b/Engine/Plugins/TextureGraph/Source/TextureGraph/Public/Expressions/TG_Expression_Variant.h
@@ -12,12 +12,14 @@
 
 #include "TG_Expression_Variant.generated.h"
 
+#define UE_API TEXTUREGRAPH_API
+
 //////////////////////////////////////////////////////////////////////////
 /// An abstract class to help with the boilerplate of handling
 /// variants in general
 //////////////////////////////////////////////////////////////////////////
-UCLASS(Abstract)
-class TEXTUREGRAPH_API UTG_Expression_Variant : public UTG_Expression
+UCLASS(MinimalAPI, Abstract)
+class UTG_Expression_Variant : public UTG_Expression
 {
 	GENERATED_BODY()
 
@@ -35,7 +37,7 @@ protected:
 	/// Evaluate vector only needs to be implemented if there vector handling isn't just a 
 	/// component-wise scalar handling. In our current implementation for instance
 	/// Dot and cross product is a good example where special vector handling is required
-	virtual FVector4f					EvaluateVector_WithValue(FTG_EvaluationContext* InContext, const FVector4f* const Value, size_t Count);
+	UE_API virtual FVector4f					EvaluateVector_WithValue(FTG_EvaluationContext* InContext, const FVector4f* const Value, size_t Count);
 
 	/// Only override this if you don't want use the base class error handling (which is
 	/// quite strict out of the box). Clamp expression is one which does it's own error
@@ -45,9 +47,9 @@ protected:
 	/// This only need to be overridden if want to have your own variant implementation
 	/// along with the EvaluateTexture(...) method defined above. Most of the time 
 	/// overriding these will not be required.
-	virtual float						EvaluateScalar(FTG_EvaluationContext* InContext);
-	virtual FVector4f					EvaluateVector(FTG_EvaluationContext* InContext);
-	virtual FLinearColor				EvaluateColor(FTG_EvaluationContext* InContext);
+	UE_API virtual float						EvaluateScalar(FTG_EvaluationContext* InContext);
+	UE_API virtual FVector4f					EvaluateVector(FTG_EvaluationContext* InContext);
+	UE_API virtual FLinearColor				EvaluateColor(FTG_EvaluationContext* InContext);
 
 public:
 	TG_DECLARE_VARIANT_EXPRESSION(TG_Category::Default);
@@ -55,5 +57,7 @@ public:
 	UPROPERTY(EditAnywhere, Category = NoCategory, meta = (TGType = "TG_Output", PinDisplayName = ""))
 	FTG_Variant							Output;
 
-	virtual void						Evaluate(FTG_EvaluationContext* InContext) override;
+	UE_API virtual void						Evaluate(FTG_EvaluationContext* InContext) override;
 };
+
+#undef UE_API
diff --git a/Engine/Plugins/TextureGraph/Source/TextureGraph/Public/Expressions/Utilities/TG_Expression_MaterialID.h b/Engine/Plugins/TextureGraph/Source/TextureGraph/Public/Expressions/Utilities/TG_Expression_MaterialID.h
--- a/Engine/Plugins/TextureGraph/Source/TextureGraph/Public/Expressions/Utilities/TG_Expression_MaterialID.h
+++ b/Engine/Plugins/TextureGraph/Source/TextureGraph/Public/Expressions/Utilities/TG_Expression_MaterialID.h
@@ -10,13 +10,15 @@
 
 #include "TG_Expression_MaterialID.generated.h"
 
-UCLASS()
-class TEXTUREGRAPH_API UTG_Expression_MaterialID : public UTG_Expression
+#define UE_API TEXTUREGRAPH_API
+
+UCLASS(MinimalAPI)
+class UTG_Expression_MaterialID : public UTG_Expression
 {
 	GENERATED_BODY()
 public:
 	TG_DECLARE_EXPRESSION(TG_Category::Utilities)
-	virtual void Evaluate(FTG_EvaluationContext* InContext) override;
+	UE_API virtual void Evaluate(FTG_EvaluationContext* InContext) override;
 
 	UPROPERTY(meta = (TGType = "TG_Input"))
 	FTG_Texture MaterialIDMap;
@@ -39,19 +41,20 @@ public:
 	FTG_Texture Output;
 	
 protected:
-    virtual void Initialize() override;
+    UE_API virtual void Initialize() override;
 
 private:
-	void UpdateActiveColors();
+	UE_API void UpdateActiveColors();
 public:
 
 #if WITH_EDITOR
-	virtual void PostEditChangeProperty(FPropertyChangedEvent& PropertyChangedEvent) override;
+	UE_API virtual void PostEditChangeProperty(FPropertyChangedEvent& PropertyChangedEvent) override;
 #endif
 	
 	virtual FTG_Name GetDefaultName() const override { return TEXT("MaterialID"); }
 	virtual FText GetTooltipText() const override { return FText::FromString(TEXT("Extracts colors from the texture and creates a mask with selected colors")); }
-	void SetMaterialIdMask(int32 Id, bool bEnable);
+	UE_API void SetMaterialIdMask(int32 Id, bool bEnable);
 	
 };
 
+#undef UE_API
diff --git a/Engine/Plugins/TextureGraph/Source/TextureGraph/Public/Expressions/Utilities/TG_Expression_Utils.h b/Engine/Plugins/TextureGraph/Source/TextureGraph/Public/Expressions/Utilities/TG_Expression_Utils.h
--- a/Engine/Plugins/TextureGraph/Source/TextureGraph/Public/Expressions/Utilities/TG_Expression_Utils.h
+++ b/Engine/Plugins/TextureGraph/Source/TextureGraph/Public/Expressions/Utilities/TG_Expression_Utils.h
@@ -9,15 +9,17 @@
 
 #include "TG_Expression_Utils.generated.h"
 
+#define UE_API TEXTUREGRAPH_API
+
 //////////////////////////////////////////////////////////
 #if 0
-UCLASS()
-class TEXTUREGRAPH_API UTG_Expression_Utils_GetWidth : public UTG_Expression
+UCLASS(MinimalAPI)
+class UTG_Expression_Utils_GetWidth : public UTG_Expression
 {
 	GENERATED_BODY()
 public:
 	TG_DECLARE_EXPRESSION(TG_Category::Utilities)
-	virtual void Evaluate(FTG_EvaluationContext* InContext) override;
+	UE_API virtual void Evaluate(FTG_EvaluationContext* InContext) override;
 
 	UPROPERTY(meta = (TGType = "TG_Input", PinDisplayName = ""))
 	FTG_Texture Input;
@@ -31,13 +33,13 @@ public:
 };
 
 //////////////////////////////////////////////////////////
-UCLASS()
-class TEXTUREGRAPH_API UTG_Expression_Utils_GetHeight : public UTG_Expression
+UCLASS(MinimalAPI)
+class UTG_Expression_Utils_GetHeight : public UTG_Expression
 {
 	GENERATED_BODY()
 public:
 	TG_DECLARE_EXPRESSION(TG_Category::Utilities)
-	virtual void Evaluate(FTG_EvaluationContext* InContext) override;
+	UE_API virtual void Evaluate(FTG_EvaluationContext* InContext) override;
 
 	UPROPERTY(meta = (TGType = "TG_Input", PinDisplayName = ""))
 	FTG_Texture Input;
@@ -51,13 +53,13 @@ public:
 };
 #endif 
 //////////////////////////////////////////////////////////
-UCLASS()
-class TEXTUREGRAPH_API UTG_Expression_Utils_MakeVector4 : public UTG_Expression
+UCLASS(MinimalAPI)
+class UTG_Expression_Utils_MakeVector4 : public UTG_Expression
 {
 	GENERATED_BODY()
 public:
 	TG_DECLARE_EXPRESSION(TG_Category::Utilities)
-	virtual void Evaluate(FTG_EvaluationContext* InContext) override;
+	UE_API virtual void Evaluate(FTG_EvaluationContext* InContext) override;
 
 	UPROPERTY(meta = (TGType = "TG_Input"))
 	float X = 0;
@@ -81,13 +83,13 @@ public:
 };
 
 //////////////////////////////////////////////////////////
-UCLASS()
-class TEXTUREGRAPH_API UTG_Expression_Utils_Resize : public UTG_Expression
+UCLASS(MinimalAPI)
+class UTG_Expression_Utils_Resize : public UTG_Expression
 {
 	GENERATED_BODY()
 public:
 	TG_DECLARE_EXPRESSION(TG_Category::Utilities)
-	virtual void Evaluate(FTG_EvaluationContext* InContext) override;
+	UE_API virtual void Evaluate(FTG_EvaluationContext* InContext) override;
 
 	/// The input image to resize
 	UPROPERTY(meta = (TGType = "TG_Input", PinDisplayName = ""))
@@ -110,3 +112,4 @@ public:
 	virtual FText GetTooltipText() const override { return FText::FromString(TEXT("Resizes an input texture to the given dimensions.")); }
 };
 
+#undef UE_API
diff --git a/Engine/Plugins/TextureGraph/Source/TextureGraph/Public/TG_Graph.h b/Engine/Plugins/TextureGraph/Source/TextureGraph/Public/TG_Graph.h
--- a/Engine/Plugins/TextureGraph/Source/TextureGraph/Public/TG_Graph.h
+++ b/Engine/Plugins/TextureGraph/Source/TextureGraph/Public/TG_Graph.h
@@ -8,6 +8,8 @@
 #include <functional>
 #include "TG_Graph.generated.h"
 
+#define UE_API TEXTUREGRAPH_API
+
 class UTG_ExportSettings;
 
 // Graph Traversal
@@ -38,8 +40,8 @@ DECLARE_MULTICAST_DELEGATE_TwoParams(FOnTGNodeRenamed, UTG_Node*, FName);
 
 #endif
 
-UCLASS()
-class TEXTUREGRAPH_API UTG_Graph : public UObject
+UCLASS(MinimalAPI)
+class UTG_Graph : public UObject
 {
     GENERATED_BODY()
 
@@ -48,23 +50,23 @@ private:
 	// The node is initialized, assigned a uuid in the graph
 	// The node pins are created from the expression signature
 	// Notifiers are installed
-    void SetupNode(UTG_Node* Node);
+    UE_API void SetupNode(UTG_Node* Node);
 
 	// Inner recreate node, used when the node's signature has changed,
 	// All the existing pins of the node are killed
 	// The new signature is obtained from the Expression
 	// the new pins are created
-	void RegenerateNode(UTG_Node* InNode);
+	UE_API void RegenerateNode(UTG_Node* InNode);
 
 	// Inner allocate node pins, used when a node is setup or when a node is regenerated
-	void AllocateNodePins(UTG_Node* Node);
+	UE_API void AllocateNodePins(UTG_Node* Node);
 	// Inner kill node pins, used when a node is removed or when a node is regenerated
-	void KillNodePins(UTG_Node* Node);
+	UE_API void KillNodePins(UTG_Node* Node);
 
 	// Inner allocate pin
-	FTG_Id AllocatePin(UTG_Node* Node, const FTG_Argument& SourceArg, int32 PinIdx = -1);
+	UE_API FTG_Id AllocatePin(UTG_Node* Node, const FTG_Argument& SourceArg, int32 PinIdx = -1);
 	// Inner kill pin
-	void KillPin(FTG_Id InPin);
+	UE_API void KillPin(FTG_Id InPin);
 
 	// The array of nodes indexed by their uuid.NodeIdx()
 	// An element can be null if the node has been removed during the authoring
@@ -77,9 +79,9 @@ private:
 	UPROPERTY(EditAnywhere, Transient, Category = "TextureGraphParams")
 	TMap<FName, FTG_Id> Params; // The Parameter pins in the graph
 
-	void NotifyGraphChanged(UTG_Node* InNode = nullptr, bool bIsTweaking = false);
+	UE_API void NotifyGraphChanged(UTG_Node* InNode = nullptr, bool bIsTweaking = false);
 	friend struct FTG_Evaluation; // NodePostEvaluate notifier is triggered from the FTG_Evaluation::EvaluateNode call
-	void NotifyNodePostEvaluate(UTG_Node* InNode, const FTG_EvaluationContext* InContext);
+	UE_API void NotifyNodePostEvaluate(UTG_Node* InNode, const FTG_EvaluationContext* InContext);
 
 	bool IsValidNode(FTG_Id NodeId) const { return (NodeId != FTG_Id::INVALID && Nodes.IsValidIndex(NodeId.NodeIdx())); }
 	bool IsValidPin(FTG_Id PinId) const { return IsValidNode(PinId) && Nodes[PinId.NodeIdx()] && Nodes[PinId.NodeIdx()]->Pins.IsValidIndex(PinId.PinIdx()); }
@@ -87,26 +89,26 @@ private:
 protected:
 	friend class UTG_Node;
 	friend class UTG_Pin;
-	void OnNodeChanged(UTG_Node* InNode, bool Tweaking);
-	void OnNodeSignatureChanged(UTG_Node* InNode);
-	void OnNodePinChanged(FTG_Id InPinId, UTG_Node* InNode);
-	void OnNodeAdded(UTG_Node* InNode);
-	void OnNodeRemoved(UTG_Node* InNode,FName InName);
-	void OnNodeRenamed(UTG_Node* InNode,FName OldName);
+	UE_API void OnNodeChanged(UTG_Node* InNode, bool Tweaking);
+	UE_API void OnNodeSignatureChanged(UTG_Node* InNode);
+	UE_API void OnNodePinChanged(FTG_Id InPinId, UTG_Node* InNode);
+	UE_API void OnNodeAdded(UTG_Node* InNode);
+	UE_API void OnNodeRemoved(UTG_Node* InNode,FName InName);
+	UE_API void OnNodeRenamed(UTG_Node* InNode,FName OldName);
 
 	// Rename existing param
-	bool RenameParam(FName OldName, FName NewName);
+	UE_API bool RenameParam(FName OldName, FName NewName);
 
 #if WITH_EDITOR
 protected:
-	virtual void PostEditChangeProperty(FPropertyChangedEvent& PropertyChangedEvent) override;
+	UE_API virtual void PostEditChangeProperty(FPropertyChangedEvent& PropertyChangedEvent) override;
 
 	// Override PostEditUndo method of UObject
-	virtual void PostEditUndo() override;
+	UE_API virtual void PostEditUndo() override;
 
 public:
 	// Overwrite the modify to also mark the transient states as dirty
-	virtual bool Modify(bool bAlwaysMarkDirty = true) override;
+	UE_API virtual bool Modify(bool bAlwaysMarkDirty = true) override;
 #endif
 
 public:
@@ -115,17 +117,17 @@ public:
 	//////////////////////////////////////////////////////////////////////////
 
 	// Construct the Graph, can only be done once before "using" the graph
-	virtual void Construct(FString InName);
+	UE_API virtual void Construct(FString InName);
 
 	// Override Serialize method of UObject
-	virtual void Serialize(FArchive& Ar) override;
+	UE_API virtual void Serialize(FArchive& Ar) override;
 
     // Override the PostLoad method of UObject to finalize the un-serialization
     // and fill in the runtime fields of the graph sub objects
-	virtual void PostLoad() override;
+	UE_API virtual void PostLoad() override;
 
 	// Override PreSave method of UObject
-	virtual void PreSave(FObjectPreSaveContext SaveContext) override;
+	UE_API virtual void PreSave(FObjectPreSaveContext SaveContext) override;
 
 	//////////////////////////////////////////////////////////////////////////
 	//// Graph methods to build the graph!
@@ -135,13 +137,13 @@ public:
 	// A new Expression instance is created along with the Node instance
 	// The Node is setup with its input and output pins reflecting
 	// the Signature of the Expression.
-    UTG_Node* CreateExpressionNode(const UClass* ExpressionClass);
+    UE_API UTG_Node* CreateExpressionNode(const UClass* ExpressionClass);
 
 	// Create a node from an initialized Expression
 	// A new Node instance is created which uses initialized expression
 	// The Node is setup with its input and output pins reflecting
 	// the Signature of the Expression.
-    UTG_Node* CreateExpressionNode(UTG_Expression* NewExpression);
+    UE_API UTG_Node* CreateExpressionNode(UTG_Expression* NewExpression);
 	
 	// Remove the node
 	// The node is destroyed, its Uuid becomes invalid
@@ -150,36 +152,36 @@ public:
 	// As a consequence, each pin's
 	//     associated connection Edge(s) are destroyed
 	//	   associated Var is dereferenced and potentially GCed
-	void RemoveNode(UTG_Node* InNode);
+	UE_API void RemoveNode(UTG_Node* InNode);
 
 	// Accessed from TG_EdGraphNode after a paste of a new node
-	void AddPostPasteNode(UTG_Node* NewNode);
+	UE_API void AddPostPasteNode(UTG_Node* NewNode);
 
     // Connect a output pin from a node to the specified input pin of the to node
     // Return true if the edge is created or false if couldn't create the edge
-	bool Connect(UTG_Node& NodeFrom, FTG_Name& PinFrom, UTG_Node& NodeTo, FTG_Name& PinTo);
+	UE_API bool Connect(UTG_Node& NodeFrom, FTG_Name& PinFrom, UTG_Node& NodeTo, FTG_Name& PinTo);
 	
     // Check that the connection from PinFrom to PInTO doesn't create a loop
-	static bool ConnectionCausesLoop(const UTG_Pin* PinFrom, const UTG_Pin* PinTo);
+	static UE_API bool ConnectionCausesLoop(const UTG_Pin* PinFrom, const UTG_Pin* PinTo);
 
 	// Are 2 Pins compatible ?
 	// This function check that the 2 pins can be connected to each other in terms of their argument types
 	// A non empty ConverterKey is also returned if required to achieve a conversion of the pin A to the pin B
 	// The converterKey is then used during the evaluation to find the proper Converter
-	static bool ArePinsCompatible(const UTG_Pin* PinFrom, const UTG_Pin* PinTo, FName& ConverterKey);
+	static UE_API bool ArePinsCompatible(const UTG_Pin* PinFrom, const UTG_Pin* PinTo, FName& ConverterKey);
 
 	// Remove all the edges connected to the specified Pin
-	void RemovePinEdges(UTG_Node& InNode, FTG_Name& InPin);
+	UE_API void RemovePinEdges(UTG_Node& InNode, FTG_Name& InPin);
 
 	// Remove a particular edge
-	void RemoveEdge(UTG_Node& NodeFrom, FTG_Name& PinFrom, UTG_Node& NodeTo, FTG_Name& PinTo);
+	UE_API void RemoveEdge(UTG_Node& NodeFrom, FTG_Name& PinFrom, UTG_Node& NodeTo, FTG_Name& PinTo);
 
 	// Reset the graph empty, destroy any nodes or associated resources
-	void Reset();
+	UE_API void Reset();
 
 	// Create a Signature from the current graph topology declaring the current Params as Arguments
 	// The signature list the input and output pins of the graph
-	void AppendParamsSignature(FTG_Arguments& InOutArguments, TArray<FTG_Id>& InParams, TArray<FTG_Id>& OutParams) const;
+	UE_API void AppendParamsSignature(FTG_Arguments& InOutArguments, TArray<FTG_Id>& InParams, TArray<FTG_Id>& OutParams) const;
 
 	//////////////////////////////////////////////////////////////////////////
 	// Accessors for nodes, pins, params
@@ -206,10 +208,10 @@ public:
 	FTG_Ids			GetParamIds() const { FTG_Ids Array;  Params.GenerateValueArray(Array); return Array; }
 	TArray<FName>	GetParamNames() const { TArray<FName> Array;  Params.GenerateKeyArray(Array); return Array; }
 
-	FName			GetParamName(const FTG_Id& InParamId) const;
+	UE_API FName			GetParamName(const FTG_Id& InParamId) const;
 
-	FTG_Ids			GetInputParamIds() const;
-	FTG_Ids			GetOutputParamIds() const;
+	UE_API FTG_Ids			GetInputParamIds() const;
+	UE_API FTG_Ids			GetOutputParamIds() const;
 
 	// Find a param pin from its name
 	FTG_Id			FindParamPinId(const FName& InName) const		{ auto* PinId = Params.Find(InName); return (PinId ? (*PinId) : FTG_Id()); }
@@ -217,25 +219,25 @@ public:
 	UTG_Pin*		FindParamPin(const FName& InName)				{ return GetPin(FindParamPinId(InName)); }
 
 	// Get the names and ids of the output params which are Textures
-	int				GetOutputParamTextures(TArray<FName>& OutNames, FTG_Ids& OutPinIds) const;
+	UE_API int				GetOutputParamTextures(TArray<FName>& OutNames, FTG_Ids& OutPinIds) const;
 
 	// Iterate through all the VALID nodes
-    void ForEachNodes(std::function<void(const UTG_Node* /*node*/, uint32 /*index*/)> visitor) const;
+    UE_API void ForEachNodes(std::function<void(const UTG_Node* /*node*/, uint32 /*index*/)> visitor) const;
 
 	// Iterate through all the VALID pins
-	void ForEachPins(std::function<void(const UTG_Pin* /*pin*/, uint32 /*index*/, uint32 /*node_index*/)> visitor) const;
+	UE_API void ForEachPins(std::function<void(const UTG_Pin* /*pin*/, uint32 /*index*/, uint32 /*node_index*/)> visitor) const;
 
 	// Iterate through all the VALID vars
-	void ForEachVars(std::function<void(const FTG_Var* /*var*/, uint32 /*index*/, uint32 /*node_index*/)> visitor) const;
+	UE_API void ForEachVars(std::function<void(const FTG_Var* /*var*/, uint32 /*index*/, uint32 /*node_index*/)> visitor) const;
 
 	// Iterate through all the VALID Param pins
-	void ForEachParams(std::function<void(const UTG_Pin* /*pin*/, uint32 /*index*/)> visitor) const;
+	UE_API void ForEachParams(std::function<void(const UTG_Pin* /*pin*/, uint32 /*index*/)> visitor) const;
 
 	// Iterate through all the VALID edges
-	void ForEachEdges(std::function<void(const UTG_Pin* /*pinFrom*/, const UTG_Pin* /*pinTo*/)> visitor) const;
+	UE_API void ForEachEdges(std::function<void(const UTG_Pin* /*pinFrom*/, const UTG_Pin* /*pinTo*/)> visitor) const;
 
 	// Iterate through all the output settings
-	void ForEachOutputSettings( std::function<void(const FTG_OutputSettings& /*settings*/)> visitor);
+	UE_API void ForEachOutputSettings( std::function<void(const FTG_OutputSettings& /*settings*/)> visitor);
 
 	//////////////////////////////////////////////////////////////////////////
 	// Accessors for output param values after evaluation
@@ -245,23 +247,23 @@ public:
 	// Get the nammed <InName> output param value as FTG_Variant if it exists.
 	// The value is collected in the <OutVariant> parameter if found.
 	// return true if the correct nammed & typed param was found, false otherwise
-	bool GetOutputParamValue(const FName& InName, FTG_Variant& OutVariant) const;
+	UE_API bool GetOutputParamValue(const FName& InName, FTG_Variant& OutVariant) const;
 
 	// Get all the output param's value as FTG_Variant
 	// along with their names optionally
-	int GetAllOutputParamValues(TArray<FTG_Variant>& OutVariants, TArray<FName>* OutNames = nullptr) const;
+	UE_API int GetAllOutputParamValues(TArray<FTG_Variant>& OutVariants, TArray<FName>* OutNames = nullptr) const;
 
 	// Get a map of input parameter vars from the graph
-	FTG_VarMap GetInputParamsVarMap();
+	UE_API FTG_VarMap GetInputParamsVarMap();
 
 	// Apply an override map of input parameter vars to the graph
-    void SetInputParamsFromVarMap(FTG_VarMap InMap);
+    UE_API void SetInputParamsFromVarMap(FTG_VarMap InMap);
 
 	// get a map of output settings from the graph
-	void CollectOutputSettings(TMap<FTG_Id, FTG_OutputSettings>& OutExportSettingsMap) const;
+	UE_API void CollectOutputSettings(TMap<FTG_Id, FTG_OutputSettings>& OutExportSettingsMap) const;
 
 	// Apply an override map of output settings to the graph
-	void SetOutputSettings(const TMap<FTG_Id, FTG_OutputSettings>& OutputSettingsMap);
+	UE_API void SetOutputSettings(const TMap<FTG_Id, FTG_OutputSettings>& OutputSettingsMap);
 
     //////////////////////////////////////////////////////////////////////////
 	// Traverseal of the graph and evaluation 
@@ -270,19 +272,19 @@ public:
 	// Useful function to navigate the Node
 	// Collect all the DIRECT nodes whose output pins which are
 	// connected to the input pins of this node
-    TArray<FTG_Id> GatherSourceNodes(const UTG_Node* Node) const;
+    UE_API TArray<FTG_Id> GatherSourceNodes(const UTG_Node* Node) const;
 
 	// Useful function to navigate the Nodes
 	// Collect all the nodes which are connecting directly or indirectly into the specified node
-	TArray<FTG_Id> GatherAllSourceNodes(const UTG_Node* Node) const;
+	UE_API TArray<FTG_Id> GatherAllSourceNodes(const UTG_Node* Node) const;
 
 	//// Graph methods to traverse the graph!
 	
 	// Validate internal checks, warnings and errors, Returns true if Graph is valid and has no validation errors
-	bool Validate(MixUpdateCyclePtr	Cycle);
+	UE_API bool Validate(MixUpdateCyclePtr	Cycle);
 
 	// Access the traversal  (update it if needed)
-	const FTG_GraphTraversal& GetTraversal() const;
+	UE_API const FTG_GraphTraversal& GetTraversal() const;
 
     // Traversal Visitor function type
     // Node visited
@@ -292,24 +294,24 @@ public:
     using NodeVisitorFunction = std::function<void (UTG_Node*, int32, int32)>;
 
     // Traverse!
-    void Traverse(NodeVisitorFunction visitor, int32 graph_depth = 0) const;
+    UE_API void Traverse(NodeVisitorFunction visitor, int32 graph_depth = 0) const;
 
 	// Evaluate the graph
 	// Traverse every nodes in the traversal order and call the node's expression evaluate method
 	// A configured evaluation context with a valid Cycle is required
 	// Evaluation is concretely implemented in FTG_Evaluation struct
-	void Evaluate(FTG_EvaluationContext* InContext);
+	UE_API void Evaluate(FTG_EvaluationContext* InContext);
 
 #if WITH_EDITOR
 	const TArray<TObjectPtr<UObject>>& GetExtraEditorNodes() const { return ExtraEditorNodes; }
-	void SetExtraEditorNodes(const TArray<TObjectPtr<const UObject>>& InNodes);
+	UE_API void SetExtraEditorNodes(const TArray<TObjectPtr<const UObject>>& InNodes);
 #endif
 	
 protected:
 	// Utility functions used to establish the traversal order
-	void GatherOuterNodes(const TArray<FTG_Ids>& sourceNodesPerNode, TSet<FTG_Id>& nodeReservoirA, TSet<FTG_Id>& nodeReservoirB) const;
-	void EvalInOutPins() const;
-	void EvalTraverseOrder() const;
+	UE_API void GatherOuterNodes(const TArray<FTG_Ids>& sourceNodesPerNode, TSet<FTG_Id>& nodeReservoirA, TSet<FTG_Id>& nodeReservoirB) const;
+	UE_API void EvalInOutPins() const;
+	UE_API void EvalTraverseOrder() const;
 
 #if WITH_EDITORONLY_DATA
 	// Extra data to hold information that is useful only in editor (like comments)
@@ -338,14 +340,14 @@ public:
 	//// Logging Utilities
 
 	// Log the graph description to the log output buffer
-	void Log() const;
+	UE_API void Log() const;
 
 	static constexpr int32 LogHeaderWidth = 32; // width in number of chars for any header token during logging
-	FString LogNodes(FString InTab = TEXT("")) const;
-	FString LogParams(FString InTab = TEXT("")) const;
-	FString LogTraversal(FString InTab = TEXT("")) const;
-	FString LogVars(FString InTab = TEXT("")) const;
-	static FString LogCall(const TArray<FTG_Id>& PinInputs, const  TArray<FTG_Id>& PinOutputs, int32 InputLogWidth = 10);
+	UE_API FString LogNodes(FString InTab = TEXT("")) const;
+	UE_API FString LogParams(FString InTab = TEXT("")) const;
+	UE_API FString LogTraversal(FString InTab = TEXT("")) const;
+	UE_API FString LogVars(FString InTab = TEXT("")) const;
+	static UE_API FString LogCall(const TArray<FTG_Id>& PinInputs, const  TArray<FTG_Id>& PinOutputs, int32 InputLogWidth = 10);
 
 	//// Testing zone
 	
@@ -361,3 +363,4 @@ private:
 
 
 
+#undef UE_API
diff --git a/Engine/Plugins/TextureGraph/Source/TextureGraph/Public/TG_GraphEvaluation.h b/Engine/Plugins/TextureGraph/Source/TextureGraph/Public/TG_GraphEvaluation.h
--- a/Engine/Plugins/TextureGraph/Source/TextureGraph/Public/TG_GraphEvaluation.h
+++ b/Engine/Plugins/TextureGraph/Source/TextureGraph/Public/TG_GraphEvaluation.h
@@ -5,11 +5,13 @@
 #include "Model/Mix/MixInterface.h"
 #include "TG_Var.h"
 #include "TG_GraphEvaluation.generated.h"
+
+#define UE_API TEXTUREGRAPH_API
 class UTG_Graph;
 class UTG_Node;
 	
 USTRUCT()
-struct TEXTUREGRAPH_API FVarArgument
+struct FVarArgument
 {
 	GENERATED_BODY()
 
@@ -40,9 +42,9 @@ struct TStructOpsTypeTraits<FVarArgument> : public TStructOpsTypeTraitsBase2<FVa
 	};
 };
 
-struct TEXTUREGRAPH_API FTG_OutputVarMap
+struct FTG_OutputVarMap
 {
-	struct TEXTUREGRAPH_API FOutputVarArgument
+	struct FOutputVarArgument
 	{
 		FTG_Var* Var;
 		FTG_Argument Argument;
@@ -81,7 +83,7 @@ struct TEXTUREGRAPH_API FTG_OutputVarMap
 };
 
 USTRUCT()
-struct TEXTUREGRAPH_API FTG_VarMap
+struct FTG_VarMap
 {
 	GENERATED_BODY()
 
@@ -119,7 +121,7 @@ struct TEXTUREGRAPH_API FTG_VarMap
 };
 
 // 
-struct TEXTUREGRAPH_API FTG_EvaluationContext
+struct FTG_EvaluationContext
 {
 	// The cycle used to call in the concrete TextureGraph engine system
 	MixUpdateCyclePtr	Cycle;
@@ -152,19 +154,19 @@ struct TEXTUREGRAPH_API FTG_EvaluationContext
 
 class UTG_Pin;
 class UTG_Node;
-struct TEXTUREGRAPH_API FTG_Evaluation
+struct FTG_Evaluation
 {
-	static const FString GVectorToTextureAutoConv_Name;
-	static const FString GColorToTextureAutoConv_Name;
-	static const FString GFloatToTextureAutoConv_Name;
+	static UE_API const FString GVectorToTextureAutoConv_Name;
+	static UE_API const FString GColorToTextureAutoConv_Name;
+	static UE_API const FString GFloatToTextureAutoConv_Name;
 
-	static void EvaluateGraph(UTG_Graph* InGraph, FTG_EvaluationContext* Context);
-	static void EvaluateNodeArray(UTG_Node* InNode, const TArray<TObjectPtr<UTG_Pin>>& ArrayInputs, 
+	static UE_API void EvaluateGraph(UTG_Graph* InGraph, FTG_EvaluationContext* Context);
+	static UE_API void EvaluateNodeArray(UTG_Node* InNode, const TArray<TObjectPtr<UTG_Pin>>& ArrayInputs, 
 		const TArray<TObjectPtr<UTG_Pin>>& NonArrayPins, int MaxCount, FTG_EvaluationContext* Context);
-	static void EvaluateNode(UTG_Node* InNode, FTG_EvaluationContext* Context);
+	static UE_API void EvaluateNode(UTG_Node* InNode, FTG_EvaluationContext* Context);
 
-	static void TransferVarToPin(UTG_Pin* InPin, FTG_EvaluationContext* Context, int Index);
-	static int FilterArrayInputs(FTG_EvaluationContext* InContext, const TArray<TObjectPtr<UTG_Pin>>& InPins, 
+	static UE_API void TransferVarToPin(UTG_Pin* InPin, FTG_EvaluationContext* Context, int Index);
+	static UE_API int FilterArrayInputs(FTG_EvaluationContext* InContext, const TArray<TObjectPtr<UTG_Pin>>& InPins, 
 		TArray<TObjectPtr<UTG_Pin>>& ArrayPins, TArray<TObjectPtr<UTG_Pin>>& NonArrayPins);
 
 	// Converter used to introduce transformation from one var to another var when the arguments are compatible but require to be converted
@@ -178,9 +180,9 @@ struct TEXTUREGRAPH_API FTG_Evaluation
 	typedef std::function<void(VarConverterInfo& Info)> VarConverter;
 
 	typedef TMap				<FName, VarConverter> ConverterMap;
-	static ConverterMap			DefaultConverters;
-	static FName				MakeConvertKey(FName From, FName To);
-	static FName				MakeConvertKey(const FTG_Argument& ArgFrom, const FTG_Argument& ArgTo);
+	static UE_API ConverterMap			DefaultConverters;
+	static UE_API FName				MakeConvertKey(FName From, FName To);
+	static UE_API FName				MakeConvertKey(const FTG_Argument& ArgFrom, const FTG_Argument& ArgTo);
 
 	// Conformer functors used to conform pin values
 	struct VarConformerInfo
@@ -192,7 +194,8 @@ struct TEXTUREGRAPH_API FTG_Evaluation
     };
     typedef std::function<bool(VarConformerInfo& Info)> VarConformer;
 	
-	static bool AreArgumentsCompatible(const FTG_Argument& ArgFrom, const FTG_Argument& ArgTo, FName& ConverterKey);
+	static UE_API bool AreArgumentsCompatible(const FTG_Argument& ArgFrom, const FTG_Argument& ArgTo, FName& ConverterKey);
 
 };
 
+#undef UE_API
diff --git a/Engine/Plugins/TextureGraph/Source/TextureGraph/Public/TG_HelperFunctions.h b/Engine/Plugins/TextureGraph/Source/TextureGraph/Public/TG_HelperFunctions.h
--- a/Engine/Plugins/TextureGraph/Source/TextureGraph/Public/TG_HelperFunctions.h
+++ b/Engine/Plugins/TextureGraph/Source/TextureGraph/Public/TG_HelperFunctions.h
@@ -9,22 +9,24 @@
 #include "TG_Graph.h"
 #include "2D/TextureHelper.h"
 
+#define UE_API TEXTUREGRAPH_API
+
 class UTextureGraphBase;
 
-class TEXTUREGRAPH_API FTG_HelperFunctions
+class FTG_HelperFunctions
 {
 
 public:
-	static void	InitTargets(UTextureGraphBase* InTextureGraph);
-	static TArray<BlobPtr> GetTexturedOutputs(const UTG_Node* Node, FTG_EvaluationContext* TextureConversionContext = nullptr);
+	static UE_API void	InitTargets(UTextureGraphBase* InTextureGraph);
+	static UE_API TArray<BlobPtr> GetTexturedOutputs(const UTG_Node* Node, FTG_EvaluationContext* TextureConversionContext = nullptr);
 
-	static void EnsureOutputIsTexture(MixUpdateCyclePtr Cycle, UTG_Node* OutputNode);
+	static UE_API void EnsureOutputIsTexture(MixUpdateCyclePtr Cycle, UTG_Node* OutputNode);
 
-	static JobBatchPtr InitExportBatch(UTextureGraphBase* InTextureGraph, FString ExportPath, FString AssetName, FExportSettings& TargetExportSettings, bool OverrideExportPath, bool OverwriteTextures, bool ExportAllOutputs, bool bSave);
-	static AsyncBool ExportAsync(UTextureGraphBase* InTextureGraph, FString ExportPath, FString AssetName, FExportSettings& TargetExportSettings, bool OverrideExportPath, bool OverwriteTextures = true, bool ExportAllOutputs = false, bool bSave = true);
+	static UE_API JobBatchPtr InitExportBatch(UTextureGraphBase* InTextureGraph, FString ExportPath, FString AssetName, FExportSettings& TargetExportSettings, bool OverrideExportPath, bool OverwriteTextures, bool ExportAllOutputs, bool bSave);
+	static UE_API AsyncBool ExportAsync(UTextureGraphBase* InTextureGraph, FString ExportPath, FString AssetName, FExportSettings& TargetExportSettings, bool OverrideExportPath, bool OverwriteTextures = true, bool ExportAllOutputs = false, bool bSave = true);
 
-	static JobBatchPtr InitRenderBatch(UTextureGraphBase* InTextureGraph, JobBatchPtr ExistingBatch = nullptr);
-	static AsyncBool RenderAsync(UTextureGraphBase* InTextureGraph, JobBatchPtr ExistingBatch = nullptr);
+	static UE_API JobBatchPtr InitRenderBatch(UTextureGraphBase* InTextureGraph, JobBatchPtr ExistingBatch = nullptr);
+	static UE_API AsyncBool RenderAsync(UTextureGraphBase* InTextureGraph, JobBatchPtr ExistingBatch = nullptr);
 
 	template <typename T_Type>
 	static TArray<T_Type> GetOutputsOfType(const UTG_Node* Node)
@@ -50,3 +52,4 @@ public:
 	}
 };
 
+#undef UE_API
diff --git a/Engine/Plugins/TextureGraph/Source/TextureGraph/Public/TG_Material.h b/Engine/Plugins/TextureGraph/Source/TextureGraph/Public/TG_Material.h
--- a/Engine/Plugins/TextureGraph/Source/TextureGraph/Public/TG_Material.h
+++ b/Engine/Plugins/TextureGraph/Source/TextureGraph/Public/TG_Material.h
@@ -6,8 +6,10 @@
 #include "Misc/OutputDeviceNull.h"
 #include "TG_Material.generated.h"
 
+#define UE_API TEXTUREGRAPH_API
+
 USTRUCT(BlueprintType)
-struct TEXTUREGRAPH_API FTG_Material
+struct FTG_Material
 {
 	GENERATED_USTRUCT_BODY()
 	
@@ -21,19 +23,19 @@ public:
 
 	// Check that the referenced asset is valid.
 	// If true, the GetMaterial call will NOT return nullptr
-	bool IsValid() const;
+	UE_API bool IsValid() const;
 
 	// Access the referenced Material asset
 	// return a valid MaterialInterface object or null if the reference is invalid.
-	UMaterialInterface* GetMaterial() const;
+	UE_API UMaterialInterface* GetMaterial() const;
 
 	// Assign the referenced asset from an actual UMaterialInsterface* live object
-	void SetMaterial(UMaterialInterface* InMaterial);
+	UE_API void SetMaterial(UMaterialInterface* InMaterial);
 
-	bool Serialize(FArchive& Ar);
+	UE_API bool Serialize(FArchive& Ar);
 	friend FArchive& operator<<(FArchive& Ar, FTG_Material& InMaterial);
 	bool operator==(const FTG_Material& RHS) const = default;
-	void ResetTexturePath();
+	UE_API void ResetTexturePath();
 
 	void InitFromString(const FString& StrVal)
 	{
@@ -62,3 +64,4 @@ struct TStructOpsTypeTraits<FTG_Material>
 	};
 };
 
+#undef UE_API
diff --git a/Engine/Plugins/TextureGraph/Source/TextureGraph/Public/TG_Node.h b/Engine/Plugins/TextureGraph/Source/TextureGraph/Public/TG_Node.h
--- a/Engine/Plugins/TextureGraph/Source/TextureGraph/Public/TG_Node.h
+++ b/Engine/Plugins/TextureGraph/Source/TextureGraph/Public/TG_Node.h
@@ -9,6 +9,8 @@
 
 #include "TG_Node.generated.h"
 
+#define UE_API TEXTUREGRAPH_API
+
 class UTG_Graph;
 class UTG_Expression;
 
@@ -38,8 +40,8 @@ struct FTG_NodeEditorData
 //
 // UTG_Node is created from a UTG_Graph
 //
-UCLASS()
-class TEXTUREGRAPH_API UTG_Node : public UObject
+UCLASS(MinimalAPI)
+class UTG_Node : public UObject
 {
     GENERATED_BODY()
 
@@ -55,30 +57,30 @@ class TEXTUREGRAPH_API UTG_Node : public UObject
 	FTG_Id Id; // invalid by default
 
 private:
-    void Construct(UTG_Expression* InExpression);
+    UE_API void Construct(UTG_Expression* InExpression);
 
 	// Initialize the node after being created from the graph or unserialized
 	// This is when the Graph assign itself and the Id to the node
 	// The expression member must be valid:
 	//		- assigned with the Construct call above for a brand new node
 	//		- or recovered from serialization
-	void Initialize(FTG_Id Id);
+	UE_API void Initialize(FTG_Id Id);
 
 	// Evaluate the arguments array deduced from the current array of Pins and the arguments
 	// This is used to compare the validity of the PIn's against the Expression signature
-	FTG_Arguments GetPinArguments() const;
+	UE_API FTG_Arguments GetPinArguments() const;
 
 	// Check that the node signature deduced from the pin's array is the same as the signature from the Expression
 	// return false if it is not the case, true otherwise
-	bool CheckPinSignatureAgainstExpression() const;
+	UE_API bool CheckPinSignatureAgainstExpression() const;
 
 public:
 
 	// Override Serialize method of UObject
-	virtual void Serialize(FArchive& Ar) override;
+	UE_API virtual void Serialize(FArchive& Ar) override;
 
 	// The Graph in which this node exists
-	UTG_Graph* GetGraph() const;
+	UE_API UTG_Graph* GetGraph() const;
 
 	// The expression owned by this node
 	UTG_Expression* GetExpression() const { return Expression; }
@@ -87,16 +89,16 @@ public:
 	FTG_Id GetId() const { return Id; }
 
 	// Validate internal checks, warnings and errors
-    virtual bool Validate(MixUpdateCyclePtr	Cycle);
+    UE_API virtual bool Validate(MixUpdateCyclePtr	Cycle);
 
 	// Node name is the title name of the inner expression
-	FTG_Name GetNodeName() const;
+	UE_API FTG_Name GetNodeName() const;
 
 	//////////////////////////////////////////////////////////////////////////
 	// Signature of the node and predicates about the pins of the node
 	//////////////////////////////////////////////////////////////////////////
 
-	const FTG_Signature& GetSignature() const;
+	UE_API const FTG_Signature& GetSignature() const;
 
 	bool HasInputs() const { return GetSignature().HasInputs(); }
 	bool HasOutputs() const { return GetSignature().HasOutputs(); }
@@ -115,43 +117,43 @@ public:
 	TArray<TObjectPtr<UTG_Pin>> Pins;
 
 	// Iterate through all the Input OR Output pins in this node, return the number traversed
-	int32 ForEachInputPins(std::function<void(const UTG_Pin* /*pin*/, uint32 /*index*/)> visitor) const;
-	int32 ForEachOutputPins(std::function<void(const UTG_Pin* /*pin*/, uint32 /*index*/)> visitor) const;
+	UE_API int32 ForEachInputPins(std::function<void(const UTG_Pin* /*pin*/, uint32 /*index*/)> visitor) const;
+	UE_API int32 ForEachOutputPins(std::function<void(const UTG_Pin* /*pin*/, uint32 /*index*/)> visitor) const;
 
 	// Collect the input / output pins
-	int32 GetInputPins(TArray<const UTG_Pin*>& OutPins) const;
-	int32 GetOutputPins(TArray<const UTG_Pin*>& OutPins) const;
-	int32 GetInputPins(TArray<UTG_Pin*>& OutPins) const;
-	int32 GetOutputPins(TArray<UTG_Pin*>& OutPins) const;
+	UE_API int32 GetInputPins(TArray<const UTG_Pin*>& OutPins) const;
+	UE_API int32 GetOutputPins(TArray<const UTG_Pin*>& OutPins) const;
+	UE_API int32 GetInputPins(TArray<UTG_Pin*>& OutPins) const;
+	UE_API int32 GetOutputPins(TArray<UTG_Pin*>& OutPins) const;
 
 	// Get pin Id from index in each scope
-	FTG_Id GetInputPinIdAt(FTG_Index inIndex) const;
-	FTG_Id GetOutputPinIdAt(FTG_Index outIndex) const;
-	FTG_Id GetPrivatePinIdAt(FTG_Index privateIndex) const;
+	UE_API FTG_Id GetInputPinIdAt(FTG_Index inIndex) const;
+	UE_API FTG_Id GetOutputPinIdAt(FTG_Index outIndex) const;
+	UE_API FTG_Id GetPrivatePinIdAt(FTG_Index privateIndex) const;
 
-	TArray<FTG_Id> GetInputPinIds() const;
-	TArray<FTG_Id> GetOutputPinIds() const;
-	TArray<FTG_Id> GetPrivatePinIds() const;
+	UE_API TArray<FTG_Id> GetInputPinIds() const;
+	UE_API TArray<FTG_Id> GetOutputPinIds() const;
+	UE_API TArray<FTG_Id> GetPrivatePinIds() const;
 
 	// Get/Find the Pin Id for Pins of this Node from its Argument Name
-	FTG_Id GetPinId(FName Name) const;
-	FTG_Id GetInputPinID(FTG_Name& Name) const;
-	FTG_Id GetOutputPinID(FTG_Name& Name) const;
-	FTG_Id GetPrivatePinID(FTG_Name& Name) const;
+	UE_API FTG_Id GetPinId(FName Name) const;
+	UE_API FTG_Id GetInputPinID(FTG_Name& Name) const;
+	UE_API FTG_Id GetOutputPinID(FTG_Name& Name) const;
+	UE_API FTG_Id GetPrivatePinID(FTG_Name& Name) const;
 
-    UTG_Pin* GetInputPinAt(FTG_Index inIndex) const;
-	UTG_Pin* GetOutputPinAt(FTG_Index outIndex) const;
-	UTG_Pin* GetPrivatePinAt(FTG_Index privateIndex) const;
+    UE_API UTG_Pin* GetInputPinAt(FTG_Index inIndex) const;
+	UE_API UTG_Pin* GetOutputPinAt(FTG_Index outIndex) const;
+	UE_API UTG_Pin* GetPrivatePinAt(FTG_Index privateIndex) const;
 
-	UTG_Pin* GetPin(FName Name) const;
-	UTG_Pin* GetInputPin(FTG_Name& name) const;
-	UTG_Pin* GetOutputPin(FTG_Name& name) const;
-	UTG_Pin* GetPrivatePin(FTG_Name& name) const;
+	UE_API UTG_Pin* GetPin(FName Name) const;
+	UE_API UTG_Pin* GetInputPin(FTG_Name& name) const;
+	UE_API UTG_Pin* GetOutputPin(FTG_Name& name) const;
+	UE_API UTG_Pin* GetPrivatePin(FTG_Name& name) const;
 
-	UTG_Pin* GetPin(FTG_Id pinID) const;
+	UE_API UTG_Pin* GetPin(FTG_Id pinID) const;
 
-	TArray<FTG_Id> GetInputVarIds() const;
-	TArray<FTG_Id> GetOutputVarIds() const;
+	UE_API TArray<FTG_Id> GetInputVarIds() const;
+	UE_API TArray<FTG_Id> GetOutputVarIds() const;
 
 	//////////////////////////////////////////////////////////////////////////
 	// Accessors for output values after evaluation
@@ -160,24 +162,24 @@ public:
 
 	// Get all the output's values which are compatible with FTG_Variant
 	// along with their names optionally
-	int GetAllOutputValues(TArray<FTG_Variant>& OutVariants, TArray<FName>* OutNames = nullptr) const;
+	UE_API int GetAllOutputValues(TArray<FTG_Variant>& OutVariants, TArray<FName>* OutNames = nullptr) const;
 
 	// Get all the output's values which are compatible with FTG_Texture
 	// along with their names optionally
-	int GetAllOutputValues(TArray<FTG_Texture>& OutTextures, TArray<FName>* OutNames = nullptr) const;
+	UE_API int GetAllOutputValues(TArray<FTG_Texture>& OutTextures, TArray<FName>* OutNames = nullptr) const;
 
 	// Get the current common input variant type of the expression
 	// only valid for expression with FTG_Variant arguments
-	FTG_Variant::EType GetExpressionCommonInputVariantType() const;
+	UE_API FTG_Variant::EType GetExpressionCommonInputVariantType() const;
 
-	void OnPinConnectionUndo(FTG_Id InPinId);
+	UE_API void OnPinConnectionUndo(FTG_Id InPinId);
 
 protected:
 
 	// Return the array of alias name of all the pins
 	// Can be filtered based on the argument access type bitfield
 	// Default is not filtered, aka AccessFilter is 0
-	TArray<FName> GetPinAliasNames(uint8 AccessFilter = 0) const;
+	UE_API TArray<FName> GetPinAliasNames(uint8 AccessFilter = 0) const;
 
 	// Check that a candidate Pin's alias name is unique in the Node's Pins
 	// And in the Graph's Param Pins if it is a Param.
@@ -186,27 +188,27 @@ protected:
 	// This is meant to be used from within the Pin::SetAliasName call
 	// or when a new pin is allocated to make sure its default alias name doesn't collide in
 	// the scope of the node or of the graph if it is a param
-	FName ValidateGeneratePinAliasName(FName CandidateName, const FTG_Id& PinId) const;
+	UE_API FName ValidateGeneratePinAliasName(FName CandidateName, const FTG_Id& PinId) const;
 
 	// Validate that the node conforms to a conformant function coming from the Expression
-    void ValidateGenerateConformer(UTG_Pin* NewPin);
+    UE_API void ValidateGenerateConformer(UTG_Pin* NewPin);
 
 	// Trigger notification to the graph that the node pin's value (default input or output) has changed  
-	void NotifyGraphOfNodeChange(bool bIsTweaking);
+	UE_API void NotifyGraphOfNodeChange(bool bIsTweaking);
 
 	// Embedded expression is able to notify state changes:
 	friend class UTG_Expression;
 	friend class UTG_Pin;
-	void OnExpressionChangedWithoutVar(const FPropertyChangedEvent& PropertyChangedEvent);
-	void OnSignatureChanged();
-	void OnPinRenamed(FTG_Id InPinId, FName OldName);
+	UE_API void OnExpressionChangedWithoutVar(const FPropertyChangedEvent& PropertyChangedEvent);
+	UE_API void OnSignatureChanged();
+	UE_API void OnPinRenamed(FTG_Id InPinId, FName OldName);
 
-	void OnPinConnectionChanged(FTG_Id InPinId, FTG_Id OldPinId, FTG_Id NewPinId);
+	UE_API void OnPinConnectionChanged(FTG_Id InPinId, FTG_Id OldPinId, FTG_Id NewPinId);
 
 #if WITH_EDITOR
-	virtual void PostEditChangeProperty(FPropertyChangedEvent& PropertyChangedEvent) override;
-	virtual void PostEditUndo() override;
-	void OnPostUndo();
+	UE_API virtual void PostEditChangeProperty(FPropertyChangedEvent& PropertyChangedEvent) override;
+	UE_API virtual void PostEditUndo() override;
+	UE_API void OnPostUndo();
 #endif
 	
 	TSet<FName> WarningStack;
@@ -219,10 +221,11 @@ public:
 #endif
 
 	//// Logging Utilities
-	FString LogHead() const;
-	FString LogPins(FString Tab) const;
+	UE_API FString LogHead() const;
+	UE_API FString LogPins(FString Tab) const;
 };
 
 using FTG_Nodes = TArray<UTG_Node*>;
 
 
+#undef UE_API
diff --git a/Engine/Plugins/TextureGraph/Source/TextureGraph/Public/TG_OutputSettings.h b/Engine/Plugins/TextureGraph/Source/TextureGraph/Public/TG_OutputSettings.h
--- a/Engine/Plugins/TextureGraph/Source/TextureGraph/Public/TG_OutputSettings.h
+++ b/Engine/Plugins/TextureGraph/Source/TextureGraph/Public/TG_OutputSettings.h
@@ -7,12 +7,14 @@
 #include "Misc/OutputDeviceNull.h"
 #include "TG_OutputSettings.generated.h"
 
+#define UE_API TEXTUREGRAPH_API
+
 class UTG_Expression_Output;
 class UTG_Node;
 class UTextureGraph;
 
 USTRUCT()
-struct TEXTUREGRAPH_API FOutputExpressionInfo
+struct FOutputExpressionInfo
 {
 	GENERATED_BODY()
 
@@ -28,8 +30,8 @@ public:
 	}
 };
 
-UCLASS()
-class TEXTUREGRAPH_API UTG_ExportSettings : public UObject
+UCLASS(MinimalAPI)
+class UTG_ExportSettings : public UObject
 {
 	GENERATED_BODY()
 
@@ -39,7 +41,7 @@ public:
 };
 
 USTRUCT(BlueprintType)
-struct TEXTUREGRAPH_API FTG_OutputSettings
+struct FTG_OutputSettings
 {
 	GENERATED_USTRUCT_BODY()
 
@@ -93,7 +95,7 @@ struct TEXTUREGRAPH_API FTG_OutputSettings
 		return OutputName == Other.OutputName && BaseName == Other.BaseName;
 	}
 
-	void Initialize(FString PathName, FName InName = "Output");
+	UE_API void Initialize(FString PathName, FName InName = "Output");
 
 	void InitFromString(const FString& StrVal)
 	{
@@ -108,9 +110,11 @@ struct TEXTUREGRAPH_API FTG_OutputSettings
 		return ExportString;
 	}
 
-	void Set(int InWidth, int InHeight, FName Name = "None", FName Path = "None", ETG_TextureFormat Format = ETG_TextureFormat::BGRA8, ETG_TexturePresetType InTextureType = ETG_TexturePresetType::None,
+	UE_API void Set(int InWidth, int InHeight, FName Name = "None", FName Path = "None", ETG_TextureFormat Format = ETG_TextureFormat::BGRA8, ETG_TexturePresetType InTextureType = ETG_TexturePresetType::None,
 		TextureCompressionSettings InCompression = TextureCompressionSettings::TC_Default, TextureGroup InLodGroup = TextureGroup::TEXTUREGROUP_World, bool InbSRGB = false);
 
-	void OnSetTexturePresetType(ETG_TexturePresetType Type);
-	bool Validate(FString& OutErrors) const;
+	UE_API void OnSetTexturePresetType(ETG_TexturePresetType Type);
+	UE_API bool Validate(FString& OutErrors) const;
 };
+
+#undef UE_API
