Ran UnrealCodeFixup to fix dll storage

diff --git a/Engine/Plugins/TextureGraph/Source/TextureGraph/Public/TG_Pin.h b/Engine/Plugins/TextureGraph/Source/TextureGraph/Public/TG_Pin.h
--- a/Engine/Plugins/TextureGraph/Source/TextureGraph/Public/TG_Pin.h
+++ b/Engine/Plugins/TextureGraph/Source/TextureGraph/Public/TG_Pin.h
@@ -11,6 +11,8 @@
 
 #include "TG_Pin.generated.h"
 
+#define UE_API TEXTUREGRAPH_API
+
 // Pin are simple structs connecting the Nodes in the context of a Graph.
 // They are not intended to be derived.
 
@@ -25,35 +27,35 @@ class FAssetThumbnail;
 //
 // Pins are stored in the Graph.
 // Created when instancing a Node.
-UCLASS()
-class TEXTUREGRAPH_API UTG_Pin : public UObject {
+UCLASS(MinimalAPI)
+class UTG_Pin : public UObject {
 	GENERATED_BODY()
 
 		friend class UTG_Node;
 	friend class UTG_Graph;
 
 	// Only accessible by UTG_Graph when allocating a new Pin.
-	void Construct(FTG_Id InId, const FTG_Argument& InArgument);
+	UE_API void Construct(FTG_Id InId, const FTG_Argument& InArgument);
 
 	// Initialize the self var with the matching expression argument
-	void InitSelfVar();
+	UE_API void InitSelfVar();
 
-	void ResetSelfVarDefault();
+	UE_API void ResetSelfVarDefault();
 
 	// Used when remapping a pin to a different id, update the sefl Var if as well
-	void RemapId(FTG_Id InId);
+	UE_API void RemapId(FTG_Id InId);
 
 
 
 	// Should be protected, Used by TG_Graph when updating connections
 	// These have NO side effects, they just update the Pin members.
-	void				RemoveAllEdges();
-	void				RemoveEdge(FTG_Id OtherPinId);
-	void				AddEdge(FTG_Id OtherPinId);
-	void				RemapEdgeId(FTG_Id OldId, FTG_Id NewId);
+	UE_API void				RemoveAllEdges();
+	UE_API void				RemoveEdge(FTG_Id OtherPinId);
+	UE_API void				AddEdge(FTG_Id OtherPinId);
+	UE_API void				RemapEdgeId(FTG_Id OldId, FTG_Id NewId);
 
-	void				InstallInputVarConverterKey(FName ConverterKey);
-	void				RemoveInputVarConverterKey();
+	UE_API void				InstallInputVarConverterKey(FName ConverterKey);
+	UE_API void				RemoveInputVarConverterKey();
 
 	// Id member is assigned by TG_Graph at creation, potentially modified during a remap
 	UPROPERTY()
@@ -87,22 +89,22 @@ class TEXTUREGRAPH_API UTG_Pin : public UObject {
 		FTG_Var				ConvertedVar;
 
 protected:
-	void				NotifyPinSelfVarChanged(bool bIsTweaking = false);
+	UE_API void				NotifyPinSelfVarChanged(bool bIsTweaking = false);
 
 #if WITH_EDITOR
-	virtual void		PostEditChangeProperty(FPropertyChangedEvent& PropertyChangedEvent) override;
+	UE_API virtual void		PostEditChangeProperty(FPropertyChangedEvent& PropertyChangedEvent) override;
 #endif
 public:
 	// Override serialization for conditional serialization of SelfVar
-	virtual	void		Serialize(FArchive& Ar) override;
+	UE_API virtual	void		Serialize(FArchive& Ar) override;
 #if WITH_EDITOR
-	virtual bool		Modify(bool bAlwaysMarkDirty = true) override;
+	UE_API virtual bool		Modify(bool bAlwaysMarkDirty = true) override;
 #endif
 
 	bool				IsValid() const { return Id != FTG_Id::INVALID; }
 
-	UTG_Node*			GetNodePtr() const;
-	FTG_Id				GetNodeId() const;
+	UE_API UTG_Node*			GetNodePtr() const;
+	UE_API FTG_Id				GetNodeId() const;
 
 	FTG_Id				GetId() const { return Id; }
 
@@ -110,7 +112,7 @@ public:
 	FTG_Name			GetArgumentName() const { return GetArgument().GetName(); }
 
 	FName				GetAliasName() const { return AliasName; }
-	void				SetAliasName(FName InAliasName); // Assign the alias name and notify the graph for change
+	UE_API void				SetAliasName(FName InAliasName); // Assign the alias name and notify the graph for change
 	bool				HasAliasName() const { return GetArgumentName() != AliasName; }
 
 	FTG_ArgumentType	GetArgumentType() const { return GetArgument().GetType(); }
@@ -150,19 +152,19 @@ public:
 	// Should only be accessed by the Graph Evaluation
 	FTG_Var*			EditConvertedVar() { return &ConvertedVar; }
 
-	FString				LogHead() const;
-	FString				Log(FString Tab) const;
-	FString				LogTooltip() const;
+	UE_API FString				LogHead() const;
+	UE_API FString				Log(FString Tab) const;
+	UE_API FString				LogTooltip() const;
 
 	// Access the FProperty associated to this pin in the Expression::Class of the node
 	// A valid FProperty is returned if it exists, null otherwise
-	FProperty*			GetExpressionProperty() const;
+	UE_API FProperty*			GetExpressionProperty() const;
 
 	FTG_Evaluation::VarConformer ConformerFunctor;
 	bool				NeedsConformance() const { return ConformerFunctor != nullptr; }
 
 	// Hashing
-	static FTG_Hash		Hash(const UTG_Pin& Pin);
+	static UE_API FTG_Hash		Hash(const UTG_Pin& Pin);
 	FTG_Hash			Hash() const { return Hash(*this); }
 
 
@@ -183,27 +185,27 @@ public:
 	// The function return true if valid, false otherwise, if false they are NO OP
 
 	// Scalar getter only on a scalar argument
-	bool GetValue(float& OutValue) const;
+	UE_API bool GetValue(float& OutValue) const;
 
 	// bool getter only on a boolean argument
-	bool GetValue(bool& OutValue) const;
+	UE_API bool GetValue(bool& OutValue) const;
 
 	// String getter only on a FString argument
-	bool GetValue(FString& OutValue) const;
+	UE_API bool GetValue(FString& OutValue) const;
 
 	// Color or Vector getter only on a vector or color argument
-	bool GetValue(FLinearColor& OutValue) const;
-	bool GetValue(FVector4f& OutValue) const;
+	UE_API bool GetValue(FLinearColor& OutValue) const;
+	UE_API bool GetValue(FVector4f& OutValue) const;
 
 	// Texture value getter only on a Texture argument
-	bool GetValue(FTG_Texture& OutValue) const;
+	UE_API bool GetValue(FTG_Texture& OutValue) const;
 
 	// Texture value getter only on a Texture argument
-	bool GetValue(FTG_VariantArray& OutValue) const;
+	UE_API bool GetValue(FTG_VariantArray& OutValue) const;
 
 	// Variant value getter only if a Variant argument OR if the argument is one of the supported type by Variant (Scalar/Color/Vector/Texture)
 	// After getting the value, the next step is to check the type contained in the variant.
-	bool GetValue(FTG_Variant& OutValue) const;
+	UE_API bool GetValue(FTG_Variant& OutValue) const;
 
 	// Useful for the UI:
 	// access the current value of this pin as passed to the evaluation if it is an input
@@ -214,28 +216,28 @@ public:
 	//		else if connected with a conversion: the convertedVar
 	//		else if connected without conversion: the Var from the connected pin
 	// This is MEANT TO BE USED for UI purpose !!!!
-	FString				GetEvaluatedVarValue() const;
+	UE_API FString				GetEvaluatedVarValue() const;
 
 	// Set the self var value from string
 	// Used by UI to assign pin default value
 	// This is the same as the selfvar setter
-	void				FromString(const FString& InValueStr, bool bIsTweaking = false);
+	UE_API void				FromString(const FString& InValueStr, bool bIsTweaking = false);
 
 	// Setters for Scalar, Color, Vector, bool and String  arguments
 	// These check that the arg is of the correct type
-	bool SetValue(float Value);
+	UE_API bool SetValue(float Value);
 	bool SetValue(double Value) { return SetValue((float)Value); }
 	bool SetValue(uint32 Value) { return SetValue((float)Value); }
 	bool SetValue(int32 Value) { return SetValue((float)Value); }
-	bool SetValue(bool bValue);
-	bool SetValue(FString Value); // FName pin is assigned by the templated version, not FString
-	bool SetValue(const FLinearColor& Value);
-	bool SetValue(const FVector4f& Value);
+	UE_API bool SetValue(bool bValue);
+	UE_API bool SetValue(FString Value); // FName pin is assigned by the templated version, not FString
+	UE_API bool SetValue(const FLinearColor& Value);
+	UE_API bool SetValue(const FVector4f& Value);
 
-	bool SetValue(const float* Value, size_t Count);
+	UE_API bool SetValue(const float* Value, size_t Count);
 	bool SetValue(FTG_Texture&) { return false; } // this is not a possible case
 	bool SetValue(FTG_VariantArray&) { return false; } // this is not a possible case
-	bool SetValue(const FTG_TextureDescriptor& Value); 
+	UE_API bool SetValue(const FTG_TextureDescriptor& Value); 
 
 	// Template version of the setter works correctly as long as the selfvar contains exactly the specified type
 	// Use with care
@@ -255,3 +257,4 @@ public:
 
 
 
+#undef UE_API
diff --git a/Engine/Plugins/TextureGraph/Source/TextureGraph/Public/TG_Signature.h b/Engine/Plugins/TextureGraph/Source/TextureGraph/Public/TG_Signature.h
--- a/Engine/Plugins/TextureGraph/Source/TextureGraph/Public/TG_Signature.h
+++ b/Engine/Plugins/TextureGraph/Source/TextureGraph/Public/TG_Signature.h
@@ -9,6 +9,8 @@
 
 #include "TG_Signature.generated.h"
 
+#define UE_API TEXTUREGRAPH_API
+
 
 // Pin access either a in(put) or out(put) or ?
 UENUM()
@@ -41,7 +43,7 @@ enum class ETG_Access : uint8
 // Descriptor struct of an argument core desc
 // all the typing information for a pin EXCEPT its name
 USTRUCT()
-struct TEXTUREGRAPH_API FTG_ArgumentType
+struct FTG_ArgumentType
 {
 	GENERATED_BODY();
 
@@ -120,7 +122,7 @@ using FTG_ArgumentTypes = TArray<FTG_ArgumentType>;
 // where the name is combined for hashing with
 // the rest of the argument type fields
 USTRUCT()
-struct TEXTUREGRAPH_API FTG_Argument
+struct FTG_Argument
 {
 	GENERATED_BODY();
 
@@ -153,22 +155,22 @@ public:
 
 	FName				GetCPPTypeName() const { return CPPTypeName; }
 
-	bool				IsScalar() const; // float or int
-	bool				IsBool() const; // boolean
-	bool				IsString() const; // string  but NOT FName
-	bool				IsColor() const; // FlinearColor
-	bool				IsVector() const; // FLinearColor or ...
-	bool				IsTexture() const; // FTG_Texture or ...
-	bool				IsArray() const; // FTG_TextureArray or ...
-	bool				IsVariant() const; // FTG_Variant
-	bool				IsObject(UClass* InClass) const; // check that the Argument is a UObject and of baseType InClass
-	bool				GetMetaData(const FName KeyToFind, FString& OutValue);
+	UE_API bool				IsScalar() const; // float or int
+	UE_API bool				IsBool() const; // boolean
+	UE_API bool				IsString() const; // string  but NOT FName
+	UE_API bool				IsColor() const; // FlinearColor
+	UE_API bool				IsVector() const; // FLinearColor or ...
+	UE_API bool				IsTexture() const; // FTG_Texture or ...
+	UE_API bool				IsArray() const; // FTG_TextureArray or ...
+	UE_API bool				IsVariant() const; // FTG_Variant
+	UE_API bool				IsObject(UClass* InClass) const; // check that the Argument is a UObject and of baseType InClass
+	UE_API bool				GetMetaData(const FName KeyToFind, FString& OutValue);
 	
     // Hashing
-    static FTG_Hash Hash(const FTG_Argument& Argument);
+    static UE_API FTG_Hash Hash(const FTG_Argument& Argument);
     FTG_Hash Hash() const { return Hash(*this); }
 
-    const static FTG_Argument Invalid;
+    UE_API const static FTG_Argument Invalid;
 
 	friend bool operator== (const FTG_Argument& lhs, const FTG_Argument& rhs) {
 		return (lhs.Name == rhs.Name) && (lhs.CPPTypeName == rhs.CPPTypeName) && (lhs.ArgumentType == rhs.ArgumentType);
@@ -188,7 +190,7 @@ TArray<FName> TG_MakeArrayOfArgumentNames(const FTG_Arguments& InArguments);
 // The arguments all have an index (their position in the set).
 // All the argument names in the set MUST be different.
 // THe hashes of each arguments are also different as a consequence.
-struct TEXTUREGRAPH_API FTG_ArgumentSet
+struct FTG_ArgumentSet
 {
 private:
 public:
@@ -205,14 +207,14 @@ public:
 
     // Find the Index of an argument at a particular name.
     // Return INVALID_INDEX if not found
-    FTG_Index FindName(FTG_Name& Name) const;
+    UE_API FTG_Index FindName(FTG_Name& Name) const;
 
     // Find the Index of a argument at a particular hash.
     // Return INVALID_INDEX if not found
-    FTG_Index FindHash(FTG_Hash Hash) const;
+    UE_API FTG_Index FindHash(FTG_Hash Hash) const;
 
     // Hashing
-    static FTG_Hash Hash(const FTG_ArgumentSet& Set);
+    static UE_API FTG_Hash Hash(const FTG_ArgumentSet& Set);
     FTG_Hash Hash() const { return Hash(*this); }
 
 };
@@ -220,7 +222,7 @@ public:
 // The Signature describe a full function signature
 // a name + collection of arguments
 USTRUCT()
-struct TEXTUREGRAPH_API FTG_Signature {
+struct FTG_Signature {
 
     GENERATED_BODY();
 
@@ -238,7 +240,7 @@ protected:
 
 public:
 
-	FTG_Signature(const FInit& Init);
+	UE_API FTG_Signature(const FInit& Init);
     FTG_Signature() {}
 
 	// Added an defaulted copy constructor here because Clang 16 expects either both or neither copy constructor
@@ -262,7 +264,7 @@ public:
 	FORCEINLINE int32 GetNumArguments() const { return InArgs.Arguments.Num()
     												+ OutArgs.Arguments.Num()
 													+ PrivateArgs.Arguments.Num(); }
-	const FTG_Argument& GetArgument(int32 ArgIdx);
+	UE_API const FTG_Argument& GetArgument(int32 ArgIdx);
 
 	FTG_Arguments GetArguments() const {	FTG_Arguments Args(InArgs.Arguments);
 												Args.Append(OutArgs.Arguments);
@@ -296,20 +298,20 @@ public:
 	};
 	using IndexAccessArray = TArray<IndexAccess>;
 
-	IndexAccess FindArgumentIndexFromName(FTG_Name& InName) const;
-	const FTG_Argument& FindArgumentFromName(FTG_Name& InName) const;
+	UE_API IndexAccess FindArgumentIndexFromName(FTG_Name& InName) const;
+	UE_API const FTG_Argument& FindArgumentFromName(FTG_Name& InName) const;
 
     // Hashing
-    static FTG_Hash Hash(const FTG_Signature& Type);
+    static UE_API FTG_Hash Hash(const FTG_Signature& Type);
     FTG_Hash Hash() const { return Hash(*this); }
 
-	const FTG_Argument& FindArgumentFromHash(FTG_Hash Hash) const;
-	IndexAccess FindArgumentIndexFromHash(FTG_Hash Hash) const;
+	UE_API const FTG_Argument& FindArgumentFromHash(FTG_Hash Hash) const;
+	UE_API IndexAccess FindArgumentIndexFromHash(FTG_Hash Hash) const;
 
 	// Generate a mapping table of the argument indices which is an array of index elements 
 	// The element at poisiton i indicate the mapping argument index in the FromSignature corresponding to index i in the ToSignature
 	// if the ToSignature argument doesn't have an equivalent, then INVALID_INDEX is at this element
-	static FTG_Indices GenerateMappingArgIdxTable(const FTG_Signature& FromSignature, const FTG_Signature& ToSignature);
+	static UE_API FTG_Indices GenerateMappingArgIdxTable(const FTG_Signature& FromSignature, const FTG_Signature& ToSignature);
 };
 using FTG_SignaturePtr = TSharedPtr< FTG_Signature >;
 
@@ -319,3 +321,4 @@ using FTG_SignaturePtr = TSharedPtr< FTG_Signature >;
 #define TG_ARGUMENT_IN(name, t) { TEXT(TG_STR(name)), { ETG_Access::In }}
 #define TG_ARGUMENT_OUT(name, t) { TEXT(TG_STR(name)), { ETG_Access::Out }}
 
+#undef UE_API
diff --git a/Engine/Plugins/TextureGraph/Source/TextureGraph/Public/TG_SystemTypes.h b/Engine/Plugins/TextureGraph/Source/TextureGraph/Public/TG_SystemTypes.h
--- a/Engine/Plugins/TextureGraph/Source/TextureGraph/Public/TG_SystemTypes.h
+++ b/Engine/Plugins/TextureGraph/Source/TextureGraph/Public/TG_SystemTypes.h
@@ -6,6 +6,8 @@
 
 #include "TG_SystemTypes.generated.h"
 
+#define UE_API TEXTUREGRAPH_API
+
 DECLARE_LOG_CATEGORY_EXTERN(LogTextureGraph, Log, All);
 
 using FTG_Index = int32;
@@ -13,13 +15,13 @@ using FTG_Indices = TArray<FTG_Index>;
 
 using ErrorReportMap = TMap<int32, TArray<struct FTextureGraphErrorReport>>;
 USTRUCT()
-struct TEXTUREGRAPH_API FTG_Id
+struct FTG_Id
 {
 	GENERATED_BODY()
 
 	static constexpr FTG_Index INVALID_INDEX = -1;
 
-	static const FTG_Id INVALID;
+	static UE_API const FTG_Id INVALID;
 
 	UPROPERTY()
 	uint32 _IndexRaw = 0xFFFFFFFF; // The saved value for the FTG_Id is a single uint32 containing the 2 sub indices as in16 packed.
@@ -155,4 +157,6 @@ namespace TG_MetadataSpecifiers
 
 	// Assign to Levels Settings FVector UProperty to go through the widget customizer
 	const FName MD_LevelsSettings = TEXT("MD_LevelsSettings");
-}
\ No newline at end of file
+}
+
+#undef UE_API
diff --git a/Engine/Plugins/TextureGraph/Source/TextureGraph/Public/TG_Texture.h b/Engine/Plugins/TextureGraph/Source/TextureGraph/Public/TG_Texture.h
--- a/Engine/Plugins/TextureGraph/Source/TextureGraph/Public/TG_Texture.h
+++ b/Engine/Plugins/TextureGraph/Source/TextureGraph/Public/TG_Texture.h
@@ -9,8 +9,10 @@
 #include "Misc/OutputDeviceNull.h"
 #include "TG_Texture.generated.h"
 
+#define UE_API TEXTUREGRAPH_API
+
 USTRUCT(BlueprintType)
-struct TEXTUREGRAPH_API FTG_TextureDescriptor
+struct FTG_TextureDescriptor
 {
 	GENERATED_USTRUCT_BODY()
 	FTG_TextureDescriptor()
@@ -95,7 +97,7 @@ PRAGMA_ENABLE_DEPRECATION_WARNINGS
 };
 
 USTRUCT()
-struct TEXTUREGRAPH_API FTG_Texture
+struct FTG_Texture
 {
 	GENERATED_BODY()
 	
@@ -117,9 +119,9 @@ public:
 	friend FArchive& operator<<(FArchive& Ar, FTG_Texture& T);
 	FTG_Texture& operator = (TiledBlobRef RHS) { RasterBlob = RHS; ResetTexturePath(); return *this; }
 	FTG_Texture& operator = (TiledBlobPtr RHS) { RasterBlob = RHS; ResetTexturePath(); return *this; }
-	bool operator == (const FTG_Texture& RHS) const;
+	UE_API bool operator == (const FTG_Texture& RHS) const;
 	FORCEINLINE TiledBlob* operator -> () const { return RasterBlob.get(); }
-	void ResetTexturePath();
+	UE_API void ResetTexturePath();
 	
 	//////////////////////////////////////////////////////////////////////////
 	/// Static functions
@@ -140,4 +142,6 @@ public:
 	/// Inline functions
 	//////////////////////////////////////////////////////////////////////////
 	FORCEINLINE BufferDescriptor GetBufferDescriptor() const { return Descriptor.ToBufferDescriptor(); }
-};
\ No newline at end of file
+};
+
+#undef UE_API
diff --git a/Engine/Plugins/TextureGraph/Source/TextureGraph/Public/TG_Var.h b/Engine/Plugins/TextureGraph/Source/TextureGraph/Public/TG_Var.h
--- a/Engine/Plugins/TextureGraph/Source/TextureGraph/Public/TG_Var.h
+++ b/Engine/Plugins/TextureGraph/Source/TextureGraph/Public/TG_Var.h
@@ -8,6 +8,8 @@
 
 #include "TG_Var.generated.h"
 
+#define UE_API TEXTUREGRAPH_API
+
 class UTG_Expression;
 
 template <typename T>
@@ -25,7 +27,7 @@ FString TG_Var_LogValue(T& Value)
 }
 
 USTRUCT()
-struct TEXTUREGRAPH_API FTG_Var
+struct FTG_Var
 {
 	GENERATED_BODY()
 
@@ -46,12 +48,12 @@ struct TEXTUREGRAPH_API FTG_Var
 	//typedef std::function<void(VarPropertySerialInfo&)>		VarPropertySerializer;
 	typedef void (*VarPropertySerializer) (VarPropertySerialInfo&)		;
 	typedef TMap<FName, VarPropertySerializer>				VarPropertySerializerMap;
-	static VarPropertySerializerMap							DefaultPropertySerializers;
-	static void RegisterDefaultSerializers();
-	static void RegisterVarPropertySerializer(FName CPPTypeName, VarPropertySerializer Serializer);
-	static void UnregisterVarPropertySerializer(FName CPPTypeName);
+	static UE_API VarPropertySerializerMap							DefaultPropertySerializers;
+	static UE_API void RegisterDefaultSerializers();
+	static UE_API void RegisterVarPropertySerializer(FName CPPTypeName, VarPropertySerializer Serializer);
+	static UE_API void UnregisterVarPropertySerializer(FName CPPTypeName);
 
-	bool CopyGeneric(UTG_Expression* Owner, const FTG_Argument& Arg, bool CopyVarToProperty, int Index = -1);
+	UE_API bool CopyGeneric(UTG_Expression* Owner, const FTG_Argument& Arg, bool CopyVarToProperty, int Index = -1);
 
 	// Map of per type serializer for the var to/from an Archive
 	struct VarArchiveSerialInfo
@@ -61,11 +63,11 @@ struct TEXTUREGRAPH_API FTG_Var
 	};
 	typedef std::function<void(VarArchiveSerialInfo&)>	VarArchiveSerializer;
 	typedef TMap<FName, VarArchiveSerializer>			VarArchiveSerializerMap;
-	static VarArchiveSerializerMap						DefaultArchiveSerializers;
+	static UE_API VarArchiveSerializerMap						DefaultArchiveSerializers;
 
 	// Serialize the var value
 	// Called from the Pin serialize()
-	void Serialize(FArchive& Ar, FTG_Id PinId, const FTG_Argument& Argument);
+	UE_API void Serialize(FArchive& Ar, FTG_Id PinId, const FTG_Argument& Argument);
 
 
 	UPROPERTY(Transient)
@@ -100,9 +102,9 @@ struct TEXTUREGRAPH_API FTG_Var
 	}
 
 public:
-	FTG_Var();
-	FTG_Var(FTG_Id InPinUuid);
-	FTG_Var(const FTG_Var& InVar);
+	UE_API FTG_Var();
+	UE_API FTG_Var(FTG_Id InPinUuid);
+	UE_API FTG_Var(const FTG_Var& InVar);
 
 	FTG_Var& operator = (const FTG_Var& InVar) { PinId = InVar.PinId; Concept = InVar.Concept; return *this; }
 
@@ -127,11 +129,11 @@ public:
 	
 	// Manage data Arg provide the type name and rely on the VarAllocator map to create the concrete cpp type
 	// If Owner is provided then Var can be initialized from the matching owner's property if found
-	void CopyTo(UTG_Expression* Owner, FTG_Argument& Arg, int Index = -1);
-	void CopyFrom(UTG_Expression* Owner, FTG_Argument& Arg, int Index = -1);
+	UE_API void CopyTo(UTG_Expression* Owner, FTG_Argument& Arg, int Index = -1);
+	UE_API void CopyFrom(UTG_Expression* Owner, FTG_Argument& Arg, int Index = -1);
 
-	void CopyTo(FTG_Var* InVar);
-	void CopyFrom(FTG_Var* InVar);
+	UE_API void CopyTo(FTG_Var* InVar);
+	UE_API void CopyFrom(FTG_Var* InVar);
 
 	// Manage the Var Data with the knowledge of the Type
 	void Reset() { Concept = nullptr; }
@@ -164,11 +166,11 @@ public:
 		return static_cast<FModel<T>*>(Concept.Get())->Value;
 	}
 
-	FString LogHead() const;
+	UE_API FString LogHead() const;
 
 private:
 
-	struct TEXTUREGRAPH_API FConcept
+	struct FConcept
 	{
 		virtual ~FConcept() {}
 
@@ -180,7 +182,7 @@ private:
 	};
 
 	template <class T>
-	struct TEXTUREGRAPH_API FModel : public FConcept
+	struct FModel : public FConcept
 	{
 		T Value;
 		bool bIsArray = false;
@@ -216,41 +218,42 @@ private:
 	mutable TSharedPtr<FConcept> Concept;
 };
 
-template <> FString TG_Var_LogValue(bool& Value);
-template <> FString TG_Var_LogValue(int& Value);
-template <> FString TG_Var_LogValue(float& Value);
-template <> FString TG_Var_LogValue(uint8& Value);
-template <> FString TG_Var_LogValue(FLinearColor& Value);
-template <> FString TG_Var_LogValue(FVector4f& Value);
-template <> FString TG_Var_LogValue(FVector2f& Value);
-template <> FString TG_Var_LogValue(FName& Value);
-template <> FString TG_Var_LogValue(FString& Value);
-template <> FString TG_Var_LogValue(TObjectPtr<UObject>& Value);
-template <> FString TG_Var_LogValue(struct FTG_Texture& Value);
-template <> FString TG_Var_LogValue(struct FTG_Scalar& Value);
-template <> FString TG_Var_LogValue(struct FTG_OutputSettings& Value);
-template <> FString TG_Var_LogValue(struct FTG_TextureDescriptor& Value);
-template <> FString TG_Var_LogValue(struct FTG_Variant& Value);
-template <> FString TG_Var_LogValue(struct FTG_Material& Value);
-
-
-template <> void TG_Var_SetValueFromString(int& Value, const FString& StrVal);
-template <> void TG_Var_SetValueFromString(bool& Value, const FString& StrVal);
-template <> void TG_Var_SetValueFromString(float& Value, const FString& StrVal);
-template <> void TG_Var_SetValueFromString(uint8& Value, const FString& StrVal);
-template <> void TG_Var_SetValueFromString(TObjectPtr<UObject>& Value, const FString& StrVal);
-
-template <> void TG_Var_SetValueFromString(FLinearColor& Value, const FString& StrVal);
-template <> void TG_Var_SetValueFromString(FVector4f& Value, const FString& StrVal);
-template <> void TG_Var_SetValueFromString(FVector2f& Value, const FString& StrVal);
-template <> void TG_Var_SetValueFromString(FName& Value, const FString& StrVal);
-template <> void TG_Var_SetValueFromString(FString& Value, const FString& StrVal);
-
-//template <> void TG_Var_SetValueFromString(struct FTG_Texture& Value, const FString& StrVal);
-template <> void TG_Var_SetValueFromString(struct FTG_Scalar& Value, const FString& StrVal);
-template <> void TG_Var_SetValueFromString(struct FTG_OutputSettings& Value, const FString& StrVal);
-template <> void TG_Var_SetValueFromString(struct FTG_TextureDescriptor& Value, const FString& StrVal);
-template <> void TG_Var_SetValueFromString(struct FTG_Variant& Value, const FString& StrVal);
-template <> void TG_Var_SetValueFromString(struct FTG_Material& Value, const FString& StrVal);
-
-
+template <> UE_API FString TG_Var_LogValue(bool& Value);
+template <> UE_API FString TG_Var_LogValue(int& Value);
+template <> UE_API FString TG_Var_LogValue(float& Value);
+template <> UE_API FString TG_Var_LogValue(uint8& Value);
+template <> UE_API FString TG_Var_LogValue(FLinearColor& Value);
+template <> UE_API FString TG_Var_LogValue(FVector4f& Value);
+template <> UE_API FString TG_Var_LogValue(FVector2f& Value);
+template <> UE_API FString TG_Var_LogValue(FName& Value);
+template <> UE_API FString TG_Var_LogValue(FString& Value);
+template <> UE_API FString TG_Var_LogValue(TObjectPtr<UObject>& Value);
+template <> UE_API FString TG_Var_LogValue(struct FTG_Texture& Value);
+template <> UE_API FString TG_Var_LogValue(struct FTG_Scalar& Value);
+template <> UE_API FString TG_Var_LogValue(struct FTG_OutputSettings& Value);
+template <> UE_API FString TG_Var_LogValue(struct FTG_TextureDescriptor& Value);
+template <> UE_API FString TG_Var_LogValue(struct FTG_Variant& Value);
+template <> UE_API FString TG_Var_LogValue(struct FTG_Material& Value);
+
+
+template <> UE_API void TG_Var_SetValueFromString(int& Value, const FString& StrVal);
+template <> UE_API void TG_Var_SetValueFromString(bool& Value, const FString& StrVal);
+template <> UE_API void TG_Var_SetValueFromString(float& Value, const FString& StrVal);
+template <> UE_API void TG_Var_SetValueFromString(uint8& Value, const FString& StrVal);
+template <> UE_API void TG_Var_SetValueFromString(TObjectPtr<UObject>& Value, const FString& StrVal);
+
+template <> UE_API void TG_Var_SetValueFromString(FLinearColor& Value, const FString& StrVal);
+template <> UE_API void TG_Var_SetValueFromString(FVector4f& Value, const FString& StrVal);
+template <> UE_API void TG_Var_SetValueFromString(FVector2f& Value, const FString& StrVal);
+template <> UE_API void TG_Var_SetValueFromString(FName& Value, const FString& StrVal);
+template <> UE_API void TG_Var_SetValueFromString(FString& Value, const FString& StrVal);
+
+//template <> UE_API void TG_Var_SetValueFromString(struct FTG_Texture& Value, const FString& StrVal);
+template <> UE_API void TG_Var_SetValueFromString(struct FTG_Scalar& Value, const FString& StrVal);
+template <> UE_API void TG_Var_SetValueFromString(struct FTG_OutputSettings& Value, const FString& StrVal);
+template <> UE_API void TG_Var_SetValueFromString(struct FTG_TextureDescriptor& Value, const FString& StrVal);
+template <> UE_API void TG_Var_SetValueFromString(struct FTG_Variant& Value, const FString& StrVal);
+template <> UE_API void TG_Var_SetValueFromString(struct FTG_Material& Value, const FString& StrVal);
+
+
+#undef UE_API
diff --git a/Engine/Plugins/TextureGraph/Source/TextureGraph/Public/TG_Variant.h b/Engine/Plugins/TextureGraph/Source/TextureGraph/Public/TG_Variant.h
--- a/Engine/Plugins/TextureGraph/Source/TextureGraph/Public/TG_Variant.h
+++ b/Engine/Plugins/TextureGraph/Source/TextureGraph/Public/TG_Variant.h
@@ -7,6 +7,8 @@
 
 #include "TG_Variant.generated.h"
 
+#define UE_API TEXTUREGRAPH_API
+
 // Inner data type of the variant relies on the TVariant
 using FTG_VariantInnerData = TVariant<float, FLinearColor, FVector4f, FTG_Texture>;
 struct FTG_Texture;
@@ -28,7 +30,7 @@ ENUM_RANGE_BY_FIRST_AND_LAST(ETG_VariantType, ETG_VariantType::Scalar, ETG_Varia
 /// Base class for making working with Variants easier and less repetitive
 //////////////////////////////////////////////////////////////////////////
 USTRUCT()
-struct TEXTUREGRAPH_API FTG_Variant
+struct FTG_Variant
 {
 	GENERATED_BODY()
 
@@ -36,19 +38,19 @@ public:
 	using Variant = FTG_VariantInnerData;
 	using EType = ETG_VariantType;
 
-	FTG_Variant();
-	FTG_Variant(const FTG_Variant& RHS);
-	FTG_Variant(float RHS);
-	FTG_Variant(FVector4f RHS);
-	FTG_Variant(FLinearColor RHS);
-	FTG_Variant(FTG_Texture RHS);
+	UE_API FTG_Variant();
+	UE_API FTG_Variant(const FTG_Variant& RHS);
+	UE_API FTG_Variant(float RHS);
+	UE_API FTG_Variant(FVector4f RHS);
+	UE_API FTG_Variant(FLinearColor RHS);
+	UE_API FTG_Variant(FTG_Texture RHS);
 
-	FTG_Variant& operator = (const FTG_Variant& RHS);
-	FTG_Variant& operator = (const float RHS);
-	FTG_Variant& operator = (const FVector4f RHS);
-	FTG_Variant& operator = (const FLinearColor RHS);
-	FTG_Variant& operator = (const FTG_Texture RHS);
-	bool operator == (const FTG_Variant& RHS) const;
+	UE_API FTG_Variant& operator = (const FTG_Variant& RHS);
+	UE_API FTG_Variant& operator = (const float RHS);
+	UE_API FTG_Variant& operator = (const FVector4f RHS);
+	UE_API FTG_Variant& operator = (const FLinearColor RHS);
+	UE_API FTG_Variant& operator = (const FTG_Texture RHS);
+	UE_API bool operator == (const FTG_Variant& RHS) const;
 
 	// Retrieve the FName corresponding to a variant type
 	static FName GetNameFromType(ETG_VariantType InType)
@@ -125,7 +127,7 @@ public:
 	{
 		return Ar << D.Data;
 	}
-	bool Serialize( FArchive& Ar );
+	UE_API bool Serialize( FArchive& Ar );
 	
 	// FTG_Variant struct members and methods
 
@@ -138,7 +140,7 @@ public:
 	// Reset the Variant to the specified type,
 	// If the type is changed, the value is reset to 0
 	// return true if the type as mutated
-	bool ResetTypeAs(EType InType);
+	UE_API bool ResetTypeAs(EType InType);
 
 
 	// Helper predicates to narrow down the type of the variant
@@ -148,26 +150,26 @@ public:
 	bool IsTexture() const { return IsTexture(GetType()); }
 
 
-	operator bool() const;
+	UE_API operator bool() const;
 
 	// Getter and Editor
 	// The getter methods are valid ONLY if the type matches
-	const float& GetScalar() const;
-	const FLinearColor& GetColor() const;
-	const FVector4f& GetVector() const;
-	const FTG_Texture& GetTexture() const;
-	FTG_Texture& GetTexture();
+	UE_API const float& GetScalar() const;
+	UE_API const FLinearColor& GetColor() const;
+	UE_API const FVector4f& GetVector() const;
+	UE_API const FTG_Texture& GetTexture() const;
+	UE_API FTG_Texture& GetTexture();
 
-	FLinearColor GetColor(FLinearColor Default = FLinearColor::Black);
-	FVector4f GetVector(FVector4f Default = FVector4f { 0, 0, 0, 0 });
-	FTG_Texture GetTexture(FTG_EvaluationContext* InContext, FTG_Texture Default = { TextureHelper::GBlack }, const BufferDescriptor* DesiredDesc = nullptr);
+	UE_API FLinearColor GetColor(FLinearColor Default = FLinearColor::Black);
+	UE_API FVector4f GetVector(FVector4f Default = FVector4f { 0, 0, 0, 0 });
+	UE_API FTG_Texture GetTexture(FTG_EvaluationContext* InContext, FTG_Texture Default = { TextureHelper::GBlack }, const BufferDescriptor* DesiredDesc = nullptr);
 
 	// The editor methods are calling ResetTypeAs(expectedType) 
 	// and thus the Variant will mutate to the expected type and assigned 0
-	float& EditScalar();
-	FLinearColor& EditColor();
-	FVector4f& EditVector();
-	FTG_Texture& EditTexture();
+	UE_API float& EditScalar();
+	UE_API FLinearColor& EditColor();
+	UE_API FVector4f& EditVector();
+	UE_API FTG_Texture& EditTexture();
 };
 
 template<>
@@ -186,7 +188,7 @@ struct TStructOpsTypeTraits<FTG_Variant>
 /// Variant Arrays
 //////////////////////////////////////////////////////////////////////////
 USTRUCT()
-struct TEXTUREGRAPH_API FTG_VariantArray
+struct FTG_VariantArray
 {
 	GENERATED_BODY()
 
@@ -196,11 +198,11 @@ private:
 
 
 public:
-	void SetNum(int Count);
-	void Set(int Index, const FTG_Variant& Var);
-	void Set(int Index, const FTG_Texture& Tex);
-	void CopyFrom(const FTG_VariantArray& RHS);
-	void SetViewIndex(int Index);
+	UE_API void SetNum(int Count);
+	UE_API void Set(int Index, const FTG_Variant& Var);
+	UE_API void Set(int Index, const FTG_Texture& Tex);
+	UE_API void CopyFrom(const FTG_VariantArray& RHS);
+	UE_API void SetViewIndex(int Index);
 
 	FORCEINLINE bool IsTexture()
 	{
@@ -246,3 +248,4 @@ public:
 	}
 };
 
+#undef UE_API
diff --git a/Engine/Plugins/TextureGraph/Source/TextureGraph/Public/TextureGraph.h b/Engine/Plugins/TextureGraph/Source/TextureGraph/Public/TextureGraph.h
--- a/Engine/Plugins/TextureGraph/Source/TextureGraph/Public/TextureGraph.h
+++ b/Engine/Plugins/TextureGraph/Source/TextureGraph/Public/TextureGraph.h
@@ -7,6 +7,8 @@
 #include "TG_GraphEvaluation.h"
 #include "TextureGraph.generated.h"
 
+#define UE_API TEXTUREGRAPH_API
+
 class UTG_Graph;
 using ErrorReportMap = TMap<int32, TArray<FTextureGraphErrorReport>>;
 
@@ -52,8 +54,8 @@ public:
 	TEXTUREGRAPH_API void Log() const;
 };
 
-UCLASS(BlueprintType)
-class TEXTUREGRAPH_API UTextureGraph: public UTextureGraphBase
+UCLASS(MinimalAPI, BlueprintType)
+class UTextureGraph: public UTextureGraphBase
 {
 	GENERATED_BODY()
 
@@ -62,20 +64,20 @@ public:
 	TObjectPtr<class UTG_Graph> TextureGraph;
 	
 protected:
-	void GatherAllDependentGraphs(TArray<UTextureGraph*>& DependentGraphs) const;
+	UE_API void GatherAllDependentGraphs(TArray<UTextureGraph*>& DependentGraphs) const;
 	
-	bool CheckRecursiveDependency(const UTextureGraph* InTextureGraph) const;
+	UE_API bool CheckRecursiveDependency(const UTextureGraph* InTextureGraph) const;
 	
 public:
-	virtual void Construct(FString Name) override;
+	UE_API virtual void Construct(FString Name) override;
 	virtual const UTG_Graph* Graph() const { return TextureGraph; }
 	virtual UTG_Graph* Graph() { return TextureGraph; }
-	bool IsDependentOn(const UTextureGraph* TextureGraph) const;
-	bool HasCyclicDependency() const;
+	UE_API bool IsDependentOn(const UTextureGraph* TextureGraph) const;
+	UE_API bool HasCyclicDependency() const;
 };
 
-UCLASS(BlueprintType)
-class TEXTUREGRAPH_API UTextureGraphInstance: public UTextureGraphBase
+UCLASS(MinimalAPI, BlueprintType)
+class UTextureGraphInstance: public UTextureGraphBase
 {
 	friend struct FTG_InstanceImpl;
 	GENERATED_BODY()
@@ -95,18 +97,20 @@ public:
 	UPROPERTY(EditAnywhere, BlueprintReadOnly, Category=TextureGraphInstance, AssetRegistrySearchable)
 	TObjectPtr<UTextureGraphBase> ParentTextureGraph; 
 
-	virtual void Construct(FString Name) override;
+	UE_API virtual void Construct(FString Name) override;
 
-	void CopyParamsToRuntimeGraph();
-	bool CheckOutputSettingsMatch(const TObjectPtr<UTextureGraphBase>& Parent) const;
-	void UpdateOutputSettingsFromGraph();
+	UE_API void CopyParamsToRuntimeGraph();
+	UE_API bool CheckOutputSettingsMatch(const TObjectPtr<UTextureGraphBase>& Parent) const;
+	UE_API void UpdateOutputSettingsFromGraph();
 	
-	void SetParent(const TObjectPtr<UTextureGraphBase>& Parent);
+	UE_API void SetParent(const TObjectPtr<UTextureGraphBase>& Parent);
 
 	// Override PreSave method of UObject
-	virtual void PreSave(FObjectPreSaveContext SaveContext) override;
+	UE_API virtual void PreSave(FObjectPreSaveContext SaveContext) override;
 	
 	virtual const UTG_Graph* Graph() const override { return RuntimeGraph.Get(); }
-	virtual UTG_Graph* Graph() override;
-	virtual void Initialize() override;
-};
\ No newline at end of file
+	UE_API virtual UTG_Graph* Graph() override;
+	UE_API virtual void Initialize() override;
+};
+
+#undef UE_API
diff --git a/Engine/Plugins/TextureGraph/Source/TextureGraphEditor/Public/TextureGraphEditorModule.h b/Engine/Plugins/TextureGraph/Source/TextureGraphEditor/Public/TextureGraphEditorModule.h
--- a/Engine/Plugins/TextureGraph/Source/TextureGraphEditor/Public/TextureGraphEditorModule.h
+++ b/Engine/Plugins/TextureGraph/Source/TextureGraphEditor/Public/TextureGraphEditorModule.h
@@ -5,6 +5,8 @@
 #include "Modules/ModuleManager.h"
 #include "Toolkits/IToolkit.h"
 #include "TG_Exporter.h"
+
+#define UE_API TEXTUREGRAPHEDITOR_API
 class IAssetTools;
 class IAssetTypeActions;
 class ITG_Editor;
@@ -16,7 +18,7 @@ typedef TArray<TSharedPtr<IAssetTypeActions>> AssetTypeActionsArray;
 extern const FName TG_EditorAppIdentifier;
 extern const FName TG_InstanceEditorAppIdentifier;
 
-class TEXTUREGRAPHEDITOR_API FTextureGraphEditorModule : public IModuleInterface
+class FTextureGraphEditorModule : public IModuleInterface
 {
 private:
 	/** All created asset type actions.  Cached here so that we can unregister them during shutdown. */
@@ -32,18 +34,20 @@ protected:
 	TSharedPtr<FTG_EditorGraphPanelPinFactory>	GraphPanelPinFactory;
 public:
 	/** IModuleInterface implementation */
-	virtual void					StartupModule() override;
-	virtual void					ShutdownModule() override;
+	UE_API virtual void					StartupModule() override;
+	UE_API virtual void					ShutdownModule() override;
 
 	
-	virtual void					StartTextureGraphEngine();
-	virtual void					ShutdownTextureGraphEngine();
-	bool							Tick(float deltaTime);
-	void							RegisterAssetTypeAction(IAssetTools& AssetTools, TSharedRef<IAssetTypeActions> Action);
-	void							UnRegisterAllAssetTypeActions();
-	static TSharedRef<ITG_Editor>	CreateTextureGraphEditor(const EToolkitMode::Type Mode, const TSharedPtr< class IToolkitHost >& InitToolkitHost, UTextureGraph* InTextureGraph);
-	static TSharedRef<ITG_Editor>	CreateTextureGraphInstanceEditor(const EToolkitMode::Type Mode, const TSharedPtr< class IToolkitHost >& InitToolkitHost, UTextureGraphInstance* InTextureGraphInstance);
+	UE_API virtual void					StartTextureGraphEngine();
+	UE_API virtual void					ShutdownTextureGraphEngine();
+	UE_API bool							Tick(float deltaTime);
+	UE_API void							RegisterAssetTypeAction(IAssetTools& AssetTools, TSharedRef<IAssetTypeActions> Action);
+	UE_API void							UnRegisterAllAssetTypeActions();
+	static UE_API TSharedRef<ITG_Editor>	CreateTextureGraphEditor(const EToolkitMode::Type Mode, const TSharedPtr< class IToolkitHost >& InitToolkitHost, UTextureGraph* InTextureGraph);
+	static UE_API TSharedRef<ITG_Editor>	CreateTextureGraphInstanceEditor(const EToolkitMode::Type Mode, const TSharedPtr< class IToolkitHost >& InitToolkitHost, UTextureGraphInstance* InTextureGraphInstance);
 
 	/** Returns a reference to the Blueprint Debugger state object */
 	const TUniquePtr<FTG_ExporterUtility>& GetTextureExporter() const { return TG_Exporter; }
 };
+
+#undef UE_API
diff --git a/Engine/Plugins/TextureGraph/Source/TextureGraphEngine/2D/TargetTextureSet.h b/Engine/Plugins/TextureGraph/Source/TextureGraphEngine/2D/TargetTextureSet.h
--- a/Engine/Plugins/TextureGraph/Source/TextureGraphEngine/2D/TargetTextureSet.h
+++ b/Engine/Plugins/TextureGraph/Source/TextureGraphEngine/2D/TargetTextureSet.h
@@ -6,6 +6,8 @@
 #include "2D/TextureType.h"
 #include "FxMat/RenderMaterial_BP.h"
 
+#define UE_API TEXTUREGRAPHENGINE_API
+
 struct FMaterialMappingInfo;
 typedef TFunction<void(TiledBlobPtr)> TextureReadyCallback;
 
@@ -13,7 +15,7 @@ class RenderMesh;
 typedef std::shared_ptr<RenderMesh>	RenderMeshPtr;
 
 
-class TEXTUREGRAPHENGINE_API TargetTextureSet : public TextureSet
+class TargetTextureSet : public TextureSet
 {
 protected:
 
@@ -29,23 +31,23 @@ protected:
 	TMap<FName, TiledBlobRef>		BoundTexturesMap;			///Textures currently bound and dont want deallocation until new textures assigned
 	int32							RenderCount = 0;			/// The number of times this layer has been rendered to
 
-	virtual void					InitTex(int32 InTypeIndex) override;
-	virtual void					BindOnTextureUpdate(RenderMaterial_BPPtr InMaterial, FMaterialMappingInfo MaterialMappingInfo) const;
-	virtual void					RegisterCallback(TextureReadyCallback Callback, FMaterialMappingInfo MaterialMappingInfo) const;
-	virtual void					UnRegisterCallback(FName TextureName) const;
+	UE_API virtual void					InitTex(int32 InTypeIndex) override;
+	UE_API virtual void					BindOnTextureUpdate(RenderMaterial_BPPtr InMaterial, FMaterialMappingInfo MaterialMappingInfo) const;
+	UE_API virtual void					RegisterCallback(TextureReadyCallback Callback, FMaterialMappingInfo MaterialMappingInfo) const;
+	UE_API virtual void					UnRegisterCallback(FName TextureName) const;
 public:
-									TargetTextureSet(int32 InId, const FString& InName, RenderMeshPtr InMesh, int32 InWidth, int32 InHeight);
-	virtual							~TargetTextureSet() override;
+									UE_API TargetTextureSet(int32 InId, const FString& InName, RenderMeshPtr InMesh, int32 InWidth, int32 InHeight);
+	UE_API virtual							~TargetTextureSet() override;
 
-	void							SetMesh(RenderMeshPtr InMesh);
+	UE_API void							SetMesh(RenderMeshPtr InMesh);
 	
-	AsyncBufferResultPtr			BindTo(RenderMaterial_BPPtr Material, TArray<struct FMaterialMappingInfo> MaterialMappingInfos);
-	AsyncBufferResultPtr			BindTo(RenderMaterial_BPPtr Material, struct FMaterialMappingInfo MaterialMappingInfo);
-	AsyncBufferResultPtr			UnbindFrom(RenderMaterial_BPPtr Material);
+	UE_API AsyncBufferResultPtr			BindTo(RenderMaterial_BPPtr Material, TArray<struct FMaterialMappingInfo> MaterialMappingInfos);
+	UE_API AsyncBufferResultPtr			BindTo(RenderMaterial_BPPtr Material, struct FMaterialMappingInfo MaterialMappingInfo);
+	UE_API AsyncBufferResultPtr			UnbindFrom(RenderMaterial_BPPtr Material);
 	
-	virtual void					SetTexture(FName TextureName, TiledBlobRef InTexture) override;
-	virtual void					FreeAt(FName TextureName) override;
-	virtual void					Free() override;
+	UE_API virtual void					SetTexture(FName TextureName, TiledBlobRef InTexture) override;
+	UE_API virtual void					FreeAt(FName TextureName) override;
+	UE_API virtual void					Free() override;
 	//////////////////////////////////////////////////////////////////////////
 	/// Inline functions
 	//////////////////////////////////////////////////////////////////////////
@@ -59,3 +61,4 @@ public:
 typedef std::unique_ptr<TargetTextureSet>	TargetTextureSetPtr;
 typedef std::vector<TargetTextureSetPtr>	TargetTextureSetPtrVec;
 
+#undef UE_API
diff --git a/Engine/Plugins/TextureGraph/Source/TextureGraphEngine/2D/Tex.h b/Engine/Plugins/TextureGraph/Source/TextureGraphEngine/2D/Tex.h
--- a/Engine/Plugins/TextureGraph/Source/TextureGraphEngine/2D/Tex.h
+++ b/Engine/Plugins/TextureGraph/Source/TextureGraphEngine/2D/Tex.h
@@ -13,6 +13,8 @@ THIRD_PARTY_INCLUDES_START
 THIRD_PARTY_INCLUDES_END
 #include "Helper/Promise.h"
 
+#define UE_API TEXTUREGRAPHENGINE_API
+
 class UTexture2D;
 class UTextureRenderTarget2D;
 class UMaterial;
@@ -81,15 +83,15 @@ struct TexDescriptor
  * This works closely with the RCManager to efficiently manage the resources available to us
  * within the system.
  */
-class TEXTUREGRAPHENGINE_API Tex : public FGCObject
+class Tex : public FGCObject
 {
 	friend class Device_FX;
 
 	//////////////////////////////////////////////////////////////////////////
 	/// FGCObject
 	//////////////////////////////////////////////////////////////////////////
-	virtual void					AddReferencedObjects(FReferenceCollector& Collector) override;
-	virtual FString					GetReferencerName() const override;
+	UE_API virtual void					AddReferencedObjects(FReferenceCollector& Collector) override;
+	UE_API virtual FString					GetReferencerName() const override;
 
 public:
 	
@@ -103,71 +105,71 @@ protected:
 	CHashPtr						HashObj;						/// The hash for this Texture. This is totally optional
 	bool							bNoCache = false;				/// Is this not supposed to be cached
 
-	void							InitTexture(const uint8* Pixels, size_t Length);
-	UTexture2D*						CreateTexture(uint32 Width, uint32 Height, EPixelFormat PixelFormat, bool sRGB, UObject* Package);
-	UTexture2D*						InitTextureDefault(int32 Width, int32 Height, EPixelFormat PixelFormat, const uint8* uncompressedData, size_t uncompressedDataSize, UObject* Package);
+	UE_API void							InitTexture(const uint8* Pixels, size_t Length);
+	UE_API UTexture2D*						CreateTexture(uint32 Width, uint32 Height, EPixelFormat PixelFormat, bool sRGB, UObject* Package);
+	UE_API UTexture2D*						InitTextureDefault(int32 Width, int32 Height, EPixelFormat PixelFormat, const uint8* uncompressedData, size_t uncompressedDataSize, UObject* Package);
 
-	bool							CopyImageToBuffer(EPixelFormat& PixelFormat, int32& Width, int32& Height, TArray<uint8>& inputBuffer, TArray<uint8>& OutputBuffer);
+	UE_API bool							CopyImageToBuffer(EPixelFormat& PixelFormat, int32& Width, int32& Height, TArray<uint8>& inputBuffer, TArray<uint8>& OutputBuffer);
 
-	UTexture2D*						InitTextureHDR(const TArray<uint8>& Buffer, UPackage* Package);
+	UE_API UTexture2D*						InitTextureHDR(const TArray<uint8>& Buffer, UPackage* Package);
 
 	
-	void							InvalidateCached();
-	RenderMaterialPtr				GetDefaultMaterial();
-	void							UpdateRaw(const uint8* SrcPixels, size_t Length);
-	bool							IsValidVirtualTexture();
+	UE_API void							InvalidateCached();
+	UE_API RenderMaterialPtr				GetDefaultMaterial();
+	UE_API void							UpdateRaw(const uint8* SrcPixels, size_t Length);
+	UE_API bool							IsValidVirtualTexture();
 public:
-	static EObjectFlags				Flags;
+	static UE_API EObjectFlags				Flags;
 	
-									Tex(int32 Width, int32 Height, EPixelFormat PixelFormat);
-									explicit Tex(const TexDescriptor& InDesc);
-									explicit Tex(UTexture2D* Texture);
-									explicit Tex(RawBufferPtr RawObj);
-									explicit Tex(UTextureRenderTarget2D* RT);
-									Tex();
-	virtual							~Tex() override;
+									UE_API Tex(int32 Width, int32 Height, EPixelFormat PixelFormat);
+									UE_API explicit Tex(const TexDescriptor& InDesc);
+									UE_API explicit Tex(UTexture2D* Texture);
+									UE_API explicit Tex(RawBufferPtr RawObj);
+									UE_API explicit Tex(UTextureRenderTarget2D* RT);
+									UE_API Tex();
+	UE_API virtual							~Tex() override;
 
 	//////////////////////////////////////////////////////////////////////////
 	/// Others
 	//////////////////////////////////////////////////////////////////////////
-	static void						FreeGenericTexture(UTexture** Texture);
-	static void						FreeTexture(UTexture2D** Texture);
-	static void						FreeRT(UTextureRenderTarget2D** RT);
+	static UE_API void						FreeGenericTexture(UTexture** Texture);
+	static UE_API void						FreeTexture(UTexture2D** Texture);
+	static UE_API void						FreeRT(UTextureRenderTarget2D** RT);
 	
-	void							InitRT(bool ForceFloat = false);
-
-	void							Release();
-	void							ReleaseRT();
-	void							ReleaseTexture();
-	void							TransferTextureToRT(FRHICommandListImmediate& RHI, UTexture2D** PrevTexture, bool FreeAfterUse);
-	void							TransferVirtualTextureToRT(FRHICommandListImmediate& RHI, UTexture2D** TextureToTransfer, bool FreeAfterUse);
-	bool							LoadAsset(FSoftObjectPath& SoftPath, const DesiredImageProperties* Props = nullptr);
-	bool							LoadFile(const FString& Filename, const DesiredImageProperties* Props = nullptr);
-	AsyncActionResultPtr			LoadFlat();
-	AsyncTiledBlobRef				ToBlob(int32 XTiles, int32 YTiles, uint32 Width = 0, uint32 Height = 0, bool TransferToRT = false); ///There is a visual difference between having material transfer to RT or directly copy from Texture
-	AsyncTiledBlobRef				ToSingleBlob(CHashPtr Hash, bool TransferToRT = false, bool ResolveOnRenderThread = false, bool NoCache = false);
-	AsyncActionResultPtr			LoadRaw(RawBufferPtr RawObj);
-	void							UpdateRaw(RawBufferPtr RawObj);
-
-	void							SetFilter(TextureFilter FilterValue);
+	UE_API void							InitRT(bool ForceFloat = false);
+
+	UE_API void							Release();
+	UE_API void							ReleaseRT();
+	UE_API void							ReleaseTexture();
+	UE_API void							TransferTextureToRT(FRHICommandListImmediate& RHI, UTexture2D** PrevTexture, bool FreeAfterUse);
+	UE_API void							TransferVirtualTextureToRT(FRHICommandListImmediate& RHI, UTexture2D** TextureToTransfer, bool FreeAfterUse);
+	UE_API bool							LoadAsset(FSoftObjectPath& SoftPath, const DesiredImageProperties* Props = nullptr);
+	UE_API bool							LoadFile(const FString& Filename, const DesiredImageProperties* Props = nullptr);
+	UE_API AsyncActionResultPtr			LoadFlat();
+	UE_API AsyncTiledBlobRef				ToBlob(int32 XTiles, int32 YTiles, uint32 Width = 0, uint32 Height = 0, bool TransferToRT = false); ///There is a visual difference between having material transfer to RT or directly copy from Texture
+	UE_API AsyncTiledBlobRef				ToSingleBlob(CHashPtr Hash, bool TransferToRT = false, bool ResolveOnRenderThread = false, bool NoCache = false);
+	UE_API AsyncActionResultPtr			LoadRaw(RawBufferPtr RawObj);
+	UE_API void							UpdateRaw(RawBufferPtr RawObj);
+
+	UE_API void							SetFilter(TextureFilter FilterValue);
 	
-	virtual void					Bind(FName Name, std::shared_ptr<RenderMaterial> Material) const;
-	void							Clear();
-	void							Clear(FRHICommandList& RHI);
-	void							Clear(FRHICommandList& RHI, FLinearColor Color);
+	UE_API virtual void					Bind(FName Name, std::shared_ptr<RenderMaterial> Material) const;
+	UE_API void							Clear();
+	UE_API void							Clear(FRHICommandList& RHI);
+	UE_API void							Clear(FRHICommandList& RHI, FLinearColor Color);
 
-	void							Free();
-	size_t							GetMemSize() const;
-	RawBufferPtr					Raw(const BufferDescriptor* SrcDesc = nullptr) const; // Create a RawBuffer from the Texture, generate the descriptor if none is proveided
-	void							GenerateMips();
+	UE_API void							Free();
+	UE_API size_t							GetMemSize() const;
+	UE_API RawBufferPtr					Raw(const BufferDescriptor* SrcDesc = nullptr) const; // Create a RawBuffer from the Texture, generate the descriptor if none is proveided
+	UE_API void							GenerateMips();
 
 	virtual bool					IsArray()  { return false; }
-	virtual FRHITexture*			GetRHITexture() const;
-	virtual UTexture*				GetTexture() const;
+	UE_API virtual FRHITexture*			GetRHITexture() const;
+	UE_API virtual UTexture*				GetTexture() const;
 
-	virtual bool					IsNull() const;
+	UE_API virtual bool					IsNull() const;
 
-	TArray<FColor>					ReadPixels();
+	UE_API TArray<FColor>					ReadPixels();
 
 	//////////////////////////////////////////////////////////////////////////
 	/// Static functions
@@ -180,7 +182,7 @@ public:
 	 * @param	Filename	 filename with extension to save in particular format
 	 * @param	bIsHDR		whether to save RTF_16 in HDR or EXR (discarded when trying to save PF_B8G8R8A8)	 
 	 */
-	static	bool					SaveImage(UTextureRenderTarget2D* RT, const FString& Path, const FString& Filename, bool bIsHDR = true);
+	static	UE_API bool					SaveImage(UTextureRenderTarget2D* RT, const FString& Path, const FString& Filename, bool bIsHDR = true);
 
 	//////////////////////////////////////////////////////////////////////////
 	/// Inline functions
@@ -206,3 +208,4 @@ public:
 
 //typedef std::shared_ptr<Tex>		TexPtr;
 
+#undef UE_API
diff --git a/Engine/Plugins/TextureGraph/Source/TextureGraphEngine/2D/TexArray.h b/Engine/Plugins/TextureGraph/Source/TextureGraphEngine/2D/TexArray.h
--- a/Engine/Plugins/TextureGraph/Source/TextureGraphEngine/2D/TexArray.h
+++ b/Engine/Plugins/TextureGraph/Source/TextureGraphEngine/2D/TexArray.h
@@ -5,6 +5,8 @@
 #include "Tex.h"
 #include "Helper/Promise.h"
 
+#define UE_API TEXTUREGRAPHENGINE_API
+
 class UTexture2D;
 class UTextureRenderTarget2DArray;
 
@@ -12,7 +14,7 @@ class TexArray;
 
 typedef std::shared_ptr<TexArray>		TexArrayPtr;
 
-class TEXTUREGRAPHENGINE_API TexArray : public Tex
+class TexArray : public Tex
 {
 	friend class Device_FX;
 
@@ -21,22 +23,24 @@ protected:
 	int32							XTiles = 0;
 	int32							YTiles = 0;
 	int32							Index = 0;
-	void							InitRTArray(bool ForceFloat = false);
+	UE_API void							InitRTArray(bool ForceFloat = false);
 
 public:
-									TexArray();
-									TexArray(Tex& TexObj, int32 XTilesTotal, int32 YTilesTotal);
-									TexArray(const TexDescriptor& Desc, UTextureRenderTarget2DArray* rtArray);
-									TexArray(const TexDescriptor& Desc,  int32 XTilesTotal, int32 YTilesTotal);
-	virtual							~TexArray() override;
-	virtual void					AddReferencedObjects(FReferenceCollector& Collector) override;
+									UE_API TexArray();
+									UE_API TexArray(Tex& TexObj, int32 XTilesTotal, int32 YTilesTotal);
+									UE_API TexArray(const TexDescriptor& Desc, UTextureRenderTarget2DArray* rtArray);
+									UE_API TexArray(const TexDescriptor& Desc,  int32 XTilesTotal, int32 YTilesTotal);
+	UE_API virtual							~TexArray() override;
+	UE_API virtual void					AddReferencedObjects(FReferenceCollector& Collector) override;
 
 	virtual bool					IsArray() override { return true; }
-	virtual UTexture*				GetTexture() const override;
-	virtual bool					IsNull() const override;
+	UE_API virtual UTexture*				GetTexture() const override;
+	UE_API virtual bool					IsNull() const override;
 
 	//////////////////////////////////////////////////////////////////////////
 	/// Inline functions
 	//////////////////////////////////////////////////////////////////////////
 	FORCEINLINE UTextureRenderTarget2DArray* GetRenderTargetArray() const { return RTArray; }
 };
+
+#undef UE_API
diff --git a/Engine/Plugins/TextureGraph/Source/TextureGraphEngine/2D/TextureHelper.h b/Engine/Plugins/TextureGraph/Source/TextureGraphEngine/2D/TextureHelper.h
--- a/Engine/Plugins/TextureGraph/Source/TextureGraphEngine/2D/TextureHelper.h
+++ b/Engine/Plugins/TextureGraph/Source/TextureGraphEngine/2D/TextureHelper.h
@@ -11,6 +11,8 @@
 #include "TextureContent.h"
 #include "TextureType.h"
 
+#define UE_API TEXTUREGRAPHENGINE_API
+
 class Tex;
 class UTextureRenderTarget2D;
 class UTexture2D;
@@ -23,24 +25,24 @@ typedef std::shared_ptr<TiledBlob>		TiledBlobPtr;
 typedef cti::continuable<bool>			AsyncBool;
 typedef cti::continuable<FLinearColor>	AsyncLinearColor;
 
-struct TEXTUREGRAPHENGINE_API TextureHelper
+struct TextureHelper
 {
-	static TiledBlobPtr					GTransparent;			/// Transparent RGBA texture
-	static TiledBlobPtr					GBlack;					/// Black RGBA texture
-	static TiledBlobPtr					GWhite;					/// White RGBA texture
-	static TiledBlobPtr					GGray;					/// Gray RGBA texture
-	static TiledBlobPtr					GRed;					/// Red RGBA texture
-	static TiledBlobPtr					GGreen;					/// Green RGBA texture
-	static TiledBlobPtr					GBlue;					/// Blue RGBA texture
-	static TiledBlobPtr					GYellow;				/// Yellow RGBA texture
-	static TiledBlobPtr					GMagenta;				/// Magenta RGBA texture
-	static TiledBlobPtr					GDefaultNormal;			/// Default Normal RGBA texture
-	static TiledBlobPtr					GWhiteMask;				/// Default white mask
-	static TiledBlobPtr					GBlackMask;				/// Default black mask
-
-	static cti::continuable<bool>		InitSolidTexture(TiledBlobPtr* BlobObj, FLinearColor Color, FString Name, struct TexDescriptor Desc);
-	static void							InitStockTextures();
-	static void							FreeStockTextures();
+	static UE_API TiledBlobPtr					GTransparent;			/// Transparent RGBA texture
+	static UE_API TiledBlobPtr					GBlack;					/// Black RGBA texture
+	static UE_API TiledBlobPtr					GWhite;					/// White RGBA texture
+	static UE_API TiledBlobPtr					GGray;					/// Gray RGBA texture
+	static UE_API TiledBlobPtr					GRed;					/// Red RGBA texture
+	static UE_API TiledBlobPtr					GGreen;					/// Green RGBA texture
+	static UE_API TiledBlobPtr					GBlue;					/// Blue RGBA texture
+	static UE_API TiledBlobPtr					GYellow;				/// Yellow RGBA texture
+	static UE_API TiledBlobPtr					GMagenta;				/// Magenta RGBA texture
+	static UE_API TiledBlobPtr					GDefaultNormal;			/// Default Normal RGBA texture
+	static UE_API TiledBlobPtr					GWhiteMask;				/// Default white mask
+	static UE_API TiledBlobPtr					GBlackMask;				/// Default black mask
+
+	static UE_API cti::continuable<bool>		InitSolidTexture(TiledBlobPtr* BlobObj, FLinearColor Color, FString Name, struct TexDescriptor Desc);
+	static UE_API void							InitStockTextures();
+	static UE_API void							FreeStockTextures();
 
 	static FORCEINLINE TiledBlobPtr		GetTransparent() { check(GTransparent); return GTransparent; }
 	static FORCEINLINE TiledBlobPtr		GetBlack() { check(GBlack); return GBlack; }
@@ -57,48 +59,48 @@ struct TEXTUREGRAPHENGINE_API TextureHelper
 	static FORCEINLINE TiledBlobPtr		GetBlackMask() { check(GBlackMask); return GBlackMask; }
 
 	//////////////////////////////////////////////////////////////////////////
-	static uint32						GetChannelsFromPixelFormat(EPixelFormat InPixelFormat);
-	static uint32						GetBppFromPixelFormat(EPixelFormat InPixelFormat);
-	static ETextureRenderTargetFormat	GetRenderTargetFormatFromPixelFormat(EPixelFormat InPixelFormat);
-	static ETextureSourceFormat			GetSourceFormat(ETG_TextureFormat TGTextureFormat);
-	static bool							GetBufferFormatAndChannelsFromTGTextureFormat(ETG_TextureFormat Format, BufferFormat& OutBufferFormat, uint32& OutBufferChannels);
-	static EPixelFormat					GetPixelFormatFromRenderTargetFormat(ETextureRenderTargetFormat RTFormat);
-	static bool							GetPixelFormatFromTextureSourceFormat(ETextureSourceFormat SourceFormat, EPixelFormat& OutPixelFormat, uint32& OutNumChannels);
-	static ETextureSourceFormat			GetTextureSourceFormat(BufferFormat Format, uint32 ItemsPerPoint);
-	static ETG_TextureFormat			GetTGTextureFormatFromChannelsAndFormat(uint32 ItemsPerPoint, BufferFormat Format);
-	static uint32						GetNumChannelsFromTGTextureFormat(ETG_TextureFormat TextureFormat);
-	static FString						GetChannelsTextFromItemsPerPoint(const int32 InItemsPerPoint);	
-	static void							ClearRT(UTextureRenderTarget2D* RenderTarget, FLinearColor Color = FLinearColor::Transparent);
-	static void							ClearRT(FRHICommandList& RHI, UTextureRenderTarget2D* RenderTarget, FLinearColor Color = FLinearColor::Transparent);
-
-	static const char*					TextureTypeToString(TextureType Type);
-	static TextureType					TextureStringToType(const FString& TypeString);
-	static const char*					TextureTypeToMegascansType(TextureType Type);
-
-	static TextureType					TextureContentToTextureType(TextureContent Content);
-	static TextureContent				TextureTypeToTextureContent(TextureType Type);
-	static MaskType						TextureContentToMaskType(TextureContent Content);
-	static TextureContent				MaskTypeToTextureContent(MaskType Type);
-	static MaskModifierType				TextureContentToMaskModifier(TextureContent Content);
-	static TextureContent				MaskModifierToTextureContent(MaskModifierType Type);
-
-	static RawBufferPtr					RawFromRT(UTextureRenderTarget2D* RenderTarget, const BufferDescriptor& Desc);
-	static RawBufferPtr					RawFromTexture(UTexture2D* Texture, const BufferDescriptor& Desc);
-	static RawBufferPtr					RawFromResource(const FTextureRHIRef& ResourceRHI, const BufferDescriptor& Desc);
-	static BufferFormat					FindOptimalSupportedFormat(BufferFormat SrcFormat);
-
-	static void							RawFromRT_Tiled(UTextureRenderTarget2D* RenderTarget, const BufferDescriptor& Desc, size_t TileSizeX, size_t TileSizeY, RawBufferPtrTiles& Tiles);
-	static void							RawFromTexture_Tiled(UTexture2D* Texture, const BufferDescriptor& Desc, size_t TileSizeX, size_t TileSizeY, RawBufferPtrTiles& Tiles);
-	static void							RawFromResource_Tiled(FTextureRHIRef ResourceRHI, const BufferDescriptor& Desc, size_t TileSizeX, size_t TileSizeY, RawBufferPtrTiles& Tiles);
-	static void							RawFromMem_Tiled(const uint8* SrcData, size_t SrcDataLength, const BufferDescriptor& SrcDesc, size_t TileSizeX, size_t TileSizeY, RawBufferPtrTiles& Tiles);
-	static RawBufferPtr					CombineRaw_Tiles(const RawBufferPtrTiles& Tiles, CHashPtr HashValue = nullptr, bool bIsTransient = false);
-
-	static bool							IsFloatRT(UTextureRenderTarget2D* RenderTarget);
-	static FLinearColor					GetPixelValueFromRaw(RawBufferPtr RawObj, int32 Width, int32 Height, int32 X, int32 Y);
-	static AsyncBool					ExportRaw(RawBufferPtr RawObj, const FString& CompletePath);
-	static bool							CanSupportTexture(UTexture* Tex);
-	static bool							CanSplitToTiles(int Width, int Height, int TilesX, int TilesY);
-	static size_t						RoundUpTo(size_t Size, size_t DesiredRounding);
+	static UE_API uint32						GetChannelsFromPixelFormat(EPixelFormat InPixelFormat);
+	static UE_API uint32						GetBppFromPixelFormat(EPixelFormat InPixelFormat);
+	static UE_API ETextureRenderTargetFormat	GetRenderTargetFormatFromPixelFormat(EPixelFormat InPixelFormat);
+	static UE_API ETextureSourceFormat			GetSourceFormat(ETG_TextureFormat TGTextureFormat);
+	static UE_API bool							GetBufferFormatAndChannelsFromTGTextureFormat(ETG_TextureFormat Format, BufferFormat& OutBufferFormat, uint32& OutBufferChannels);
+	static UE_API EPixelFormat					GetPixelFormatFromRenderTargetFormat(ETextureRenderTargetFormat RTFormat);
+	static UE_API bool							GetPixelFormatFromTextureSourceFormat(ETextureSourceFormat SourceFormat, EPixelFormat& OutPixelFormat, uint32& OutNumChannels);
+	static UE_API ETextureSourceFormat			GetTextureSourceFormat(BufferFormat Format, uint32 ItemsPerPoint);
+	static UE_API ETG_TextureFormat			GetTGTextureFormatFromChannelsAndFormat(uint32 ItemsPerPoint, BufferFormat Format);
+	static UE_API uint32						GetNumChannelsFromTGTextureFormat(ETG_TextureFormat TextureFormat);
+	static UE_API FString						GetChannelsTextFromItemsPerPoint(const int32 InItemsPerPoint);	
+	static UE_API void							ClearRT(UTextureRenderTarget2D* RenderTarget, FLinearColor Color = FLinearColor::Transparent);
+	static UE_API void							ClearRT(FRHICommandList& RHI, UTextureRenderTarget2D* RenderTarget, FLinearColor Color = FLinearColor::Transparent);
+
+	static UE_API const char*					TextureTypeToString(TextureType Type);
+	static UE_API TextureType					TextureStringToType(const FString& TypeString);
+	static UE_API const char*					TextureTypeToMegascansType(TextureType Type);
+
+	static UE_API TextureType					TextureContentToTextureType(TextureContent Content);
+	static UE_API TextureContent				TextureTypeToTextureContent(TextureType Type);
+	static UE_API MaskType						TextureContentToMaskType(TextureContent Content);
+	static UE_API TextureContent				MaskTypeToTextureContent(MaskType Type);
+	static UE_API MaskModifierType				TextureContentToMaskModifier(TextureContent Content);
+	static UE_API TextureContent				MaskModifierToTextureContent(MaskModifierType Type);
+
+	static UE_API RawBufferPtr					RawFromRT(UTextureRenderTarget2D* RenderTarget, const BufferDescriptor& Desc);
+	static UE_API RawBufferPtr					RawFromTexture(UTexture2D* Texture, const BufferDescriptor& Desc);
+	static UE_API RawBufferPtr					RawFromResource(const FTextureRHIRef& ResourceRHI, const BufferDescriptor& Desc);
+	static UE_API BufferFormat					FindOptimalSupportedFormat(BufferFormat SrcFormat);
+
+	static UE_API void							RawFromRT_Tiled(UTextureRenderTarget2D* RenderTarget, const BufferDescriptor& Desc, size_t TileSizeX, size_t TileSizeY, RawBufferPtrTiles& Tiles);
+	static UE_API void							RawFromTexture_Tiled(UTexture2D* Texture, const BufferDescriptor& Desc, size_t TileSizeX, size_t TileSizeY, RawBufferPtrTiles& Tiles);
+	static UE_API void							RawFromResource_Tiled(FTextureRHIRef ResourceRHI, const BufferDescriptor& Desc, size_t TileSizeX, size_t TileSizeY, RawBufferPtrTiles& Tiles);
+	static UE_API void							RawFromMem_Tiled(const uint8* SrcData, size_t SrcDataLength, const BufferDescriptor& SrcDesc, size_t TileSizeX, size_t TileSizeY, RawBufferPtrTiles& Tiles);
+	static UE_API RawBufferPtr					CombineRaw_Tiles(const RawBufferPtrTiles& Tiles, CHashPtr HashValue = nullptr, bool bIsTransient = false);
+
+	static UE_API bool							IsFloatRT(UTextureRenderTarget2D* RenderTarget);
+	static UE_API FLinearColor					GetPixelValueFromRaw(RawBufferPtr RawObj, int32 Width, int32 Height, int32 X, int32 Y);
+	static UE_API AsyncBool					ExportRaw(RawBufferPtr RawObj, const FString& CompletePath);
+	static UE_API bool							CanSupportTexture(UTexture* Tex);
+	static UE_API bool							CanSplitToTiles(int Width, int Height, int TilesX, int TilesY);
+	static UE_API size_t						RoundUpTo(size_t Size, size_t DesiredRounding);
 
 	//////////////////////////////////////////////////////////////////////////
 	/// Inline functions
@@ -148,3 +150,5 @@ struct TEXTUREGRAPHENGINE_API TextureHelper
 		return (uint32)(ceil(log2(double(BaseSize))) + 1.0);
 	}
 };
+
+#undef UE_API
diff --git a/Engine/Plugins/TextureGraph/Source/TextureGraphEngine/2D/TextureSet.h b/Engine/Plugins/TextureGraph/Source/TextureGraphEngine/2D/TextureSet.h
--- a/Engine/Plugins/TextureGraph/Source/TextureGraphEngine/2D/TextureSet.h
+++ b/Engine/Plugins/TextureGraph/Source/TextureGraphEngine/2D/TextureSet.h
@@ -6,14 +6,16 @@
 #include "Data/RawBuffer.h"
 #include "Data/Blobber.h"
 
+#define UE_API TEXTUREGRAPHENGINE_API
+
 class MixUpdateCycle;
 typedef std::shared_ptr<MixUpdateCycle> MixUpdateCyclePtr;
 
-class TEXTUREGRAPHENGINE_API TextureSet
+class TextureSet
 {
 public:
 	static constexpr int			GMaxTextures = (int)TextureType::Count;
-	static const BufferDescriptor	GDesc[GMaxTextures];		/// The default descriptors for the textures
+	static UE_API const BufferDescriptor	GDesc[GMaxTextures];		/// The default descriptors for the textures
 	static constexpr int32			GDefaultTexSize = 256;		/// Default texture size
 	
 protected:
@@ -26,32 +28,32 @@ public:
 	TMap<FName,TiledBlobRef>		TexturesMap;				/// The textures within this texture set that we don't own a reference of (could be one of the owned textures)
 	BufferDescriptor				Desc[GMaxTextures];			/// Descriptors for all the blobs
 	
-	virtual void					InitTex(int32 InTypeIndex);
-	virtual void					InitTexFromFile(int32 InTypeIndex, const FString& FileName, int32 NumTilesX, int32 NumTilesY);
-	void							InitBufferDescriptor(int32 ti);
-	TSet<FName>						GetTextureList() const;
+	UE_API virtual void					InitTex(int32 InTypeIndex);
+	UE_API virtual void					InitTexFromFile(int32 InTypeIndex, const FString& FileName, int32 NumTilesX, int32 NumTilesY);
+	UE_API void							InitBufferDescriptor(int32 ti);
+	UE_API TSet<FName>						GetTextureList() const;
 
 public:
-									TextureSet(int32 InWidth, int32 InHeight);
+									UE_API TextureSet(int32 InWidth, int32 InHeight);
 
 
-	virtual							~TextureSet();
+	UE_API virtual							~TextureSet();
 
-	virtual void					Init();
-	virtual void					Free();
-	virtual void					FreeAt(FName TextureName);
-	void							FreeAt(TextureType type);
+	UE_API virtual void					Init();
+	UE_API virtual void					Free();
+	UE_API virtual void					FreeAt(FName TextureName);
+	UE_API void							FreeAt(TextureType type);
 
-	virtual void					ChangeFormat(TextureType InType, EPixelFormat InPixelFormat, bool bRecreate);
-	virtual void					ChangeResolution(int32 InWidth, int32 InHeight, bool bRecreate);
-	virtual void					Invalidate(FName TextureName, bool bRecreate);
-	virtual void					Invalidate(TextureType InType, bool bRecreate);
+	UE_API virtual void					ChangeFormat(TextureType InType, EPixelFormat InPixelFormat, bool bRecreate);
+	UE_API virtual void					ChangeResolution(int32 InWidth, int32 InHeight, bool bRecreate);
+	UE_API virtual void					Invalidate(FName TextureName, bool bRecreate);
+	UE_API virtual void					Invalidate(TextureType InType, bool bRecreate);
 
-	virtual void					SetTexture(FName TextureName, TiledBlobRef Texture);
-	virtual bool					RenameTexture(FName TextureName, FName NewName);
-	virtual void					SetTexture(int32 TypeIndex, TiledBlobRef texture);
+	UE_API virtual void					SetTexture(FName TextureName, TiledBlobRef Texture);
+	UE_API virtual bool					RenameTexture(FName TextureName, FName NewName);
+	UE_API virtual void					SetTexture(int32 TypeIndex, TiledBlobRef texture);
 	
-	virtual void					Update(MixUpdateCyclePtr cycle, int32 targetId);
+	UE_API virtual void					Update(MixUpdateCyclePtr cycle, int32 targetId);
 
 	
 	FORCEINLINE bool				ContainsTexture(FName TextureName) { return TexturesMap.Contains(TextureName); }
@@ -89,3 +91,5 @@ public:
 
 typedef std::unique_ptr<TextureSet>	TextureSetPtr;
 typedef std::vector<TextureSet>		TextureSetVec;
+
+#undef UE_API
diff --git a/Engine/Plugins/TextureGraph/Source/TextureGraphEngine/3D/CoreMesh.h b/Engine/Plugins/TextureGraph/Source/TextureGraphEngine/3D/CoreMesh.h
--- a/Engine/Plugins/TextureGraph/Source/TextureGraphEngine/3D/CoreMesh.h
+++ b/Engine/Plugins/TextureGraph/Source/TextureGraphEngine/3D/CoreMesh.h
@@ -13,7 +13,7 @@
 #include "RHIUtilities.h"
 
 //////////////////////////////////////////////////////////////////////////
-struct TEXTUREGRAPHENGINE_API CoreMeshVertex
+struct CoreMeshVertex
 {
 	FVector3f					position;			/// The vertex position
 	FVector2f					uv;					/// The texture coordinates
@@ -22,7 +22,7 @@ struct TEXTUREGRAPHENGINE_API CoreMeshVertex
 	FVector4f					color;				/// Per vertex colour
 };
 
-struct TEXTUREGRAPHENGINE_API CoreMesh
+struct CoreMesh
 {
 	TArray<FVector>				vertices;
 	TArray<FVector>				normals;
diff --git a/Engine/Plugins/TextureGraph/Source/TextureGraphEngine/3D/MaterialInfo.h b/Engine/Plugins/TextureGraph/Source/TextureGraphEngine/3D/MaterialInfo.h
--- a/Engine/Plugins/TextureGraph/Source/TextureGraphEngine/3D/MaterialInfo.h
+++ b/Engine/Plugins/TextureGraph/Source/TextureGraphEngine/3D/MaterialInfo.h
@@ -6,7 +6,7 @@
 /**
  * 
  */
-struct TEXTUREGRAPHENGINE_API MaterialInfo
+struct MaterialInfo
 {
 	FString							name;
 	int32							id;
diff --git a/Engine/Plugins/TextureGraph/Source/TextureGraphEngine/3D/MeshDetails/MeshDetails.h b/Engine/Plugins/TextureGraph/Source/TextureGraphEngine/3D/MeshDetails/MeshDetails.h
--- a/Engine/Plugins/TextureGraph/Source/TextureGraphEngine/3D/MeshDetails/MeshDetails.h
+++ b/Engine/Plugins/TextureGraph/Source/TextureGraphEngine/3D/MeshDetails/MeshDetails.h
@@ -7,6 +7,8 @@
 #include "Helper/DataUtil.h"
 #include <array>
 
+#define UE_API TEXTUREGRAPHENGINE_API
+
 class MeshInfo;
 class MeshDetails;
 typedef cti::continuable<MeshDetails*>	MeshDetailsPAsync;
@@ -15,24 +17,26 @@ class RawBuffer;
 typedef std::shared_ptr<RawBuffer>		RawBufferPtr;
 
 //////////////////////////////////////////////////////////////////////////
-class TEXTUREGRAPHENGINE_API MeshDetails
+class MeshDetails
 {
 protected:
 	MeshInfo*					Mesh;					/// The mesh that is containing this detail
 	bool						bIsFinalised = false;		/// Whether this detail has been finalised or not
 	CHashPtr					HashValue;					/// What is the hash of this detail
 
-	virtual void				CalculateTri(size_t ti);
-	virtual void				CalculateVertex(size_t vi);
+	UE_API virtual void				CalculateTri(size_t ti);
+	UE_API virtual void				CalculateVertex(size_t vi);
 
 public:
-								MeshDetails(MeshInfo* mesh);
-	virtual						~MeshDetails();
+								UE_API MeshDetails(MeshInfo* mesh);
+	UE_API virtual						~MeshDetails();
 
-	virtual MeshDetailsPAsync	Calculate();
-	virtual MeshDetailsPAsync	Finalise();
-	virtual void				RenderDebug();
-	virtual void				Release();
+	UE_API virtual MeshDetailsPAsync	Calculate();
+	UE_API virtual MeshDetailsPAsync	Finalise();
+	UE_API virtual void				RenderDebug();
+	UE_API virtual void				Release();
 };
 
 typedef std::shared_ptr<MeshDetails> MeshDetailsPtr;
+
+#undef UE_API
diff --git a/Engine/Plugins/TextureGraph/Source/TextureGraphEngine/3D/MeshDetails/MeshDetails_Edges.h b/Engine/Plugins/TextureGraph/Source/TextureGraphEngine/3D/MeshDetails/MeshDetails_Edges.h
--- a/Engine/Plugins/TextureGraph/Source/TextureGraphEngine/3D/MeshDetails/MeshDetails_Edges.h
+++ b/Engine/Plugins/TextureGraph/Source/TextureGraphEngine/3D/MeshDetails/MeshDetails_Edges.h
@@ -7,7 +7,9 @@
 #include <list>
 #include <unordered_map>
 
-class TEXTUREGRAPHENGINE_API MeshDetails_Edges : public MeshDetails
+#define UE_API TEXTUREGRAPHENGINE_API
+
+class MeshDetails_Edges : public MeshDetails
 {
 public:
 	static const int			s_maxCongruent = 8;
@@ -42,12 +44,14 @@ protected:
     IntVec         				_adjStart;              /// Adjacency start index into the _adjTriangles array (for each triangle)
     IntVec         				_adjCount;              /// Number of adjacent triangles in the _adjTriangles array (for each triangle)
 
-	virtual void				CalculateTri(size_t ti) override;
+	UE_API virtual void				CalculateTri(size_t ti) override;
 
 public:
-								MeshDetails_Edges(MeshInfo* mesh);
-	virtual						~MeshDetails_Edges();
+								UE_API MeshDetails_Edges(MeshInfo* mesh);
+	UE_API virtual						~MeshDetails_Edges();
 
-	virtual MeshDetailsPAsync	Calculate() override;
-	virtual void				Release() override;
+	UE_API virtual MeshDetailsPAsync	Calculate() override;
+	UE_API virtual void				Release() override;
 };
+
+#undef UE_API
diff --git a/Engine/Plugins/TextureGraph/Source/TextureGraphEngine/3D/MeshDetails/MeshDetails_Spatial.h b/Engine/Plugins/TextureGraph/Source/TextureGraphEngine/3D/MeshDetails/MeshDetails_Spatial.h
--- a/Engine/Plugins/TextureGraph/Source/TextureGraphEngine/3D/MeshDetails/MeshDetails_Spatial.h
+++ b/Engine/Plugins/TextureGraph/Source/TextureGraphEngine/3D/MeshDetails/MeshDetails_Spatial.h
@@ -8,9 +8,11 @@
 #include <list>
 #include <array>
 
+#define UE_API TEXTUREGRAPHENGINE_API
+
 class MeshDetails_Tri;
 
-class TEXTUREGRAPHENGINE_API MeshDetails_Spatial : MeshDetails
+class MeshDetails_Spatial : MeshDetails
 {
 public:
     struct ONode
@@ -121,12 +123,14 @@ protected:
 	std::vector<int>            _tris;                  /// The triangles
 	int                         _triCount = 0;          /// How many triangles do we have
 
-	virtual void				CalculateTri(size_t ti) override;
+	UE_API virtual void				CalculateTri(size_t ti) override;
 
 public:
-								MeshDetails_Spatial(MeshInfo* mesh);
-	virtual						~MeshDetails_Spatial();
+								UE_API MeshDetails_Spatial(MeshInfo* mesh);
+	UE_API virtual						~MeshDetails_Spatial();
 
-	virtual MeshDetailsPAsync	Calculate() override;
-	virtual void				Release() override;
+	UE_API virtual MeshDetailsPAsync	Calculate() override;
+	UE_API virtual void				Release() override;
 };
+
+#undef UE_API
diff --git a/Engine/Plugins/TextureGraph/Source/TextureGraphEngine/3D/MeshDetails/MeshDetails_SpatialUV.h b/Engine/Plugins/TextureGraph/Source/TextureGraphEngine/3D/MeshDetails/MeshDetails_SpatialUV.h
--- a/Engine/Plugins/TextureGraph/Source/TextureGraphEngine/3D/MeshDetails/MeshDetails_SpatialUV.h
+++ b/Engine/Plugins/TextureGraph/Source/TextureGraphEngine/3D/MeshDetails/MeshDetails_SpatialUV.h
@@ -6,15 +6,19 @@
 #include <vector>
 #include <list>
 
-class TEXTUREGRAPHENGINE_API MeshDetails_SpatialUV : public MeshDetails
+#define UE_API TEXTUREGRAPHENGINE_API
+
+class MeshDetails_SpatialUV : public MeshDetails
 {
 protected:
-	virtual void				CalculateTri(size_t ti) override;
+	UE_API virtual void				CalculateTri(size_t ti) override;
 
 public:
-								MeshDetails_SpatialUV(MeshInfo* mesh);
-	virtual						~MeshDetails_SpatialUV();
+								UE_API MeshDetails_SpatialUV(MeshInfo* mesh);
+	UE_API virtual						~MeshDetails_SpatialUV();
 
-	virtual MeshDetailsPAsync	Calculate() override;
-	virtual void				Release() override;
+	UE_API virtual MeshDetailsPAsync	Calculate() override;
+	UE_API virtual void				Release() override;
 };
+
+#undef UE_API
diff --git a/Engine/Plugins/TextureGraph/Source/TextureGraphEngine/3D/MeshDetails/MeshDetails_Tri.h b/Engine/Plugins/TextureGraph/Source/TextureGraphEngine/3D/MeshDetails/MeshDetails_Tri.h
--- a/Engine/Plugins/TextureGraph/Source/TextureGraphEngine/3D/MeshDetails/MeshDetails_Tri.h
+++ b/Engine/Plugins/TextureGraph/Source/TextureGraphEngine/3D/MeshDetails/MeshDetails_Tri.h
@@ -4,7 +4,9 @@
 #include "MeshDetails.h"
 #include "GraphicsDefs.h"
 
-class TEXTUREGRAPHENGINE_API MeshDetails_Tri : public MeshDetails
+#define UE_API TEXTUREGRAPHENGINE_API
+
+class MeshDetails_Tri : public MeshDetails
 {
 public:
 	struct Triangle
@@ -25,24 +27,24 @@ public:
 	};
 
 protected:
-	static const size_t			s_maxBatch;
+	static UE_API const size_t			s_maxBatch;
 	Triangle*					_triangles = nullptr;			/// The triangles that we have
 	size_t						_numTriangles = 0;				/// How many triangles do we have
 	RawBufferPtr				_raw;							/// The triangle buffer is indexed into this RawBuffer
 
-	virtual void				CalculateTri(size_t ti) override;
+	UE_API virtual void				CalculateTri(size_t ti) override;
 
 public:
-								MeshDetails_Tri(MeshInfo* mesh);
-	virtual						~MeshDetails_Tri();
+								UE_API MeshDetails_Tri(MeshInfo* mesh);
+	UE_API virtual						~MeshDetails_Tri();
 
-	virtual MeshDetailsPAsync	Calculate() override;
+	UE_API virtual MeshDetailsPAsync	Calculate() override;
 
-	int32*						GetIndices(size_t ti) const;
-	std::array<Vector3, 3>		GetVertices(size_t ti) const;
-	std::array<Vector2, 3>		GetUVs(size_t ti) const;
-	std::array<Vector3, 3>		GetNormals(size_t ti) const;
-	std::array<Tangent, 3>		GetTangents(size_t ti) const;
+	UE_API int32*						GetIndices(size_t ti) const;
+	UE_API std::array<Vector3, 3>		GetVertices(size_t ti) const;
+	UE_API std::array<Vector2, 3>		GetUVs(size_t ti) const;
+	UE_API std::array<Vector3, 3>		GetNormals(size_t ti) const;
+	UE_API std::array<Tangent, 3>		GetTangents(size_t ti) const;
 
 	//////////////////////////////////////////////////////////////////////////
 	/// Inline functions
@@ -52,3 +54,4 @@ public:
 };
 
 
+#undef UE_API
diff --git a/Engine/Plugins/TextureGraph/Source/TextureGraphEngine/3D/MeshInfo.h b/Engine/Plugins/TextureGraph/Source/TextureGraphEngine/3D/MeshInfo.h
--- a/Engine/Plugins/TextureGraph/Source/TextureGraphEngine/3D/MeshInfo.h
+++ b/Engine/Plugins/TextureGraph/Source/TextureGraphEngine/3D/MeshInfo.h
@@ -9,6 +9,8 @@
 #include <unordered_map>
 #include <limits>
 
+#define UE_API TEXTUREGRAPHENGINE_API
+
 struct CoreMesh;
 typedef std::shared_ptr<CoreMesh>	CoreMeshPtr;
 
@@ -17,7 +19,7 @@ typedef std::shared_ptr<MeshDetails> MeshDetailsPtr;
 
 class MeshDetails_Tri;
 
-class TEXTUREGRAPHENGINE_API MeshInfo
+class MeshInfo
 {
 protected:
     typedef std::unordered_map<const char*, MeshDetailsPtr> MeshDetailsPtrLookup;
@@ -28,18 +30,18 @@ protected:
     MeshDetailsPtrLookup        _details;           /// Details lookup
 
 public:
-								MeshInfo(CoreMeshPtr cmesh);
-								~MeshInfo();
+								UE_API MeshInfo(CoreMeshPtr cmesh);
+								UE_API ~MeshInfo();
 
-	size_t						NumVertices() const;
-	size_t						NumTriangles() const;
-	std::array<Vector3, 3>		Vertices(int32 i0, int32 i1, int32 i2) const;
-	int32		                GetMaterialIndex();
-    void                        InitBounds(FVector min, FVector max);
-    void                        UpdateBounds(const FVector& vert);
+	UE_API size_t						NumVertices() const;
+	UE_API size_t						NumTriangles() const;
+	UE_API std::array<Vector3, 3>		Vertices(int32 i0, int32 i1, int32 i2) const;
+	UE_API int32		                GetMaterialIndex();
+    UE_API void                        InitBounds(FVector min, FVector max);
+    UE_API void                        UpdateBounds(const FVector& vert);
 
 	//////////////////////////////////////////////////////////////////////////
-    MeshDetails_Tri*            d_Tri();
+    UE_API MeshDetails_Tri*            d_Tri();
 
     template <typename T_Details>
     T_Details* GetAttribute(const char* name, bool create = true) 
@@ -69,4 +71,6 @@ public:
 	FORCEINLINE CHashPtr		Hash() const { return _hash; }
 };
 
-typedef std::shared_ptr<MeshInfo> MeshInfoPtr;
\ No newline at end of file
+typedef std::shared_ptr<MeshInfo> MeshInfoPtr;
+
+#undef UE_API
diff --git a/Engine/Plugins/TextureGraph/Source/TextureGraphEngine/3D/ProceduralMeshActor.h b/Engine/Plugins/TextureGraph/Source/TextureGraphEngine/3D/ProceduralMeshActor.h
--- a/Engine/Plugins/TextureGraph/Source/TextureGraphEngine/3D/ProceduralMeshActor.h
+++ b/Engine/Plugins/TextureGraph/Source/TextureGraphEngine/3D/ProceduralMeshActor.h
@@ -18,12 +18,14 @@
 #include "Device/FX/DeviceBuffer_FX.h"
 #include "ProceduralMeshActor.generated.h"
 
+#define UE_API TEXTUREGRAPHENGINE_API
+
 class USceneCaptureComponent2D;
 
 class Tex;
 typedef std::shared_ptr<Tex>		TexPtr;
-UCLASS()
-class TEXTUREGRAPHENGINE_API AProceduralMeshActor : public AActor
+UCLASS(MinimalAPI)
+class AProceduralMeshActor : public AActor
 {
 	GENERATED_BODY()
 private:
@@ -31,13 +33,13 @@ private:
 	bool								_isDebug = false;
 	FString								_meshName = FString("");
 
-	void								BlitInternal(UMaterialInterface* mat);
+	UE_API void								BlitInternal(UMaterialInterface* mat);
 
 protected:
 	// Called when the game starts or when spawned
-	virtual void						BeginPlay() override;
+	UE_API virtual void						BeginPlay() override;
 	// Called every frame
-	virtual void						Tick(float DeltaTime) override;
+	UE_API virtual void						Tick(float DeltaTime) override;
 
 	
 	UPROPERTY(VisibleAnywhere, BlueprintReadWrite, Category = "Hierarchy")
@@ -47,23 +49,23 @@ protected:
 		TObjectPtr<USceneCaptureComponent2D>		_sceneCaptureComp;
 public:	
 	// Sets default values for this actor's properties
-										AProceduralMeshActor();
-										~AProceduralMeshActor();
+										UE_API AProceduralMeshActor();
+										UE_API ~AProceduralMeshActor();
 
-										void InitSceneCaptureComponent();
+										UE_API void InitSceneCaptureComponent();
 
-										void InitMeshObj();
+										UE_API void InitMeshObj();
 
-	void								SetMeshData(CoreMeshPtr mesh);
-	void 								DrawDebugLines(TArray<FVector>& positionArray, FVector length);
-	void 								SetMaterial(UMaterialInterface* mat);
-	UMaterialInterface*  				GetMaterial();
+	UE_API void								SetMeshData(CoreMeshPtr mesh);
+	UE_API void 								DrawDebugLines(TArray<FVector>& positionArray, FVector length);
+	UE_API void 								SetMaterial(UMaterialInterface* mat);
+	UE_API UMaterialInterface*  				GetMaterial();
 	
-	void								BlitTo(UTextureRenderTarget2D* rt, UMaterialInterface* mat);
-	virtual void						BeginDestroy();
+	UE_API void								BlitTo(UTextureRenderTarget2D* rt, UMaterialInterface* mat);
+	UE_API virtual void						BeginDestroy();
 
 	UFUNCTION(BlueprintCallable, Category = "Debug")
-		void							ToggleDebug(int32 debugType);
+		UE_API void							ToggleDebug(int32 debugType);
 
 	UFUNCTION(BlueprintCallable,Category = "Info")
 		const FString &					GetMeshName() { return _meshName; }
@@ -75,3 +77,5 @@ public:
 	/// Inline functions
 	//////////////////////////////////////////////////////////////////////////
 };
+
+#undef UE_API
diff --git a/Engine/Plugins/TextureGraph/Source/TextureGraphEngine/3D/RenderMesh.h b/Engine/Plugins/TextureGraph/Source/TextureGraphEngine/3D/RenderMesh.h
--- a/Engine/Plugins/TextureGraph/Source/TextureGraphEngine/3D/RenderMesh.h
+++ b/Engine/Plugins/TextureGraph/Source/TextureGraphEngine/3D/RenderMesh.h
@@ -11,6 +11,8 @@
 
 #include "RenderMesh.generated.h"
 
+#define UE_API TEXTUREGRAPHENGINE_API
+
 UENUM()
 enum class MeshType : uint8
 {
@@ -75,7 +77,7 @@ typedef std::shared_ptr<RenderMesh>				RenderMeshPtr;
 
 //////////////////////////////////////////////////////////////////////////
 DEFINE_LOG_CATEGORY_STATIC(LogMesh, Log, All);
-class TEXTUREGRAPHENGINE_API RenderMesh 
+class RenderMesh 
 {
 
 protected:
@@ -107,39 +109,39 @@ protected:
 	TiledBlobPtr					_worldMaps[WorldTextures::WorldTextureCount];
 	TiledBlobPtr					_worldMapsSingleBlob[WorldTextures::WorldTextureCount];
 	
-	void							AddMaterialInfo(int32 id, FString& matName);
-	virtual void					SpawnActors(UWorld* world);
-	void							UpdateMeshTransforms();
-	void							UpdateBounds();
+	UE_API void							AddMaterialInfo(int32 id, FString& matName);
+	UE_API virtual void					SpawnActors(UWorld* world);
+	UE_API void							UpdateMeshTransforms();
+	UE_API void							UpdateBounds();
 protected:
 	virtual void					LoadInternal() = 0;
 public:
 									RenderMesh() = default;
-									RenderMesh(const MeshLoadInfo loadInfo);
-									RenderMesh(RenderMesh* parent, TArray<MeshInfoPtr> mesh, MaterialInfoPtr matInfo);
+									UE_API RenderMesh(const MeshLoadInfo loadInfo);
+									UE_API RenderMesh(RenderMesh* parent, TArray<MeshInfoPtr> mesh, MaterialInfoPtr matInfo);
 
-	virtual							~RenderMesh();
+	UE_API virtual							~RenderMesh();
 
-	virtual void					PrepareForRendering(UWorld* world,FVector scale);
+	UE_API virtual void					PrepareForRendering(UWorld* world,FVector scale);
 	
-	void							DrawBounds(UWorld* world);
-	virtual void					SetMaterial(UMaterialInterface* material);	
+	UE_API void							DrawBounds(UWorld* world);
+	UE_API virtual void					SetMaterial(UMaterialInterface* material);	
 	virtual AsyncActionResultPtr	Load() = 0;	
-	FString							GetMaterialName();
+	UE_API FString							GetMaterialName();
 
-	virtual CHashPtr				Hash() const;
-	virtual FMatrix					LocalToWorldMatrix() const;
-	virtual void					RemoveActors();
-	virtual void					Render_Now(FRHICommandList& rhi, int32 targetId) const;
-	virtual void					Init_PSO(FGraphicsPipelineStateInitializer& pso) const;
+	UE_API virtual CHashPtr				Hash() const;
+	UE_API virtual FMatrix					LocalToWorldMatrix() const;
+	UE_API virtual void					RemoveActors();
+	UE_API virtual void					Render_Now(FRHICommandList& rhi, int32 targetId) const;
+	UE_API virtual void					Init_PSO(FGraphicsPipelineStateInitializer& pso) const;
 
-	virtual TiledBlobPtr			WorldPosTexture(MixUpdateCyclePtr cycle, int32 targetId, bool singleBLob = false);
-	virtual TiledBlobPtr			WorldNormalsTexture(MixUpdateCyclePtr cycle, int32 targetId, bool singleBLob = false);
-	virtual TiledBlobPtr			WorldTangentsTexture(MixUpdateCyclePtr cycle, int32 targetId, bool singleBLob = false);
-	virtual TiledBlobPtr			WorldUVMaskTexture(MixUpdateCyclePtr cycle, int32 targetId, bool singleBLob = false);
+	UE_API virtual TiledBlobPtr			WorldPosTexture(MixUpdateCyclePtr cycle, int32 targetId, bool singleBLob = false);
+	UE_API virtual TiledBlobPtr			WorldNormalsTexture(MixUpdateCyclePtr cycle, int32 targetId, bool singleBLob = false);
+	UE_API virtual TiledBlobPtr			WorldTangentsTexture(MixUpdateCyclePtr cycle, int32 targetId, bool singleBLob = false);
+	UE_API virtual TiledBlobPtr			WorldUVMaskTexture(MixUpdateCyclePtr cycle, int32 targetId, bool singleBLob = false);
 
-	virtual FString					Name() const;
-	virtual void					Clear();
+	UE_API virtual FString					Name() const;
+	UE_API virtual void					Clear();
 
 	bool							ContainsWorldTexture(WorldTextures texture) { return _worldMaps[texture] != nullptr; }
 	bool							ContainsWorldNormals()		{ return _worldMaps[WorldTextures::Normals] != nullptr;}
@@ -183,6 +185,7 @@ public:
 	//////////////////////////////////////////////////////////////////////////
 	//// Static Methods
 	//////////////////////////////////////////////////////////////////////////
-	static RenderMeshPtr			Create(const MeshLoadInfo _loadInfo);
+	static UE_API RenderMeshPtr			Create(const MeshLoadInfo _loadInfo);
 };
 
+#undef UE_API
diff --git a/Engine/Plugins/TextureGraph/Source/TextureGraphEngine/3D/RenderMesh_Editor.h b/Engine/Plugins/TextureGraph/Source/TextureGraphEngine/3D/RenderMesh_Editor.h
--- a/Engine/Plugins/TextureGraph/Source/TextureGraphEngine/3D/RenderMesh_Editor.h
+++ b/Engine/Plugins/TextureGraph/Source/TextureGraphEngine/3D/RenderMesh_Editor.h
@@ -7,26 +7,26 @@
 
 struct CoreMesh;
 
-class TEXTUREGRAPHENGINE_API RenderMesh_Editor: public RenderMesh
+class RenderMesh_Editor: public RenderMesh
 {
 
 protected:
 	TArray<UStaticMeshComponent*>			MeshComponents;
 	UWorld*									World = nullptr;
-	virtual void							LoadInternal() override;
-	void									LoadSingleMeshComponent(const UStaticMeshComponent& mesh);
+	TEXTUREGRAPHENGINE_API virtual void							LoadInternal() override;
+	TEXTUREGRAPHENGINE_API void									LoadSingleMeshComponent(const UStaticMeshComponent& mesh);
 
 public:
-											RenderMesh_Editor();
-											RenderMesh_Editor(RenderMesh* parent, TArray<MeshInfoPtr> meshes, MaterialInfoPtr matInfo);
-											RenderMesh_Editor(UStaticMeshComponent* staticMeshComponent, UWorld* world);
+											TEXTUREGRAPHENGINE_API RenderMesh_Editor();
+											TEXTUREGRAPHENGINE_API RenderMesh_Editor(RenderMesh* parent, TArray<MeshInfoPtr> meshes, MaterialInfoPtr matInfo);
+											TEXTUREGRAPHENGINE_API RenderMesh_Editor(UStaticMeshComponent* staticMeshComponent, UWorld* world);
 
 	FORCEINLINE								TArray<UStaticMeshComponent*> GetMeshComponents() { return MeshComponents; }
-	virtual AsyncActionResultPtr			Load() override;
+	TEXTUREGRAPHENGINE_API virtual AsyncActionResultPtr			Load() override;
 	virtual FMatrix							LocalToWorldMatrix() const override { return MeshComponents[0]->GetComponentTransform().ToMatrixWithScale(); }
-	virtual void							PrepareForRendering(UWorld* world, FVector scale) override;
-	virtual void							SetMaterial(UMaterialInterface* material) override;
+	TEXTUREGRAPHENGINE_API virtual void							PrepareForRendering(UWorld* world, FVector scale) override;
+	TEXTUREGRAPHENGINE_API virtual void							SetMaterial(UMaterialInterface* material) override;
 };
 
 typedef std::shared_ptr<RenderMesh_Editor> RenderMesh_EditorPtr;
-#endif
\ No newline at end of file
+#endif
diff --git a/Engine/Plugins/TextureGraph/Source/TextureGraphEngine/3D/RenderMesh_EditorScene.h b/Engine/Plugins/TextureGraph/Source/TextureGraphEngine/3D/RenderMesh_EditorScene.h
--- a/Engine/Plugins/TextureGraph/Source/TextureGraphEngine/3D/RenderMesh_EditorScene.h
+++ b/Engine/Plugins/TextureGraph/Source/TextureGraphEngine/3D/RenderMesh_EditorScene.h
@@ -6,21 +6,21 @@
 
 struct CoreMesh;
 
-class TEXTUREGRAPHENGINE_API RenderMesh_EditorScene: public RenderMesh_Editor
+class RenderMesh_EditorScene: public RenderMesh_Editor
 {
 private:
 	AActor*									MeshActor;				/// Actor that this rendermesh is associated to
 
 public:
-											RenderMesh_EditorScene();
-											RenderMesh_EditorScene(AActor* actor);
-											RenderMesh_EditorScene(RenderMesh* parent, TArray<MeshInfoPtr> meshes, MaterialInfoPtr matInfo);
+											TEXTUREGRAPHENGINE_API RenderMesh_EditorScene();
+											TEXTUREGRAPHENGINE_API RenderMesh_EditorScene(AActor* actor);
+											TEXTUREGRAPHENGINE_API RenderMesh_EditorScene(RenderMesh* parent, TArray<MeshInfoPtr> meshes, MaterialInfoPtr matInfo);
 
-	AsyncActionResultPtr					Load();
-	void									PrepareForRendering(UWorld* world, FVector scale) override;
-	void									RemoveActors() override;
+	TEXTUREGRAPHENGINE_API AsyncActionResultPtr					Load();
+	TEXTUREGRAPHENGINE_API void									PrepareForRendering(UWorld* world, FVector scale) override;
+	TEXTUREGRAPHENGINE_API void									RemoveActors() override;
 
 };
 
 typedef std::shared_ptr<RenderMesh_EditorScene> RenderMesh_EditorScenePtr;
-#endif
\ No newline at end of file
+#endif
diff --git a/Engine/Plugins/TextureGraph/Source/TextureGraphEngine/3D/RenderMesh_Procedural.h b/Engine/Plugins/TextureGraph/Source/TextureGraphEngine/3D/RenderMesh_Procedural.h
--- a/Engine/Plugins/TextureGraph/Source/TextureGraphEngine/3D/RenderMesh_Procedural.h
+++ b/Engine/Plugins/TextureGraph/Source/TextureGraphEngine/3D/RenderMesh_Procedural.h
@@ -4,13 +4,15 @@
 #include "3D/RenderMesh.h"
 #include "Helper/Promise.h"
 
+#define UE_API TEXTUREGRAPHENGINE_API
+
 struct CoreMesh;
 typedef std::shared_ptr<CoreMesh>	CoreMeshPtr;
 
 
 class RenderMesh;
 
-class TEXTUREGRAPHENGINE_API RenderMesh_Procedural: public RenderMesh
+class RenderMesh_Procedural: public RenderMesh
 {
 private:
 	int										_tesselation = 32;
@@ -18,10 +20,10 @@ private:
 	FVector									_offSet;
 	
 public:
-	explicit								RenderMesh_Procedural(const MeshLoadInfo loadInfo);
-	virtual AsyncActionResultPtr			Load() override;
-	virtual void							LoadInternal() override;
-	void									GenerateProcedural(int tesselation, FVector2D dimension, CoreMeshPtr cmesh);
+	UE_API explicit								RenderMesh_Procedural(const MeshLoadInfo loadInfo);
+	UE_API virtual AsyncActionResultPtr			Load() override;
+	UE_API virtual void							LoadInternal() override;
+	UE_API void									GenerateProcedural(int tesselation, FVector2D dimension, CoreMeshPtr cmesh);
 
 	//////////////////////////////////////////////////////////////////////////
 	////Inline Functions
@@ -29,4 +31,6 @@ public:
 	FORCEINLINE	int&						Tesselation() { return _tesselation; }
 };
 
-typedef std::shared_ptr<RenderMesh_Procedural> RenderMesh_ProceduralPtr;
\ No newline at end of file
+typedef std::shared_ptr<RenderMesh_Procedural> RenderMesh_ProceduralPtr;
+
+#undef UE_API
diff --git a/Engine/Plugins/TextureGraph/Source/TextureGraphEngine/Data/BinCache.h b/Engine/Plugins/TextureGraph/Source/TextureGraphEngine/Data/BinCache.h
--- a/Engine/Plugins/TextureGraph/Source/TextureGraphEngine/Data/BinCache.h
+++ b/Engine/Plugins/TextureGraph/Source/TextureGraphEngine/Data/BinCache.h
@@ -7,7 +7,7 @@
 
 //////////////////////////////////////////////////////////////////////////
 template <typename T_Key, typename T_Value>
-class TEXTUREGRAPHENGINE_API BinCache
+class BinCache
 {
 	typedef LruCache<T_Key, T_Value> LRUCache;
 
diff --git a/Engine/Plugins/TextureGraph/Source/TextureGraphEngine/Data/Blobber.h b/Engine/Plugins/TextureGraph/Source/TextureGraphEngine/Data/Blobber.h
--- a/Engine/Plugins/TextureGraph/Source/TextureGraphEngine/Data/Blobber.h
+++ b/Engine/Plugins/TextureGraph/Source/TextureGraphEngine/Data/Blobber.h
@@ -11,6 +11,8 @@
 
 #include "Helper/Util.h"
 
+#define UE_API TEXTUREGRAPHENGINE_API
+
 class Device;
 struct JobResult;
 typedef std::shared_ptr<JobResult>			JobResultPtr;
@@ -23,13 +25,13 @@ DECLARE_LOG_CATEGORY_EXTERN(LogBlob, Log, Verbose);
 typedef std::vector<HashType>   HashTypes;
 class JobArg;
 
-class TEXTUREGRAPHENGINE_API BlobberObserverSource
+class BlobberObserverSource
 {
 public:
 	using HashArray = std::vector<HashType>;
 
-	void AddHash(HashType Hash);
-	void RemapHash(HashType OldHash, HashType NewHash);
+	UE_API void AddHash(HashType Hash);
+	UE_API void RemapHash(HashType OldHash, HashType NewHash);
 
 protected:
 	/// Protected interface of emitters called by the scheduler to notify the observers
@@ -43,7 +45,7 @@ protected:
 	uint32					Version = 1;
 
 	/// Trigger the broadcast of the changes, this call is issued from the Blobber::Update function
-	void Broadcast();
+	UE_API void Broadcast();
 
 	/// The customisable handler method triggered from Broadcast if any buffers were added or removed
 	virtual void BlobberUpdated(HashArray&& InAddedHashes, HashArray&& InRemappedHashes) {}
@@ -63,7 +65,7 @@ struct BlobCacheOptions
 	bool						NoCacheBatch = false;		/// Whether the batch is marked for no caching
 };
 
-class TEXTUREGRAPHENGINE_API Blobber
+class Blobber
 {
 	friend class Blob;
 	friend class DeviceTransferService;
@@ -149,9 +151,9 @@ public:
 	RefBlobMap						ReferencedBlobs;			/// Blobs that are referenced (Debug only!)
 
 public:
-	void							AddReferencedBlob(Blob* BlobObj, JobArg* Arg);
-	void							RemoveReferencedBlob(Blob* BlobObj, JobArg* Arg);
-	bool							IsBlobReferenced(Blob* BlobObj);
+	UE_API void							AddReferencedBlob(Blob* BlobObj, JobArg* Arg);
+	UE_API void							RemoveReferencedBlob(Blob* BlobObj, JobArg* Arg);
+	UE_API bool							IsBlobReferenced(Blob* BlobObj);
 private:
 #endif 
 	
@@ -161,55 +163,55 @@ private:
 	FCriticalSection				DebugBlobMutex;				/// The TiledBlobs that contain this blob as a tile
 #endif 
 
-	void							AddHashMapping(HashType LHS, CHashPtr RHS);
-	void							AddBlobEntry(HashType Hash, BlobPtr BlobObj, BlobCacheOptions Options);
-	void							AddBlobEntryThreadSafe(HashType Hash, BlobPtr BlobObj, BlobCacheOptions Options);
-	void							RemoveHashMapping(HashType Hash);
-
-	void							PrintStats();
-	BlobRef							AddInternal(BlobPtr BlobObj, BlobCacheOptions Options);
-	BlobPtr							Remove(HashType Hash);
-
-	CHashPtr						AddGloballyUniqueHashInternal(CHashPtr Hash);
-	void							UpdateGloballyUniqueHashInternal(HashType OldValue);
-	void							UpdateMappingInternal(HashType OldValue, CHashPtr NewValue);
-	void							UpdateBlobLookupInternal(HashType OldValue, BlobPtr BlobObj);
-	void							Touch(const Blob* BlobObj);
-	BlobCacheEntryPtr				FindInternal(HashType Hash);
-	BlobRef							AddResult(CHashPtr LHash, BlobRef Result, BlobCacheOptions Options);
-	void							UpdateTransient();
-	void							UpdateTouchedHashes();
-	void							UpdateBlobCache();
+	UE_API void							AddHashMapping(HashType LHS, CHashPtr RHS);
+	UE_API void							AddBlobEntry(HashType Hash, BlobPtr BlobObj, BlobCacheOptions Options);
+	UE_API void							AddBlobEntryThreadSafe(HashType Hash, BlobPtr BlobObj, BlobCacheOptions Options);
+	UE_API void							RemoveHashMapping(HashType Hash);
+
+	UE_API void							PrintStats();
+	UE_API BlobRef							AddInternal(BlobPtr BlobObj, BlobCacheOptions Options);
+	UE_API BlobPtr							Remove(HashType Hash);
+
+	UE_API CHashPtr						AddGloballyUniqueHashInternal(CHashPtr Hash);
+	UE_API void							UpdateGloballyUniqueHashInternal(HashType OldValue);
+	UE_API void							UpdateMappingInternal(HashType OldValue, CHashPtr NewValue);
+	UE_API void							UpdateBlobLookupInternal(HashType OldValue, BlobPtr BlobObj);
+	UE_API void							Touch(const Blob* BlobObj);
+	UE_API BlobCacheEntryPtr				FindInternal(HashType Hash);
+	UE_API BlobRef							AddResult(CHashPtr LHash, BlobRef Result, BlobCacheOptions Options);
+	UE_API void							UpdateTransient();
+	UE_API void							UpdateTouchedHashes();
+	UE_API void							UpdateBlobCache();
 
 public:
-									Blobber();
-									~Blobber();
+									UE_API Blobber();
+									UE_API ~Blobber();
 
-	BlobRef							FindSingle(HashType Hash);
-	TiledBlobRef					FindTiled(HashType Hash);
+	UE_API BlobRef							FindSingle(HashType Hash);
+	UE_API TiledBlobRef					FindTiled(HashType Hash);
 
-	BlobRef							Create(const BufferDescriptor& Desc);
-	BlobRef							Create(Device* Dev, RawBufferPtr Raw);
-	BlobRef							Create(DeviceBufferRef DevBuffer, bool NoCache = false);
+	UE_API BlobRef							Create(const BufferDescriptor& Desc);
+	UE_API BlobRef							Create(Device* Dev, RawBufferPtr Raw);
+	UE_API BlobRef							Create(DeviceBufferRef DevBuffer, bool NoCache = false);
 
-	BlobRef							AddResult(CHashPtr LHash, BlobPtr InBlob, BlobCacheOptions Options = {});
-	void							ClearCache();
+	UE_API BlobRef							AddResult(CHashPtr LHash, BlobPtr InBlob, BlobCacheOptions Options = {});
+	UE_API void							ClearCache();
 
 	/// In order to minimise the changes brought upon by weak reference work we've made this function.
 	/// If we were to force this to take a unique_ptr<TiledBlob> it would've resulted in lots of changes
 	/// around asset manager and other bits of code (Jobs etc), that we can do without.
 	/// Eventually, this should either go away, or be converted to accept a unique pointer to force stop
 	/// shared_ptr being passed around
-	TiledBlobRef					AddTiledResult(CHashPtr LHash, TiledBlobPtr Result, BlobCacheOptions Options = {});
-	TiledBlobRef					AddTiledResult(TiledBlobPtr Result, BlobCacheOptions Options = {});
+	UE_API TiledBlobRef					AddTiledResult(CHashPtr LHash, TiledBlobPtr Result, BlobCacheOptions Options = {});
+	UE_API TiledBlobRef					AddTiledResult(TiledBlobPtr Result, BlobCacheOptions Options = {});
 
-	void							UpdateBlobHash(HashType OldHash, BlobRef InBlob);
-	void							UpdateHash(HashType OldHash, CHashPtr NewHash);
-	void							Update(float DT);
-	AsyncJobResultPtr				UpdateIdle();
-	bool							IsBlobCached(HashType Hash);
+	UE_API void							UpdateBlobHash(HashType OldHash, BlobRef InBlob);
+	UE_API void							UpdateHash(HashType OldHash, CHashPtr NewHash);
+	UE_API void							Update(float DT);
+	UE_API AsyncJobResultPtr				UpdateIdle();
+	UE_API bool							IsBlobCached(HashType Hash);
 
-	void							RegisterObserverSource(const BlobberObserverSourcePtr& observerSource); ///
+	UE_API void							RegisterObserverSource(const BlobberObserverSourcePtr& observerSource); ///
 
 	template <class BlobType>
 	T_BlobRef<BlobType>				Find(HashType Hash)
@@ -328,3 +330,4 @@ std::shared_ptr<BlobType> T_BlobRef<BlobType>::get() const
 #endif 
 }
 
+#undef UE_API
