Ran UnrealCodeFixup to fix dll storage

diff --git a/Engine/Plugins/TextureGraph/Source/TextureGraphEngine/Data/G_Collectible.h b/Engine/Plugins/TextureGraph/Source/TextureGraphEngine/Data/G_Collectible.h
--- a/Engine/Plugins/TextureGraph/Source/TextureGraphEngine/Data/G_Collectible.h
+++ b/Engine/Plugins/TextureGraph/Source/TextureGraphEngine/Data/G_Collectible.h
@@ -4,6 +4,8 @@
 #include "CoreMinimal.h"
 #include <atomic>
 
+#define UE_API TEXTUREGRAPHENGINE_API
+
 //////////////////////////////////////////////////////////////////////////
 struct AccessInfo
 {
@@ -36,13 +38,13 @@ struct AccessInfo
 
 //////////////////////////////////////////////////////////////////////////
 /// Garbage collectible data structure
-class TEXTUREGRAPHENGINE_API G_Collectible
+class G_Collectible
 {
 protected:
 	AccessInfo						AccessDetails;				/// Information about the last time this DeviceBuffer was accessed
 
 public:
-	virtual void					UpdateAccessInfo(uint64 BatchId = 0);
+	UE_API virtual void					UpdateAccessInfo(uint64 BatchId = 0);
 	virtual							~G_Collectible() {}
 
 	//////////////////////////////////////////////////////////////////////////
@@ -51,3 +53,5 @@ public:
 	FORCEINLINE const AccessInfo&	GetAccessInfo() const { return AccessDetails; }
 	FORCEINLINE bool				IsSame(const std::atomic<double>& Timestamp) const { return AccessDetails.Timestamp == Timestamp; }
 };
+
+#undef UE_API
diff --git a/Engine/Plugins/TextureGraph/Source/TextureGraphEngine/Data/RawBuffer.h b/Engine/Plugins/TextureGraph/Source/TextureGraphEngine/Data/RawBuffer.h
--- a/Engine/Plugins/TextureGraph/Source/TextureGraphEngine/Data/RawBuffer.h
+++ b/Engine/Plugins/TextureGraph/Source/TextureGraphEngine/Data/RawBuffer.h
@@ -11,6 +11,8 @@ THIRD_PARTY_INCLUDES_END
 #include <PixelFormat.h>
 #include <Engine/Texture.h>
 
+#define UE_API TEXTUREGRAPHENGINE_API
+
 enum class BufferFormat 
 {
 	Byte,
@@ -61,7 +63,7 @@ struct RawBufferMetadataDefs
 };
 
 //////////////////////////////////////////////////////////////////////////
-struct TEXTUREGRAPHENGINE_API BufferDescriptor
+struct BufferDescriptor
 {
 	typedef std::set<FString>		BufferMetadata;
 
@@ -79,24 +81,24 @@ struct TEXTUREGRAPHENGINE_API BufferDescriptor
 
 	BufferMetadata					Metadata;				/// The metadata associated with this buffer descriptor
 
-	HashType						HashValue() const;
-	HashType						FormatHashValue() const;
-	static const char*				FormatToString(BufferFormat InFormat);
+	UE_API HashType						HashValue() const;
+	UE_API HashType						FormatHashValue() const;
+	static UE_API const char*				FormatToString(BufferFormat InFormat);
 
 									BufferDescriptor() {}
-									BufferDescriptor(uint32 InWidth, uint32 InHeight, uint32 InItemsPerPoint, 
+									UE_API BufferDescriptor(uint32 InWidth, uint32 InHeight, uint32 InItemsPerPoint, 
 										BufferFormat InFormat = BufferFormat::Float, FLinearColor InDefaultValue = FLinearColor::Black, BufferType InType = BufferType::Image, bool bInMipMaps = false, bool bInSRGB = false);
 
-	static EPixelFormat				BufferPixelFormat(BufferFormat InFormat, uint32 InItemsPerPoint);
-	static BufferFormat				BufferFormatFromPixelFormat(EPixelFormat PixelFormat);
-	static size_t					BufferFormatSize(BufferFormat InFormat);
-	static ETextureSourceFormat		TextureSourceFormat(BufferFormat InFormat, uint32 InItemsPerPoint);
-	static BufferDescriptor			Combine(const BufferDescriptor& Desc1, const BufferDescriptor& Desc2);
-	static BufferDescriptor			CombineWithPreference(const BufferDescriptor* BaseDesc, const BufferDescriptor* OverrideDesc, const BufferDescriptor* RefDesc);
+	static UE_API EPixelFormat				BufferPixelFormat(BufferFormat InFormat, uint32 InItemsPerPoint);
+	static UE_API BufferFormat				BufferFormatFromPixelFormat(EPixelFormat PixelFormat);
+	static UE_API size_t					BufferFormatSize(BufferFormat InFormat);
+	static UE_API ETextureSourceFormat		TextureSourceFormat(BufferFormat InFormat, uint32 InItemsPerPoint);
+	static UE_API BufferDescriptor			Combine(const BufferDescriptor& Desc1, const BufferDescriptor& Desc2);
+	static UE_API BufferDescriptor			CombineWithPreference(const BufferDescriptor* BaseDesc, const BufferDescriptor* OverrideDesc, const BufferDescriptor* RefDesc);
 
-	ETextureSourceFormat			TextureSourceFormat() const;
-	bool							IsValid() const;
-	bool							IsFinal() const;
+	UE_API ETextureSourceFormat			TextureSourceFormat() const;
+	UE_API bool							IsValid() const;
+	UE_API bool							IsFinal() const;
 
 	FORCEINLINE size_t				FormatSize() const { return BufferDescriptor::BufferFormatSize(Format); }
 	FORCEINLINE size_t				Pitch() const { return Width * FormatSize() * ItemsPerPoint; }
@@ -114,8 +116,8 @@ struct TEXTUREGRAPHENGINE_API BufferDescriptor
 	void							AddMetadata(const FString& meta) { Metadata.insert(meta); };
 	void							AddMetadata(FString&& meta) { Metadata.insert(meta); };
 
-	bool							operator == (const BufferDescriptor& rhs) const;
-	bool							operator != (const BufferDescriptor& RHS) const;
+	UE_API bool							operator == (const BufferDescriptor& rhs) const;
+	UE_API bool							operator != (const BufferDescriptor& RHS) const;
 	
 	FORCEINLINE void				AllowUAV() { AddMetadata(RawBufferMetadataDefs::G_FX_UAV); }
 	FORCEINLINE bool				RequiresUAV() const { return HasMetadata(RawBufferMetadataDefs::G_FX_UAV); }
@@ -140,10 +142,10 @@ enum class RawBufferCompressionType
 	LZ4,							/// LZ4 compression
 };
 
-class TEXTUREGRAPHENGINE_API RawBuffer
+class RawBuffer
 {
 public:
-	static const uint64				GMinCompress;				/// Must have at least this much data to consider compression. 
+	static UE_API const uint64				GMinCompress;				/// Must have at least this much data to consider compression. 
 																/// Otherwise, its not worth it
 
 private:
@@ -157,7 +159,7 @@ private:
 		/// For backwards compatibility Anything that 
 	};
 
-	static const RawBufferCompressionType GDefaultCompression;
+	static UE_API const RawBufferCompressionType GDefaultCompression;
 
 	friend class					Device_Mem;
 
@@ -182,39 +184,39 @@ private:
 	////////////////////////////////////////////////////////////////////////// 
 	FString							FileName;		/// The path where the file data is kept
 
-	RawBufferCompressionType		ChooseBestCompressionFormat() const;
+	UE_API RawBufferCompressionType		ChooseBestCompressionFormat() const;
 
-	void							FreeUncompressed();
-	void							FreeCompressed();
-	void							FreeDisk();
+	UE_API void							FreeUncompressed();
+	UE_API void							FreeCompressed();
+	UE_API void							FreeDisk();
 
-	bool							UncompressPNG();
-	bool							UncompressGeneric(FName InName);
+	UE_API bool							UncompressPNG();
+	UE_API bool							UncompressGeneric(FName InName);
 
-	bool							CompressPNG();
-	bool							CompressGeneric(FName InName, RawBufferCompressionType InType);
+	UE_API bool							CompressPNG();
+	UE_API bool							CompressGeneric(FName InName, RawBufferCompressionType InType);
 
-	AsyncRawBufferP					ReadFromFile(bool bDoUncompress = false);
+	UE_API AsyncRawBufferP					ReadFromFile(bool bDoUncompress = false);
 
 public:
-									RawBuffer(const uint8* InData, size_t InLength, const BufferDescriptor& InDesc, CHashPtr InHashValue = nullptr, bool bInIsMemoryAutoManaged = false);
-	explicit						RawBuffer(const BufferDescriptor& InDesc);
-	explicit						RawBuffer(const RawBufferPtrTiles& InTiles);
-									~RawBuffer();
+									UE_API RawBuffer(const uint8* InData, size_t InLength, const BufferDescriptor& InDesc, CHashPtr InHashValue = nullptr, bool bInIsMemoryAutoManaged = false);
+	UE_API explicit						RawBuffer(const BufferDescriptor& InDesc);
+	UE_API explicit						RawBuffer(const RawBufferPtrTiles& InTiles);
+									UE_API ~RawBuffer();
 
-	AsyncRawBufferP					Compress(RawBufferCompressionType InCompressionType = RawBufferCompressionType::Auto, bool bFreeMemory = true);
-	AsyncRawBufferP					Uncompress(bool bFreeMemory = true);
+	UE_API AsyncRawBufferP					Compress(RawBufferCompressionType InCompressionType = RawBufferCompressionType::Auto, bool bFreeMemory = true);
+	UE_API AsyncRawBufferP					Uncompress(bool bFreeMemory = true);
 
-	AsyncRawBufferP					WriteToFile(const FString& InFileName, bool bFreeMemory = true);
+	UE_API AsyncRawBufferP					WriteToFile(const FString& InFileName, bool bFreeMemory = true);
 
-	AsyncRawBufferP					LoadRawBuffer(bool bInDoUncompress, bool bFreeMemory = true);
+	UE_API AsyncRawBufferP					LoadRawBuffer(bool bInDoUncompress, bool bFreeMemory = true);
 
-	void							GetAsLinearColor(TArray<FLinearColor>& Pixels) const;
-	FLinearColor					GetAsLinearColor(int PixelIndex) const;
+	UE_API void							GetAsLinearColor(TArray<FLinearColor>& Pixels) const;
+	UE_API FLinearColor					GetAsLinearColor(int PixelIndex) const;
 
-	bool							IsPadded() const;
-	size_t							GetUnpaddedSize();
-	void							CopyUnpaddedBytes(uint8* DestData);
+	UE_API bool							IsPadded() const;
+	UE_API size_t							GetUnpaddedSize();
+	UE_API void							CopyUnpaddedBytes(uint8* DestData);
 	//////////////////////////////////////////////////////////////////////////
 	/// Inline functions
 	//////////////////////////////////////////////////////////////////////////
@@ -230,3 +232,4 @@ public:
 	FORCEINLINE bool				HasData() const { check(IsInGameThread()); return Data != nullptr; }
 };
 
+#undef UE_API
diff --git a/Engine/Plugins/TextureGraph/Source/TextureGraphEngine/Device/CL/Device_CL.h b/Engine/Plugins/TextureGraph/Source/TextureGraphEngine/Device/CL/Device_CL.h
--- a/Engine/Plugins/TextureGraph/Source/TextureGraphEngine/Device/CL/Device_CL.h
+++ b/Engine/Plugins/TextureGraph/Source/TextureGraphEngine/Device/CL/Device_CL.h
@@ -3,19 +3,23 @@
 
 #include "CoreMinimal.h"
 
+#define UE_API TEXTUREGRAPHENGINE_API
+
 /**
  * 
  */
-class TEXTUREGRAPHENGINE_API Device_CL
+class Device_CL
 {
 public:
-									Device_CL();
-	virtual							~Device_CL();
+									UE_API Device_CL();
+	UE_API virtual							~Device_CL();
 
 	virtual const char*				Name() const { return "Device_CL"; }
 
 	//////////////////////////////////////////////////////////////////////////
 	/// Static functions
 	//////////////////////////////////////////////////////////////////////////
-	static Device_CL*				Get();
+	static UE_API Device_CL*				Get();
 };
+
+#undef UE_API
diff --git a/Engine/Plugins/TextureGraph/Source/TextureGraphEngine/Device/Device.h b/Engine/Plugins/TextureGraph/Source/TextureGraphEngine/Device/Device.h
--- a/Engine/Plugins/TextureGraph/Source/TextureGraphEngine/Device/Device.h
+++ b/Engine/Plugins/TextureGraph/Source/TextureGraphEngine/Device/Device.h
@@ -14,6 +14,8 @@
 #include <vector>
 #include <queue>
 
+#define UE_API TEXTUREGRAPHENGINE_API
+
 
 class Blob;
 class DeviceNativeTask;
@@ -44,7 +46,7 @@ struct CombineSplitArgs
 /// Represents a general compute device that is capable of running some
 /// computation on a blob.
 //////////////////////////////////////////////////////////////////////////
-class TEXTUREGRAPHENGINE_API Device
+class Device
 {
 	friend class					DeviceBuffer;
 	friend class					Blob;
@@ -55,7 +57,7 @@ public:
 
 	Event_ResourceUpdate			ResourceUpdateEvent;			/// Triggered when history is updated
 	
-	static FString					DeviceType_Names(DeviceType DevType);
+	static UE_API FString					DeviceType_Names(DeviceType DevType);
 
 	typedef BinCache<HashType, DeviceBufferPtrW> BufferCache;
 	typedef std::vector<DeviceBuffer*> DeviceBufferPList;
@@ -65,13 +67,13 @@ public:
 protected:
 	friend class DeviceManager;
 
-	static const double				PrintStatsInterval;		/// Print stats every <this many> milliseconds
-	static const double				DefaultBufferTransferInternal; /// Buffer transfers are updated every few milliseconds represented by this variable
+	static UE_API const double				PrintStatsInterval;		/// Print stats every <this many> milliseconds
+	static UE_API const double				DefaultBufferTransferInternal; /// Buffer transfers are updated every few milliseconds represented by this variable
 
-	static const float				DefaultMaxCacheConsumptionBeforeCollection;
-	static const float				DefaultMinCacheConsumptionBeforeCollection;
-	static const uint64				DefaultMaxBatchDeltaDefault;
-	static const double				DefaultMinLastUsageDuration;
+	static UE_API const float				DefaultMaxCacheConsumptionBeforeCollection;
+	static UE_API const float				DefaultMinCacheConsumptionBeforeCollection;
+	static UE_API const uint64				DefaultMaxBatchDeltaDefault;
+	static UE_API const double				DefaultMinLastUsageDuration;
 
 	//////////////////////////////////////////////////////////////////////////
 	float							MaxCacheConsumptionBeforeCollection 
@@ -134,45 +136,45 @@ protected:
 	mutable std::mutex				AwaitingCollectionMutex;		/// Mutex for the list below
 	DeviceBufferPList				AwaitingCollection;				/// Device buffers that are awaiting collection
 
-	virtual void					UpdateAwaitingCollection();
-	virtual void					UpdateNonAsyncPostExec();
-	virtual void					UpdateNonAsync();
+	UE_API virtual void					UpdateAwaitingCollection();
+	UE_API virtual void					UpdateNonAsyncPostExec();
+	UE_API virtual void					UpdateNonAsync();
 
-									Device(DeviceType InType, DeviceBuffer* InBufferFactory);
-	virtual							~Device();
+									UE_API Device(DeviceType InType, DeviceBuffer* InBufferFactory);
+	UE_API virtual							~Device();
 
 	// Method to wait for all the queued tasks. 
 	// Can be extremely slow and should not
 	// be used outside of tests
-	virtual AsyncBool				WaitForQueuedTasks(ENamedThreads::Type ReturnThread = ENamedThreads::UnusedAnchor);
-	virtual	void					CallDeviceUpdate();
-	virtual void					GarbageCollect();
+	UE_API virtual AsyncBool				WaitForQueuedTasks(ENamedThreads::Type ReturnThread = ENamedThreads::UnusedAnchor);
+	UE_API virtual	void					CallDeviceUpdate();
+	UE_API virtual void					GarbageCollect();
 
-	virtual void					UpdatePreExec();
+	UE_API virtual void					UpdatePreExec();
 
-	virtual bool					ShouldCollect(double TimeNow, DeviceBuffer* Buffer);
-	virtual AsyncJobResultPtr		UpdateDeviceTransfers();
+	UE_API virtual bool					ShouldCollect(double TimeNow, DeviceBuffer* Buffer);
+	UE_API virtual AsyncJobResultPtr		UpdateDeviceTransfers();
 
 	/// To be called from DeviceBuffer only (on destruction)
 public:
-	static void						CollectBuffer(DeviceBuffer* Buffer);
+	static UE_API void						CollectBuffer(DeviceBuffer* Buffer);
 
 protected:
-	virtual void					Collect(DeviceBuffer* Buffer);
+	UE_API virtual void					Collect(DeviceBuffer* Buffer);
 
-	DeviceBufferRef					AddNewRef_Internal(DeviceBuffer* Buffer);
-	virtual DeviceBufferRef			AddInternal(DeviceBufferRef& Buffer);
-	virtual void					RemoveInternal(DeviceBuffer* Buffer);
-	virtual void					Touch(HashType Hash);
-	virtual void					PrintStats();
+	UE_API DeviceBufferRef					AddNewRef_Internal(DeviceBuffer* Buffer);
+	UE_API virtual DeviceBufferRef			AddInternal(DeviceBufferRef& Buffer);
+	UE_API virtual void					RemoveInternal(DeviceBuffer* Buffer);
+	UE_API virtual void					Touch(HashType Hash);
+	UE_API virtual void					PrintStats();
 
-	virtual void					Terminate();
+	UE_API virtual void					Terminate();
 
-	virtual void					UpdateNative();
-	virtual void					PreExecThreadFunc();
-	virtual void					ExecThreadFunc();
-	virtual void					PostExecThreadFunc();
-	virtual void					MarkTaskFinished(DeviceNativeTaskPtr& Task, bool bRemoveFromOwnerQueue);
+	UE_API virtual void					UpdateNative();
+	UE_API virtual void					PreExecThreadFunc();
+	UE_API virtual void					ExecThreadFunc();
+	UE_API virtual void					PostExecThreadFunc();
+	UE_API virtual void					MarkTaskFinished(DeviceNativeTaskPtr& Task, bool bRemoveFromOwnerQueue);
 
 	/// Flushes a specific number of tasks from the _nativeTasks vector.
 	/// Set maxCount = 0 to run all the tasks currently queued up. This
@@ -185,41 +187,41 @@ protected:
 	//virtual DeviceNativeTaskPtrVec	FetchNativeTasks(size_t maxCount, size_t minCount);
 	//virtual void					FlushNativeTasks_Internal(DeviceNativeTaskPtrVec& tasks);
 
-	virtual void					Free();
+	UE_API virtual void					Free();
 
 public:
-	virtual DeviceBufferRef			Create(RawBufferPtr Raw);
-	virtual DeviceBufferRef			Create(BufferDescriptor Desc, CHashPtr Hash);
+	UE_API virtual DeviceBufferRef			Create(RawBufferPtr Raw);
+	UE_API virtual DeviceBufferRef			Create(BufferDescriptor Desc, CHashPtr Hash);
 
-	virtual AsyncDeviceBufferRef	CombineFromTiles(const CombineSplitArgs& CombineArgs);
-	virtual AsyncDeviceBufferRef	SplitToTiles(const CombineSplitArgs& SplitArgs);
+	UE_API virtual AsyncDeviceBufferRef	CombineFromTiles(const CombineSplitArgs& CombineArgs);
+	UE_API virtual AsyncDeviceBufferRef	SplitToTiles(const CombineSplitArgs& SplitArgs);
 
 
 	/// Acquire the device for use
-	virtual cti::continuable<int32>	Use() const;
+	UE_API virtual cti::continuable<int32>	Use() const;
 
-	virtual size_t					GetMaxMemory() const;
-	virtual size_t					GetMaxThreads() const;
+	UE_API virtual size_t					GetMaxMemory() const;
+	UE_API virtual size_t					GetMaxThreads() const;
 
-	virtual DeviceBufferRef			CreateCompatible_Copy(DeviceBufferRef source);
-	virtual AsyncDeviceBufferRef	Transfer(DeviceBufferRef source);
-	virtual DeviceBufferRef			Find(HashType Hash, bool Touch);
-	virtual void					Update(float dt);
-	virtual AsyncJobResultPtr		UpdateIdle();
+	UE_API virtual DeviceBufferRef			CreateCompatible_Copy(DeviceBufferRef source);
+	UE_API virtual AsyncDeviceBufferRef	Transfer(DeviceBufferRef source);
+	UE_API virtual DeviceBufferRef			Find(HashType Hash, bool Touch);
+	UE_API virtual void					Update(float dt);
+	UE_API virtual AsyncJobResultPtr		UpdateIdle();
 
 	virtual FString					Name() const { return "Generic"; }
 	
 	/// Adds a new task to the _nativeTasks vector. Thread-safe
-	virtual void					AddNativeTask(DeviceNativeTaskPtr Task);
+	UE_API virtual void					AddNativeTask(DeviceNativeTaskPtr Task);
 
-	virtual void					TransferAborted(DeviceBufferRef Buffer);
+	UE_API virtual void					TransferAborted(DeviceBufferRef Buffer);
 
 	using							Visitor = std::function<void (const std::shared_ptr<DeviceBuffer>&, uint32)>; /// Visitor callable receiving the bufferDevice visited and the index in the lru Cache
-	uint32							TraverseBufferCache(Visitor visitor);
+	UE_API uint32							TraverseBufferCache(Visitor visitor);
 
-	void							RegisterObserverSource(const DeviceObserverSourcePtr& observerSource);
-	bool							CanFree() const;
-	virtual void					ClearCache();
+	UE_API void							RegisterObserverSource(const DeviceObserverSourcePtr& observerSource);
+	UE_API bool							CanFree() const;
+	UE_API virtual void					ClearCache();
 
 	//////////////////////////////////////////////////////////////////////////
 	/// Inline functions
@@ -240,10 +242,12 @@ public:
 	//////////////////////////////////////////////////////////////////////////
 	/// Static functions
 	//////////////////////////////////////////////////////////////////////////
-	static AsyncDeviceBufferRef		SplitToTiles_Generic(const CombineSplitArgs& splitArgs);
+	static UE_API AsyncDeviceBufferRef		SplitToTiles_Generic(const CombineSplitArgs& splitArgs);
 
 	//////////////////////////////////////////////////////////////////////////
 	/// Events
 	//////////////////////////////////////////////////////////////////////////
 	FORCEINLINE Event_ResourceUpdate& E_ResourceUpdate() { return ResourceUpdateEvent; }
 };
+
+#undef UE_API
diff --git a/Engine/Plugins/TextureGraph/Source/TextureGraphEngine/Device/DeviceBuffer.h b/Engine/Plugins/TextureGraph/Source/TextureGraphEngine/Device/DeviceBuffer.h
--- a/Engine/Plugins/TextureGraph/Source/TextureGraphEngine/Device/DeviceBuffer.h
+++ b/Engine/Plugins/TextureGraph/Source/TextureGraphEngine/Device/DeviceBuffer.h
@@ -7,13 +7,15 @@
 #include "DeviceType.h"
 #include "Data/G_Collectible.h"
 
+#define UE_API TEXTUREGRAPHENGINE_API
+
 
 class Device;
 class Device_Mem;
 class BlobTransform;
 struct ResourceBindInfo;
 
-struct TEXTUREGRAPHENGINE_API BufferResult
+struct BufferResult
 {
 	std::exception_ptr				ExInner;			/// Original exception that was raised by the action
 	int32							ErrorCode = 0;		/// What is the error code
@@ -35,16 +37,16 @@ typedef std::shared_ptr<DeviceNativeTask>		DeviceNativeTaskPtr;
 typedef std::vector<DeviceNativeTaskPtr>		DeviceNativeTaskPtrVec;
 
 //////////////////////////////////////////////////////////////////////////
-class TEXTUREGRAPHENGINE_API DeviceBuffer : public G_Collectible
+class DeviceBuffer : public G_Collectible
 {
 	friend class Device;
 	friend class Blob;
 	friend class Blobber;
 
 public:
-	static const DeviceTransferChain DefaultTransferChain;		/// Default transfer Chain
-	static const DeviceTransferChain FXOnlyTransferChain;		/// Only exist on the FX device
-	static const DeviceTransferChain PersistentTransferChain;	/// Persistent
+	static UE_API const DeviceTransferChain DefaultTransferChain;		/// Default transfer Chain
+	static UE_API const DeviceTransferChain FXOnlyTransferChain;		/// Only exist on the FX device
+	static UE_API const DeviceTransferChain PersistentTransferChain;	/// Persistent
 
 
 protected:
@@ -59,30 +61,30 @@ protected:
 	bool							bIsPersistent = false;		/// Whether this is persistent or not
 	bool							Chain[static_cast<int32>(DeviceType::Count)] = {false}; /// Device transfer Chain
 
-	virtual CHashPtr				CalcHash();
+	UE_API virtual CHashPtr				CalcHash();
 
 	/// Only accessible from Device. Should not be used elsewhere
 	virtual DeviceBuffer*			CreateFromRaw(RawBufferPtr RawObj) = 0;
 	virtual DeviceBuffer*			CreateFromDesc(BufferDescriptor Desc, CHashPtr HashValue) = 0;
-	virtual DeviceBuffer*			CopyFrom(DeviceBuffer* RHS);
-	virtual DeviceBuffer*			Clone();
+	UE_API virtual DeviceBuffer*			CopyFrom(DeviceBuffer* RHS);
+	UE_API virtual DeviceBuffer*			Clone();
 
 	/// DO NOT MAKE THIS METHOD PUBLIC!
 	/// Its not meant to be accessible outside
-	virtual AsyncBufferResultPtr	UpdateRaw(RawBufferPtr RawObj);
+	UE_API virtual AsyncBufferResultPtr	UpdateRaw(RawBufferPtr RawObj);
 
 	/// To be called from the Device owner only!
-	virtual void					ReleaseNative();
+	UE_API virtual void					ReleaseNative();
 
-	virtual void					Touch(uint64 BatchId);
+	UE_API virtual void					Touch(uint64 BatchId);
 
 	FORCEINLINE void				MarkForCollection() { bMarkedForCollection = true; }
 
 public:
-									DeviceBuffer(Device* Dev, BufferDescriptor NewDesc, CHashPtr NewHash);
-									DeviceBuffer(Device* Dev, RawBufferPtr RawObj);
+									UE_API DeviceBuffer(Device* Dev, BufferDescriptor NewDesc, CHashPtr NewHash);
+									UE_API DeviceBuffer(Device* Dev, RawBufferPtr RawObj);
 									DeviceBuffer(const DeviceBuffer& RHS) = delete;	/// non-copyable
-	virtual							~DeviceBuffer() override;
+	UE_API virtual							~DeviceBuffer() override;
 
 	//////////////////////////////////////////////////////////////////////////
 	/// Get the RawObj memory buffer against this device buffer. Note that
@@ -95,9 +97,9 @@ public:
 	/// attempting
 	//////////////////////////////////////////////////////////////////////////
 	virtual RawBufferPtr			Raw_Now() = 0;	/// SLOW: Read above
-	virtual AsyncRawBufferPtr		Raw();			/// SLOW: Read above [use Blob::Raw instead. This is to be used from Blob only]
+	UE_API virtual AsyncRawBufferPtr		Raw();			/// SLOW: Read above [use Blob::Raw instead. This is to be used from Blob only]
 
-	AsyncRawBufferPtr				GetRawOrMaketIt();
+	UE_API AsyncRawBufferPtr				GetRawOrMaketIt();
 
 protected:
 	virtual AsyncBufferResultPtr	TransferFrom(DeviceBufferRef& Source) = 0;
@@ -109,27 +111,27 @@ public:
 	virtual size_t					DeviceNative_MemSize() const { return 0; }
 
 	virtual AsyncBufferResultPtr	Bind(const BlobTransform* Transform, const ResourceBindInfo& BindInfo) = 0;
-	virtual AsyncBufferResultPtr	Unbind(const BlobTransform* Transform, const ResourceBindInfo& BindInfo);
+	UE_API virtual AsyncBufferResultPtr	Unbind(const BlobTransform* Transform, const ResourceBindInfo& BindInfo);
 
-	virtual AsyncBufferResultPtr	Flush(const ResourceBindInfo& BindInfo);
-	virtual bool					IsValid() const;
-	virtual bool					IsNull() const;
+	UE_API virtual AsyncBufferResultPtr	Flush(const ResourceBindInfo& BindInfo);
+	UE_API virtual bool					IsValid() const;
+	UE_API virtual bool					IsNull() const;
 
 	/// With base implementations
-	virtual CHashPtr				Hash(bool Calculate = true);
-	virtual bool					IsCompatible(Device* Dev) const;
+	UE_API virtual CHashPtr				Hash(bool Calculate = true);
+	UE_API virtual bool					IsCompatible(Device* Dev) const;
 
 	virtual RawBufferPtr			Min() { return nullptr; }
 	virtual RawBufferPtr			Max() { return nullptr; }
 
-	virtual void					SetHash(CHashPtr NewHash);
-	virtual AsyncPrepareResult		PrepareForWrite(const ResourceBindInfo& BindInfo);
+	UE_API virtual void					SetHash(CHashPtr NewHash);
+	UE_API virtual AsyncPrepareResult		PrepareForWrite(const ResourceBindInfo& BindInfo);
 
-	void							SetDeviceTransferChain(const DeviceTransferChain& Chain, bool bPersistent = false);
-	DeviceTransferChain				GetDeviceTransferChain(bool* Persistent = nullptr) const;
-	virtual Device*					GetDowngradeDevice() const;
-	virtual Device*					GetUpgradeDevice() const;
-	virtual bool					IsTransient() const;
+	UE_API void							SetDeviceTransferChain(const DeviceTransferChain& Chain, bool bPersistent = false);
+	UE_API DeviceTransferChain				GetDeviceTransferChain(bool* Persistent = nullptr) const;
+	UE_API virtual Device*					GetDowngradeDevice() const;
+	UE_API virtual Device*					GetUpgradeDevice() const;
+	UE_API virtual bool					IsTransient() const;
 
 	//////////////////////////////////////////////////////////////////////////
 	/// Inline functions
@@ -147,14 +149,14 @@ public:
 
 //typedef std::shared_ptr<DeviceBuffer>	DeviceBufferPtr;
 
-class TEXTUREGRAPHENGINE_API DeviceBufferPtr : public std::shared_ptr<DeviceBuffer>
+class DeviceBufferPtr : public std::shared_ptr<DeviceBuffer>
 {
 public:
 	DeviceBufferPtr() : std::shared_ptr<DeviceBuffer>() {}
-	DeviceBufferPtr(DeviceBuffer* buffer);
-	DeviceBufferPtr(const DeviceBufferPtr& RHS);
-	DeviceBufferPtr(const std::shared_ptr<DeviceBuffer>& RHS);
-	~DeviceBufferPtr();
+	UE_API DeviceBufferPtr(DeviceBuffer* buffer);
+	UE_API DeviceBufferPtr(const DeviceBufferPtr& RHS);
+	UE_API DeviceBufferPtr(const std::shared_ptr<DeviceBuffer>& RHS);
+	UE_API ~DeviceBufferPtr();
 };
 
 typedef std::weak_ptr<DeviceBuffer>		DeviceBufferPtrW;
@@ -166,7 +168,7 @@ typedef std::weak_ptr<DeviceBuffer>		DeviceBufferPtrW;
 //////////////////////////////////////////////////////////////////////////
 /// DeviceBufferRef
 //////////////////////////////////////////////////////////////////////////
-class TEXTUREGRAPHENGINE_API DeviceBufferRef
+class DeviceBufferRef
 {
 public:
 	
@@ -178,17 +180,17 @@ private:
 
 	DeviceBufferPtr					Buffer;				/// The buffer info pointer
 
-	void							Clear();
+	UE_API void							Clear();
 public:
 									DeviceBufferRef() {}
-	explicit						DeviceBufferRef(DeviceBufferPtr RHS);
-	explicit						DeviceBufferRef(DeviceBuffer* RHS);
+	UE_API explicit						DeviceBufferRef(DeviceBufferPtr RHS);
+	UE_API explicit						DeviceBufferRef(DeviceBuffer* RHS);
 
-									~DeviceBufferRef();
+									UE_API ~DeviceBufferRef();
 
-	DeviceBufferRef&				operator = (const DeviceBufferRef& RHS);
+	UE_API DeviceBufferRef&				operator = (const DeviceBufferRef& RHS);
 
-	DeviceType						GetDeviceType() const; /// Need complete Device Definition so cannot be inlined here
+	UE_API DeviceType						GetDeviceType() const; /// Need complete Device Definition so cannot be inlined here
 
 	//////////////////////////////////////////////////////////////////////////
 	/// Inline functions
@@ -211,3 +213,5 @@ public:
 
 typedef cti::continuable<DeviceBufferRef>		AsyncDeviceBufferRef;
 typedef std::function<void(DeviceBufferRef)>	DeviceBufferRef_Callback;
+
+#undef UE_API
diff --git a/Engine/Plugins/TextureGraph/Source/TextureGraphEngine/Device/DeviceManager.h b/Engine/Plugins/TextureGraph/Source/TextureGraphEngine/Device/DeviceManager.h
--- a/Engine/Plugins/TextureGraph/Source/TextureGraphEngine/Device/DeviceManager.h
+++ b/Engine/Plugins/TextureGraph/Source/TextureGraphEngine/Device/DeviceManager.h
@@ -7,28 +7,30 @@
 #include "Helper/Promise.h"
 #include <memory>
 
+#define UE_API TEXTUREGRAPHENGINE_API
+
 class Device;
 class Device_Mem;
 
 using AsyncBool = cti::continuable<bool>;
 
-class TEXTUREGRAPHENGINE_API DeviceManager
+class DeviceManager
 {
 private:
 	Device*							Devices[(size_t)DeviceType::Count] = {};	/// The devices available within the system
 	Device_Mem*						DefaultDevice = nullptr;					/// What is the default device on the system. This is always the CPU device
 
-	void							InitDevices();
-	void							ReleaseDevices();
+	UE_API void							InitDevices();
+	UE_API void							ReleaseDevices();
 
 public:
-									DeviceManager();
-	virtual							~DeviceManager();
+									UE_API DeviceManager();
+	UE_API virtual							~DeviceManager();
 
-	void							Update(float dt);
-	AsyncBool						WaitForQueuedTasks(ENamedThreads::Type returnThread = ENamedThreads::UnusedAnchor);
-	Device*							GetDevice(DeviceType type) const;
-	Device*							GetDevice(size_t Index) const;
+	UE_API void							Update(float dt);
+	UE_API AsyncBool						WaitForQueuedTasks(ENamedThreads::Type returnThread = ENamedThreads::UnusedAnchor);
+	UE_API Device*							GetDevice(DeviceType type) const;
+	UE_API Device*							GetDevice(size_t Index) const;
 
 	//////////////////////////////////////////////////////////////////////////
 	/// Inline functions
@@ -38,3 +40,5 @@ public:
 };
 
 typedef std::unique_ptr<DeviceManager> DeviceManagerPtr;
+
+#undef UE_API
diff --git a/Engine/Plugins/TextureGraph/Source/TextureGraphEngine/Device/DeviceNativeTask.h b/Engine/Plugins/TextureGraph/Source/TextureGraphEngine/Device/DeviceNativeTask.h
--- a/Engine/Plugins/TextureGraph/Source/TextureGraphEngine/Device/DeviceNativeTask.h
+++ b/Engine/Plugins/TextureGraph/Source/TextureGraphEngine/Device/DeviceNativeTask.h
@@ -8,6 +8,8 @@
 #include "Helper/Util.h"
 #include "Async/TaskGraphInterfaces.h"
 
+#define UE_API TEXTUREGRAPHENGINE_API
+
 class Device;
 class DeviceNativeTask;
 typedef cti::continuable<DeviceNativeTask>	AsyncDeviceNativeTask;
@@ -21,11 +23,11 @@ typedef std::vector<DeviceNativeTaskPtr>	DeviceNativeTaskPtrVec;
 /// Represents a Task that runs natively on a Device (e.g. Rendering 
 /// thread on an FX device or CPU and GPU threads on an OpenCL device)
 //////////////////////////////////////////////////////////////////////////
-class TEXTUREGRAPHENGINE_API DeviceNativeTask
+class DeviceNativeTask
 {
 private:
-	static std::atomic<int64_t>		GTaskId;				/// Task ID
-	static int64_t					GetNewTaskId();
+	static UE_API std::atomic<int64_t>		GTaskId;				/// Task ID
+	static UE_API int64_t					GetNewTaskId();
 
 	std::promise<int32>				Promise;				/// Our promise
 	std::shared_future<int32>		Future;					/// The future for our promise
@@ -48,57 +50,57 @@ protected:
 
 	int64_t							BatchId = 0;			/// What's the batch id of this native Task
 
-	void							Reset();
-	void							WaitSelf();
-	bool							WaitSelfFor(int32 Seconds);
+	UE_API void							Reset();
+	UE_API void							WaitSelf();
+	UE_API bool							WaitSelfFor(int32 Seconds);
 
-	void							SetPromise(int32 Value);
-	FString							DebugWaitChain();
+	UE_API void							SetPromise(int32 Value);
+	UE_API FString							DebugWaitChain();
 
-	virtual cti::continuable<int32>	GenericExecAsync(ENamedThreads::Type CallingThread, ENamedThreads::Type ReturnThread);
+	UE_API virtual cti::continuable<int32>	GenericExecAsync(ENamedThreads::Type CallingThread, ENamedThreads::Type ReturnThread);
 
-	void							AdjustPrevPriority(DeviceNativeTask* Ref, DeviceNativeTask* Parent);
-	void							AdjustPrevPriority(DeviceNativeTask* Ref, DeviceNativeTask* Parent, DeviceNativeTask* PrevDeps);
-	virtual void					SetPriority(DeviceNativeTask* Ref);
+	UE_API void							AdjustPrevPriority(DeviceNativeTask* Ref, DeviceNativeTask* Parent);
+	UE_API void							AdjustPrevPriority(DeviceNativeTask* Ref, DeviceNativeTask* Parent, DeviceNativeTask* PrevDeps);
+	UE_API virtual void					SetPriority(DeviceNativeTask* Ref);
 
 public:
-									DeviceNativeTask(int32 TaskPriority, FString TaskName);
-									DeviceNativeTask(DeviceNativeTaskPtrVec& PrevTasks, int32 TaskPriority, FString TaskName);
-	virtual							~DeviceNativeTask();
-
-	virtual bool					WaitFor(int32 Seconds);
-	virtual void					Wait();
-	virtual cti::continuable<int32>	PreExecAsync(ENamedThreads::Type CallingThread, ENamedThreads::Type ReturnThread);
-	virtual void					PreExec();
+									UE_API DeviceNativeTask(int32 TaskPriority, FString TaskName);
+									UE_API DeviceNativeTask(DeviceNativeTaskPtrVec& PrevTasks, int32 TaskPriority, FString TaskName);
+	UE_API virtual							~DeviceNativeTask();
+
+	UE_API virtual bool					WaitFor(int32 Seconds);
+	UE_API virtual void					Wait();
+	UE_API virtual cti::continuable<int32>	PreExecAsync(ENamedThreads::Type CallingThread, ENamedThreads::Type ReturnThread);
+	UE_API virtual void					PreExec();
 	virtual int32					Exec() = 0;
-	virtual cti::continuable<int32>	ExecAsync(ENamedThreads::Type CallingThread, ENamedThreads::Type ReturnThread);
-	virtual void					PostExec();
-	virtual cti::continuable<int32>	WaitAsync(ENamedThreads::Type ReturnThread = ENamedThreads::UnusedAnchor);
+	UE_API virtual cti::continuable<int32>	ExecAsync(ENamedThreads::Type CallingThread, ENamedThreads::Type ReturnThread);
+	UE_API virtual void					PostExec();
+	UE_API virtual cti::continuable<int32>	WaitAsync(ENamedThreads::Type ReturnThread = ENamedThreads::UnusedAnchor);
 
 	virtual Device*					GetTargetDevice() const = 0;
 
 #if (DEVICE_NATIVE_TASKS_CHECK_CYCLES == 1)
-	void							CheckQueued(DeviceNativeTask* Ref);
-	void							DebugVerifyDepsQueued(const DeviceNativeTaskPtrVec& Queue, const DeviceNativeTaskPtrVec& AllItems);
+	UE_API void							CheckQueued(DeviceNativeTask* Ref);
+	UE_API void							DebugVerifyDepsQueued(const DeviceNativeTaskPtrVec& Queue, const DeviceNativeTaskPtrVec& AllItems);
 #endif 
 
-	void							FixPriorities();
+	UE_API void							FixPriorities();
 
 	/// Waits for the _prev Task to finish before executing. Its the responsibility of the system
 	/// to ensure that both the _prev
-	virtual ENamedThreads::Type		GetExecutionThread() const;
-	virtual FString					GetName() const;
-	virtual FString					GetDebugName() const;
-	virtual void					Terminate();
-
-	virtual void					AddPrev(DeviceNativeTaskPtr PrevTasks);
-	virtual void					AddPrev(const DeviceNativeTaskPtrVec& Tasks);
-	virtual void					SetPrev(const DeviceNativeTaskPtrVec& PrevTasks);
-	virtual void					SetPriority(int32 TaskPriority, bool AdjustPrev);
-	bool							IsHigherPriorityThan(const DeviceNativeTask& RHS) const;
-
-	void							CheckCyclicalDependency(std::unordered_set<DeviceNativeTask*>& TransferChain);
-	virtual bool					DebugCompleteCheck();
+	UE_API virtual ENamedThreads::Type		GetExecutionThread() const;
+	UE_API virtual FString					GetName() const;
+	UE_API virtual FString					GetDebugName() const;
+	UE_API virtual void					Terminate();
+
+	UE_API virtual void					AddPrev(DeviceNativeTaskPtr PrevTasks);
+	UE_API virtual void					AddPrev(const DeviceNativeTaskPtrVec& Tasks);
+	UE_API virtual void					SetPrev(const DeviceNativeTaskPtrVec& PrevTasks);
+	UE_API virtual void					SetPriority(int32 TaskPriority, bool AdjustPrev);
+	UE_API bool							IsHigherPriorityThan(const DeviceNativeTask& RHS) const;
+
+	UE_API void							CheckCyclicalDependency(std::unordered_set<DeviceNativeTask*>& TransferChain);
+	UE_API virtual bool					DebugCompleteCheck();
 	virtual bool					IsAsync() const { return bIsAsync; }
 
 	//////////////////////////////////////////////////////////////////////////
@@ -125,7 +127,7 @@ public:
 //////////////////////////////////////////////////////////////////////////
 
 //////////////////////////////////////////////////////////////////////////
-class TEXTUREGRAPHENGINE_API DeviceNativeTask_Generic : public DeviceNativeTask
+class DeviceNativeTask_Generic : public DeviceNativeTask
 {
 	friend class Device;
 
@@ -134,9 +136,9 @@ protected:
 	ENamedThreads::Type				ThreadId = ENamedThreads::AnyThread; /// What thread is this going to run on
 
 public:
-									DeviceNativeTask_Generic(Device* TargetDevice, FString TaskName, ENamedThreads::Type TaskThreadId = ENamedThreads::AnyThread, 
+									UE_API DeviceNativeTask_Generic(Device* TargetDevice, FString TaskName, ENamedThreads::Type TaskThreadId = ENamedThreads::AnyThread, 
 										bool IsTaskAsync = false, int32 TaskPriority = (int32)E_Priority::kNormal);
-									DeviceNativeTask_Generic(DeviceNativeTaskPtrVec& PrevDeps, Device* TargetDevice, FString TaskName, ENamedThreads::Type TaskThreadId = ENamedThreads::AnyThread, 
+									UE_API DeviceNativeTask_Generic(DeviceNativeTaskPtrVec& PrevDeps, Device* TargetDevice, FString TaskName, ENamedThreads::Type TaskThreadId = ENamedThreads::AnyThread, 
 										bool IsTAskAsync = false, int32 TaskPriority = (int32)E_Priority::kNormal);
 
 	virtual Device*					GetTargetDevice() const override { return Dev; }
@@ -153,24 +155,25 @@ public:
 
 //////////////////////////////////////////////////////////////////////////
 typedef TFunction<int32()>			DeviceNativeTask_Lambda_Func;
-class TEXTUREGRAPHENGINE_API DeviceNativeTask_Lambda : public DeviceNativeTask_Generic
+class DeviceNativeTask_Lambda : public DeviceNativeTask_Generic
 {
 protected:
 	DeviceNativeTask_Lambda_Func	Callback = nullptr;	/// The lambda callback
 
-	virtual int32					Exec() override;
+	UE_API virtual int32					Exec() override;
 
 public:
-									DeviceNativeTask_Lambda(DeviceNativeTask_Lambda_Func LambdaCallback, Device* Dev, FString TaskName, 
+									UE_API DeviceNativeTask_Lambda(DeviceNativeTask_Lambda_Func LambdaCallback, Device* Dev, FString TaskName, 
 										ENamedThreads::Type TaskThreadId = ENamedThreads::AnyThread, bool IsTaskAsync = false, int32 TaskPriority = (int32)E_Priority::kNormal);
 
 	//////////////////////////////////////////////////////////////////////////
 	/// Static functions
 	//////////////////////////////////////////////////////////////////////////
-	static std::shared_ptr<DeviceNativeTask> Create(Device* device, int32 TaskPriority, FString TaskName, DeviceNativeTask_Lambda_Func callback);
-	static std::shared_ptr<DeviceNativeTask> Create(Device* device, ENamedThreads::Type threadId, bool isAsync, int32 TaskPriority, FString TaskName, DeviceNativeTask_Lambda_Func callback);
+	static UE_API std::shared_ptr<DeviceNativeTask> Create(Device* device, int32 TaskPriority, FString TaskName, DeviceNativeTask_Lambda_Func callback);
+	static UE_API std::shared_ptr<DeviceNativeTask> Create(Device* device, ENamedThreads::Type threadId, bool isAsync, int32 TaskPriority, FString TaskName, DeviceNativeTask_Lambda_Func callback);
 
 	//static cti::continuable<int32>	Create_Promise(Device* device, int32 TaskPriority, DeviceNativeTask_Lambda_Func callback);
 	//static cti::continuable<int32>	Create_Promise(Device* device, ENamedThreads::Type threadId, bool isAsync, int32 TaskPriority, DeviceNativeTask_Lambda_Func callback);
 };
 
+#undef UE_API
diff --git a/Engine/Plugins/TextureGraph/Source/TextureGraphEngine/Device/DeviceObserverSource.h b/Engine/Plugins/TextureGraph/Source/TextureGraphEngine/Device/DeviceObserverSource.h
--- a/Engine/Plugins/TextureGraph/Source/TextureGraphEngine/Device/DeviceObserverSource.h
+++ b/Engine/Plugins/TextureGraph/Source/TextureGraphEngine/Device/DeviceObserverSource.h
@@ -7,6 +7,8 @@
 #include <memory>
 #include <vector>
 
+#define UE_API TEXTUREGRAPHENGINE_API
+
 
 class Blob;
 class DeviceNativeTask;
@@ -18,7 +20,7 @@ typedef std::shared_ptr<JobResult>	JobResultPtr;
 typedef cti::continuable<JobResultPtr> AsyncJobResultPtr;
 
 //////////////////////////////////////////////////////////////////////////
-class TEXTUREGRAPHENGINE_API DeviceObserverSource
+class DeviceObserverSource
 {
 public:
 	struct HashNDesc {
@@ -29,8 +31,8 @@ public:
 	using HashNDescArray			= std::vector<HashNDesc>;
 	using HashArray					= std::vector<HashType>;
 
-	void							AddBuffer(const DeviceBuffer* Buffer, HashType Hash, const BufferDescriptor& Desc);
-	void							RemoveBuffer(const DeviceBuffer* buffer, HashType hash, HashType prevHash);
+	UE_API void							AddBuffer(const DeviceBuffer* Buffer, HashType Hash, const BufferDescriptor& Desc);
+	UE_API void							RemoveBuffer(const DeviceBuffer* buffer, HashType hash, HashType prevHash);
 
 protected:
 	/// Protected interface of emitters called by the scheduler to notify the observers
@@ -44,7 +46,7 @@ protected:
 	uint32							Version = 1;
 
 	/// Trigger the broadcast of the changes, this call is issued from the Device::Update function
-	void							Broadcast();
+	UE_API void							Broadcast();
 
 	/// The customisable handler method triggered from Broadcast if any buffers were added or removed
 	virtual void					DeviceBuffersUpdated(HashNDescArray&& AddedBuffers, HashArray&& RemovedBuffers) {}
@@ -57,3 +59,4 @@ public:
 };
 typedef std::shared_ptr<DeviceObserverSource> DeviceObserverSourcePtr;
 
+#undef UE_API
diff --git a/Engine/Plugins/TextureGraph/Source/TextureGraphEngine/Device/Disk/DeviceBuffer_Disk.h b/Engine/Plugins/TextureGraph/Source/TextureGraphEngine/Device/Disk/DeviceBuffer_Disk.h
--- a/Engine/Plugins/TextureGraph/Source/TextureGraphEngine/Device/Disk/DeviceBuffer_Disk.h
+++ b/Engine/Plugins/TextureGraph/Source/TextureGraphEngine/Device/Disk/DeviceBuffer_Disk.h
@@ -5,6 +5,8 @@
 #include "Device/DeviceBuffer.h"
 #include "Data/RawBuffer.h"
 
+#define UE_API TEXTUREGRAPHENGINE_API
+
 class Device;
 class Device_Disk;
 
@@ -12,33 +14,35 @@ class Device_Disk;
 /// This is simply a wrapper around RawBuffer. Must be owned by the 
 /// CPU device (though that can change in the future)
 //////////////////////////////////////////////////////////////////////////
-class TEXTUREGRAPHENGINE_API DeviceBuffer_Disk : public DeviceBuffer
+class DeviceBuffer_Disk : public DeviceBuffer
 {
 protected:
 	RawBufferPtr					RawObj;				/// The underlying raw buffer
 
-	virtual DeviceBuffer*			CreateFromRaw(RawBufferPtr NewRawObj) override;
-	virtual DeviceBuffer*			CreateFromDesc(BufferDescriptor NewDesc, CHashPtr NewHashValue) override;
+	UE_API virtual DeviceBuffer*			CreateFromRaw(RawBufferPtr NewRawObj) override;
+	UE_API virtual DeviceBuffer*			CreateFromDesc(BufferDescriptor NewDesc, CHashPtr NewHashValue) override;
 
-	virtual AsyncBufferResultPtr	TransferFrom(DeviceBufferRef& Source) override;
+	UE_API virtual AsyncBufferResultPtr	TransferFrom(DeviceBufferRef& Source) override;
 
 public:
-									DeviceBuffer_Disk(Device_Disk* Dev, BufferDescriptor Desc, CHashPtr HashValue);
-									DeviceBuffer_Disk(Device_Disk* Dev, RawBufferPtr RawObj);
-	virtual							~DeviceBuffer_Disk() override;
+									UE_API DeviceBuffer_Disk(Device_Disk* Dev, BufferDescriptor Desc, CHashPtr HashValue);
+									UE_API DeviceBuffer_Disk(Device_Disk* Dev, RawBufferPtr RawObj);
+	UE_API virtual							~DeviceBuffer_Disk() override;
 
 	//////////////////////////////////////////////////////////////////////////
 	/// DeviceBuffer implementation
 	//////////////////////////////////////////////////////////////////////////
-	virtual RawBufferPtr			Raw_Now() override;
-	virtual size_t					MemSize() const override;
-	virtual AsyncBufferResultPtr	Bind(const BlobTransform* Transform, const ResourceBindInfo& BindInfo) override;
+	UE_API virtual RawBufferPtr			Raw_Now() override;
+	UE_API virtual size_t					MemSize() const override;
+	UE_API virtual AsyncBufferResultPtr	Bind(const BlobTransform* Transform, const ResourceBindInfo& BindInfo) override;
 
 	/// With base implementations
 	//virtual void					Release() override;
-	virtual bool					IsCompatible(Device* Dev) const override;
+	UE_API virtual bool					IsCompatible(Device* Dev) const override;
 
 	//////////////////////////////////////////////////////////////////////////
 	/// Inline functions
 	//////////////////////////////////////////////////////////////////////////
 };
+
+#undef UE_API
diff --git a/Engine/Plugins/TextureGraph/Source/TextureGraphEngine/Device/Disk/Device_Disk.h b/Engine/Plugins/TextureGraph/Source/TextureGraphEngine/Device/Disk/Device_Disk.h
--- a/Engine/Plugins/TextureGraph/Source/TextureGraphEngine/Device/Disk/Device_Disk.h
+++ b/Engine/Plugins/TextureGraph/Source/TextureGraphEngine/Device/Disk/Device_Disk.h
@@ -4,27 +4,31 @@
 #include "CoreMinimal.h"
 #include "Device/Mem/Device_Mem.h"
 
-class TEXTUREGRAPHENGINE_API Device_Disk : public Device_Mem
+#define UE_API TEXTUREGRAPHENGINE_API
+
+class Device_Disk : public Device_Mem
 {
 protected:
 	FString							BaseDir;				/// The base directory for 
 
 public:
-									Device_Disk();
-	virtual							~Device_Disk() override;
+									UE_API Device_Disk();
+	UE_API virtual							~Device_Disk() override;
 
 	virtual FString					Name() const override { return "Device_Disk"; }
-	virtual void					Update(float Delta) override;
-	virtual FString					SetBaseDirectory(const FString& InBaseDir, bool MigrateExisting = true);
-	FString							GetCacheFilename(HashType HashValue) const;
+	UE_API virtual void					Update(float Delta) override;
+	UE_API virtual FString					SetBaseDirectory(const FString& InBaseDir, bool MigrateExisting = true);
+	UE_API FString							GetCacheFilename(HashType HashValue) const;
 
 	//////////////////////////////////////////////////////////////////////////
 	/// Static functions
 	//////////////////////////////////////////////////////////////////////////
-	static Device_Disk*				Get();
+	static UE_API Device_Disk*				Get();
 
 	//////////////////////////////////////////////////////////////////////////
 	/// Inline function
 	//////////////////////////////////////////////////////////////////////////
 	FORCEINLINE const FString&		GetBaseDirectory() const { return BaseDir; }
 };
+
+#undef UE_API
diff --git a/Engine/Plugins/TextureGraph/Source/TextureGraphEngine/Device/FX/DeviceBuffer_FX.h b/Engine/Plugins/TextureGraph/Source/TextureGraphEngine/Device/FX/DeviceBuffer_FX.h
--- a/Engine/Plugins/TextureGraph/Source/TextureGraphEngine/Device/FX/DeviceBuffer_FX.h
+++ b/Engine/Plugins/TextureGraph/Source/TextureGraphEngine/Device/FX/DeviceBuffer_FX.h
@@ -4,51 +4,53 @@
 #include "CoreMinimal.h"
 #include "Device/DeviceBuffer.h"
 
+#define UE_API TEXTUREGRAPHENGINE_API
+
 class Device_FX;
 class DeviceBuffer_Mem;
 
 class Tex;
 typedef std::shared_ptr<Tex>		TexPtr;
 
-class TEXTUREGRAPHENGINE_API DeviceBuffer_FX : public DeviceBuffer
+class DeviceBuffer_FX : public DeviceBuffer
 {
 	friend class Device_FX;
 protected: 
 	TexPtr							Texture;				/// The underlying texture object
 
-	AsyncBufferResultPtr			AllocateTexture();
-	AsyncBufferResultPtr			AllocateRenderTarget();
-	void							UpdateRenderTarget(TexPtr TextureObj);
+	UE_API AsyncBufferResultPtr			AllocateTexture();
+	UE_API AsyncBufferResultPtr			AllocateRenderTarget();
+	UE_API void							UpdateRenderTarget(TexPtr TextureObj);
 
-	virtual DeviceBuffer*			CreateFromRaw(RawBufferPtr RawObj) override;
-	virtual DeviceBuffer*			CreateFromDesc(BufferDescriptor NewDesc, CHashPtr NewHashValue) override;
-	virtual DeviceBuffer*			CopyFrom(DeviceBuffer* RHS) override;
-	virtual AsyncBufferResultPtr	UpdateRaw(RawBufferPtr RawObj) override;
+	UE_API virtual DeviceBuffer*			CreateFromRaw(RawBufferPtr RawObj) override;
+	UE_API virtual DeviceBuffer*			CreateFromDesc(BufferDescriptor NewDesc, CHashPtr NewHashValue) override;
+	UE_API virtual DeviceBuffer*			CopyFrom(DeviceBuffer* RHS) override;
+	UE_API virtual AsyncBufferResultPtr	UpdateRaw(RawBufferPtr RawObj) override;
 
-	virtual CHashPtr				CalcHash() override;
+	UE_API virtual CHashPtr				CalcHash() override;
 
-	virtual AsyncBufferResultPtr	TransferFrom(DeviceBufferRef& Source) override;
+	UE_API virtual AsyncBufferResultPtr	TransferFrom(DeviceBufferRef& Source) override;
 
 public:
-									DeviceBuffer_FX(Device_FX* Dev, BufferDescriptor Desc, CHashPtr NewHashValue);
-									DeviceBuffer_FX(Device_FX* Dev, TexPtr TextureObj, RawBufferPtr RawObj);
-									DeviceBuffer_FX(Device_FX* Dev, TexPtr TextureObj, BufferDescriptor Desc, CHashPtr NewHashValue);
-									DeviceBuffer_FX(Device_FX* Dev, RawBufferPtr RawObj);
-	virtual							~DeviceBuffer_FX() override;
+									UE_API DeviceBuffer_FX(Device_FX* Dev, BufferDescriptor Desc, CHashPtr NewHashValue);
+									UE_API DeviceBuffer_FX(Device_FX* Dev, TexPtr TextureObj, RawBufferPtr RawObj);
+									UE_API DeviceBuffer_FX(Device_FX* Dev, TexPtr TextureObj, BufferDescriptor Desc, CHashPtr NewHashValue);
+									UE_API DeviceBuffer_FX(Device_FX* Dev, RawBufferPtr RawObj);
+	UE_API virtual							~DeviceBuffer_FX() override;
 
 	//////////////////////////////////////////////////////////////////////////
 	/// DeviceBuffer Implementation
 	//////////////////////////////////////////////////////////////////////////
-	virtual RawBufferPtr			Raw_Now() override;
-	virtual size_t					MemSize() const override;
-	virtual size_t					DeviceNative_MemSize() const override;
-	virtual AsyncBufferResultPtr	Bind(const BlobTransform* Transform, const ResourceBindInfo& BindInfo) override;
-	virtual AsyncBufferResultPtr	Unbind(const BlobTransform* Transform, const ResourceBindInfo& BindInfo) override;
-	virtual AsyncBufferResultPtr	Flush(const ResourceBindInfo& BindInfo) override;
-	virtual AsyncPrepareResult		PrepareForWrite(const ResourceBindInfo& BindInfo) override;
-	virtual bool					IsNull() const override;
+	UE_API virtual RawBufferPtr			Raw_Now() override;
+	UE_API virtual size_t					MemSize() const override;
+	UE_API virtual size_t					DeviceNative_MemSize() const override;
+	UE_API virtual AsyncBufferResultPtr	Bind(const BlobTransform* Transform, const ResourceBindInfo& BindInfo) override;
+	UE_API virtual AsyncBufferResultPtr	Unbind(const BlobTransform* Transform, const ResourceBindInfo& BindInfo) override;
+	UE_API virtual AsyncBufferResultPtr	Flush(const ResourceBindInfo& BindInfo) override;
+	UE_API virtual AsyncPrepareResult		PrepareForWrite(const ResourceBindInfo& BindInfo) override;
+	UE_API virtual bool					IsNull() const override;
 
-	virtual void					ReleaseNative() override;
+	UE_API virtual void					ReleaseNative() override;
 
 	//////////////////////////////////////////////////////////////////////////
 	/// Inline functions
@@ -57,3 +59,5 @@ public:
 };
 
 typedef std::shared_ptr<DeviceBuffer_FX> DeviceBuffer_FXPtr;
+
+#undef UE_API
diff --git a/Engine/Plugins/TextureGraph/Source/TextureGraphEngine/Device/FX/Device_FX.h b/Engine/Plugins/TextureGraph/Source/TextureGraphEngine/Device/FX/Device_FX.h
--- a/Engine/Plugins/TextureGraph/Source/TextureGraphEngine/Device/FX/Device_FX.h
+++ b/Engine/Plugins/TextureGraph/Source/TextureGraphEngine/Device/FX/Device_FX.h
@@ -10,6 +10,8 @@
 #include <list>
 #include <unordered_map>
 
+#define UE_API TEXTUREGRAPHENGINE_API
+
 typedef cti::continuable<TexPtr>	AsyncTexPtr;
 class RenderMaterial_FX;
 typedef std::shared_ptr<RenderMaterial_FX> RenderMaterial_FXPtr;
@@ -17,16 +19,16 @@ typedef std::shared_ptr<class TexArray>	TexArrayPtr;
 class UTextureRenderTarget;
 class IRendererModule;
 
-struct TEXTUREGRAPHENGINE_API DrawTilesSettings
+struct DrawTilesSettings
 {
 	T_Tiles<FIntPoint>				Position;						/// The position where to put the tiles
 	T_Tiles<FIntPoint>				Size;							/// The size of the tiles
 };
 
-class TEXTUREGRAPHENGINE_API Device_FX : public Device,  public FGCObject
+class Device_FX : public Device,  public FGCObject
 {
 private: 
-	static size_t					s_maxRenderBatch;
+	static UE_API size_t					s_maxRenderBatch;
 
 	typedef std::list<TexPtr>		TextureNodeList;
 
@@ -45,58 +47,58 @@ private:
 
 	IRendererModule*				RendererModule;					/// Useful rendering functionalities wrapped in this module
 
-	void							GCTextures();
-	virtual void					Free() override;
-	virtual void					ClearCache() override;
+	UE_API void							GCTextures();
+	UE_API virtual void					Free() override;
+	UE_API virtual void					ClearCache() override;
 
 	//////////////////////////////////////////////////////////////////////////
 	/// FGCObject
 	//////////////////////////////////////////////////////////////////////////
-	virtual void					AddReferencedObjects(FReferenceCollector& Collector) override;
-	virtual FString					GetReferencerName() const override;
+	UE_API virtual void					AddReferencedObjects(FReferenceCollector& Collector) override;
+	UE_API virtual FString					GetReferencerName() const override;
 
-	void							FreeCacheInternal(RenderTargetCache& TargetRTCache);
-	void							FreeRTList(RTList& RTList);
-	void							FreeRenderTarget(HashType HashValue, UTextureRenderTarget2D* RT);
+	UE_API void							FreeCacheInternal(RenderTargetCache& TargetRTCache);
+	UE_API void							FreeRTList(RTList& RTList);
+	UE_API void							FreeRenderTarget(HashType HashValue, UTextureRenderTarget2D* RT);
 
-	static UTextureRenderTarget2D*	CreateRenderTarget(const BufferDescriptor& Desc);
-	AsyncDeviceBufferRef			SplitToTiles_Internal(const CombineSplitArgs& SplitArgs);
+	static UE_API UTextureRenderTarget2D*	CreateRenderTarget(const BufferDescriptor& Desc);
+	UE_API AsyncDeviceBufferRef			SplitToTiles_Internal(const CombineSplitArgs& SplitArgs);
 
-	int32							AllocateRTArrayResource(TexArrayPtr TexArrayObj);
-	int32							AllocateRTResource(TexPtr TextureObj);
-	int32							InitRTResource(TexPtr TextureObj, UTextureRenderTarget* RT);
+	UE_API int32							AllocateRTArrayResource(TexArrayPtr TexArrayObj);
+	UE_API int32							AllocateRTResource(TexPtr TextureObj);
+	UE_API int32							InitRTResource(TexPtr TextureObj, UTextureRenderTarget* RT);
 
 protected:
-	virtual void					PrintStats() override;
+	UE_API virtual void					PrintStats() override;
 
 public:
-									Device_FX();
-	virtual							~Device_FX() override;
+									UE_API Device_FX();
+	UE_API virtual							~Device_FX() override;
 
-	virtual cti::continuable<int32>	Use() const override;
-	void							GetStatArray(TArray<FString>& ResourceArrayTiled, TArray<FString>& ResourceArrayUnTiled, TArray<FString>& TooltipListTiled, TArray<FString>& TooltipListUnTiled, FString& TotalStats);
+	UE_API virtual cti::continuable<int32>	Use() const override;
+	UE_API void							GetStatArray(TArray<FString>& ResourceArrayTiled, TArray<FString>& ResourceArrayUnTiled, TArray<FString>& TooltipListTiled, TArray<FString>& TooltipListUnTiled, FString& TotalStats);
 	/// Must be called from the rendering thread
-	DeviceBufferRef					CreateFromRT(UTextureRenderTarget2D* RT, const BufferDescriptor& Desc);
-	DeviceBufferRef					CreateFromTexture(UTexture2D* TextureObj, const BufferDescriptor& Desc);
-	DeviceBufferRef					CreateFromTex(TexPtr TextureObj, bool bInitRaw);
-	DeviceBufferRef					CreateFromTexAndRaw(TexPtr TextureObj, RawBufferPtr RawObj);
+	UE_API DeviceBufferRef					CreateFromRT(UTextureRenderTarget2D* RT, const BufferDescriptor& Desc);
+	UE_API DeviceBufferRef					CreateFromTexture(UTexture2D* TextureObj, const BufferDescriptor& Desc);
+	UE_API DeviceBufferRef					CreateFromTex(TexPtr TextureObj, bool bInitRaw);
+	UE_API DeviceBufferRef					CreateFromTexAndRaw(TexPtr TextureObj, RawBufferPtr RawObj);
 
-	virtual AsyncDeviceBufferRef	CombineFromTiles(const CombineSplitArgs& CombineArgs) override;
+	UE_API virtual AsyncDeviceBufferRef	CombineFromTiles(const CombineSplitArgs& CombineArgs) override;
 
-	AsyncDeviceBufferRef			DrawTilesToBuffer_Deferred(DeviceBufferRef buffer, const T_Tiles<DeviceBufferRef>& Tiles, const DrawTilesSettings* DrawSettings = nullptr, FLinearColor* ClearColor = nullptr);
-	AsyncDeviceBufferRef			FillTextureArray_Deferred(DeviceBufferRef buffer, const T_Tiles<DeviceBufferRef>& Tiles);
+	UE_API AsyncDeviceBufferRef			DrawTilesToBuffer_Deferred(DeviceBufferRef buffer, const T_Tiles<DeviceBufferRef>& Tiles, const DrawTilesSettings* DrawSettings = nullptr, FLinearColor* ClearColor = nullptr);
+	UE_API AsyncDeviceBufferRef			FillTextureArray_Deferred(DeviceBufferRef buffer, const T_Tiles<DeviceBufferRef>& Tiles);
 
-	virtual AsyncDeviceBufferRef	SplitToTiles(const CombineSplitArgs& splitArgs) override;
+	UE_API virtual AsyncDeviceBufferRef	SplitToTiles(const CombineSplitArgs& splitArgs) override;
 
-	TexPtr							AllocateRenderTarget(const BufferDescriptor& Desc);
-	TexPtr							AllocateRenderTargetArray(const BufferDescriptor& Desc, int32 NumTilesX, int32 NumTilesY);
+	UE_API TexPtr							AllocateRenderTarget(const BufferDescriptor& Desc);
+	UE_API TexPtr							AllocateRenderTargetArray(const BufferDescriptor& Desc, int32 NumTilesX, int32 NumTilesY);
 
 
 
-	virtual void					Update(float Delta) override;
-	void							MarkForCollection(TexPtr TextureObj);
+	UE_API virtual void					Update(float Delta) override;
+	UE_API void							MarkForCollection(TexPtr TextureObj);
 	virtual FString					Name() const override { return "Device_FX"; }
-	virtual void					AddNativeTask(DeviceNativeTaskPtr task) override;
+	UE_API virtual void					AddNativeTask(DeviceNativeTaskPtr task) override;
 
 	//////////////////////////////////////////////////////////////////////////
 	/// Inline functions
@@ -106,8 +108,10 @@ public:
 	//////////////////////////////////////////////////////////////////////////
 	/// Static functions
 	//////////////////////////////////////////////////////////////////////////
-	static void						InitTiles_Texture(T_Tiles<BlobPtr>* Tiles, const BufferDescriptor& InTileDesc, bool bInitRaw);
-	static void						InitTiles_RenderTargets(BlobUPtr* Tiles, size_t NumRows, size_t NumCols, const BufferDescriptor& InTileDesc, bool bInitRaw);
+	static UE_API void						InitTiles_Texture(T_Tiles<BlobPtr>* Tiles, const BufferDescriptor& InTileDesc, bool bInitRaw);
+	static UE_API void						InitTiles_RenderTargets(BlobUPtr* Tiles, size_t NumRows, size_t NumCols, const BufferDescriptor& InTileDesc, bool bInitRaw);
 
-	static Device_FX*				Get();
+	static UE_API Device_FX*				Get();
 };
+
+#undef UE_API
diff --git a/Engine/Plugins/TextureGraph/Source/TextureGraphEngine/Device/Mem/DeviceBuffer_Mem.h b/Engine/Plugins/TextureGraph/Source/TextureGraphEngine/Device/Mem/DeviceBuffer_Mem.h
--- a/Engine/Plugins/TextureGraph/Source/TextureGraphEngine/Device/Mem/DeviceBuffer_Mem.h
+++ b/Engine/Plugins/TextureGraph/Source/TextureGraphEngine/Device/Mem/DeviceBuffer_Mem.h
@@ -5,6 +5,8 @@
 #include "Device/DeviceBuffer.h"
 #include "Data/RawBuffer.h"
 
+#define UE_API TEXTUREGRAPHENGINE_API
+
 class Device;
 class Device_Mem;
 
@@ -12,32 +14,34 @@ class Device_Mem;
 /// This is simply a wrapper around RawBuffer. Must be owned by the 
 /// CPU device (though that can change in the future)
 //////////////////////////////////////////////////////////////////////////
-class TEXTUREGRAPHENGINE_API DeviceBuffer_Mem : public DeviceBuffer
+class DeviceBuffer_Mem : public DeviceBuffer
 {
 protected:
-	virtual DeviceBuffer*			CreateFromRaw(RawBufferPtr RawObj) override;
-	virtual DeviceBuffer*			CreateFromDesc(BufferDescriptor NewDesc, CHashPtr NewHashValue) override;
+	UE_API virtual DeviceBuffer*			CreateFromRaw(RawBufferPtr RawObj) override;
+	UE_API virtual DeviceBuffer*			CreateFromDesc(BufferDescriptor NewDesc, CHashPtr NewHashValue) override;
 
-	void							Allocate();
-	virtual AsyncBufferResultPtr	TransferFrom(DeviceBufferRef& Source) override;
+	UE_API void							Allocate();
+	UE_API virtual AsyncBufferResultPtr	TransferFrom(DeviceBufferRef& Source) override;
 
 public:
-									DeviceBuffer_Mem(Device_Mem* Device, BufferDescriptor Desc, CHashPtr HashValue);
-									DeviceBuffer_Mem(Device_Mem* Device, RawBufferPtr RawObj);
-	virtual							~DeviceBuffer_Mem() override;
+									UE_API DeviceBuffer_Mem(Device_Mem* Device, BufferDescriptor Desc, CHashPtr HashValue);
+									UE_API DeviceBuffer_Mem(Device_Mem* Device, RawBufferPtr RawObj);
+	UE_API virtual							~DeviceBuffer_Mem() override;
 
 	//////////////////////////////////////////////////////////////////////////
 	/// DeviceBuffer implementation
 	//////////////////////////////////////////////////////////////////////////
-	virtual RawBufferPtr			Raw_Now() override;
-	virtual size_t					MemSize() const override;
-	virtual AsyncBufferResultPtr	Bind(const BlobTransform* Transform, const ResourceBindInfo& BindInfo) override;
+	UE_API virtual RawBufferPtr			Raw_Now() override;
+	UE_API virtual size_t					MemSize() const override;
+	UE_API virtual AsyncBufferResultPtr	Bind(const BlobTransform* Transform, const ResourceBindInfo& BindInfo) override;
 
 	/// With base implementations
 	//virtual void					Release() override;
-	virtual bool					IsCompatible(Device* Dev) const override;
+	UE_API virtual bool					IsCompatible(Device* Dev) const override;
 
 	//////////////////////////////////////////////////////////////////////////
 	/// Inline functions
 	//////////////////////////////////////////////////////////////////////////
 };
+
+#undef UE_API
diff --git a/Engine/Plugins/TextureGraph/Source/TextureGraphEngine/Device/Mem/Device_Mem.h b/Engine/Plugins/TextureGraph/Source/TextureGraphEngine/Device/Mem/Device_Mem.h
--- a/Engine/Plugins/TextureGraph/Source/TextureGraphEngine/Device/Mem/Device_Mem.h
+++ b/Engine/Plugins/TextureGraph/Source/TextureGraphEngine/Device/Mem/Device_Mem.h
@@ -4,21 +4,25 @@
 #include "CoreMinimal.h"
 #include "Device/Device.h"
 
-class TEXTUREGRAPHENGINE_API Device_Mem : public Device
+#define UE_API TEXTUREGRAPHENGINE_API
+
+class Device_Mem : public Device
 {
 protected:
-									Device_Mem(DeviceType Type, DeviceBuffer* BufferFactory);
+									UE_API Device_Mem(DeviceType Type, DeviceBuffer* BufferFactory);
 
 public:
-									Device_Mem();
-	virtual							~Device_Mem() override;
+									UE_API Device_Mem();
+	UE_API virtual							~Device_Mem() override;
 
 	virtual FString					Name() const override { return "Device_Mem"; }
-	virtual void					Update(float Delta) override;
-	virtual void					AddNativeTask(DeviceNativeTaskPtr Task) override;
+	UE_API virtual void					Update(float Delta) override;
+	UE_API virtual void					AddNativeTask(DeviceNativeTaskPtr Task) override;
 
 	//////////////////////////////////////////////////////////////////////////
 	/// Static functions
 	//////////////////////////////////////////////////////////////////////////
-	static Device_Mem*				Get();
+	static UE_API Device_Mem*				Get();
 };
+
+#undef UE_API
diff --git a/Engine/Plugins/TextureGraph/Source/TextureGraphEngine/Device/MemCM/DeviceBuffer_MemCM.h b/Engine/Plugins/TextureGraph/Source/TextureGraphEngine/Device/MemCM/DeviceBuffer_MemCM.h
--- a/Engine/Plugins/TextureGraph/Source/TextureGraphEngine/Device/MemCM/DeviceBuffer_MemCM.h
+++ b/Engine/Plugins/TextureGraph/Source/TextureGraphEngine/Device/MemCM/DeviceBuffer_MemCM.h
@@ -5,6 +5,8 @@
 #include "Device/Mem/DeviceBuffer_Mem.h"
 #include "Data/RawBuffer.h"
 
+#define UE_API TEXTUREGRAPHENGINE_API
+
 class Device;
 class Device_MemCM;
 
@@ -12,30 +14,32 @@ class Device_MemCM;
 /// This is simply a wrapper around RawBuffer. Must be owned by the 
 /// CPU device (though that can change in the future)
 //////////////////////////////////////////////////////////////////////////
-class TEXTUREGRAPHENGINE_API DeviceBuffer_MemCM : public DeviceBuffer_Mem
+class DeviceBuffer_MemCM : public DeviceBuffer_Mem
 {
 protected:
-	virtual DeviceBuffer*			CreateFromRaw(RawBufferPtr InRawObj) override;
-	virtual DeviceBuffer*			CreateFromDesc(BufferDescriptor InDesc, CHashPtr InHashValue) override;
+	UE_API virtual DeviceBuffer*			CreateFromRaw(RawBufferPtr InRawObj) override;
+	UE_API virtual DeviceBuffer*			CreateFromDesc(BufferDescriptor InDesc, CHashPtr InHashValue) override;
 
-	virtual AsyncBufferResultPtr	TransferFrom(DeviceBufferRef& Source) override;
+	UE_API virtual AsyncBufferResultPtr	TransferFrom(DeviceBufferRef& Source) override;
 public:
-									DeviceBuffer_MemCM(Device_MemCM* Dev, BufferDescriptor InDesc, CHashPtr InHashValue);
-									DeviceBuffer_MemCM(Device_MemCM* Dev, RawBufferPtr InRawObj);
-	virtual							~DeviceBuffer_MemCM() override;
+									UE_API DeviceBuffer_MemCM(Device_MemCM* Dev, BufferDescriptor InDesc, CHashPtr InHashValue);
+									UE_API DeviceBuffer_MemCM(Device_MemCM* Dev, RawBufferPtr InRawObj);
+	UE_API virtual							~DeviceBuffer_MemCM() override;
 
 	//////////////////////////////////////////////////////////////////////////
 	/// DeviceBuffer implementation
 	//////////////////////////////////////////////////////////////////////////
-	virtual RawBufferPtr			Raw_Now() override;
-	virtual size_t					MemSize() const override;
-	virtual AsyncBufferResultPtr	Bind(const BlobTransform* Transform, const ResourceBindInfo& BindInfo) override;
+	UE_API virtual RawBufferPtr			Raw_Now() override;
+	UE_API virtual size_t					MemSize() const override;
+	UE_API virtual AsyncBufferResultPtr	Bind(const BlobTransform* Transform, const ResourceBindInfo& BindInfo) override;
 
 	/// With base implementations
 	//virtual void					Release() override;
-	virtual bool					IsCompatible(Device* Dev) const override;
+	UE_API virtual bool					IsCompatible(Device* Dev) const override;
 
 	//////////////////////////////////////////////////////////////////////////
 	/// Inline functions
 	//////////////////////////////////////////////////////////////////////////
 };
+
+#undef UE_API
diff --git a/Engine/Plugins/TextureGraph/Source/TextureGraphEngine/Device/MemCM/Device_MemCM.h b/Engine/Plugins/TextureGraph/Source/TextureGraphEngine/Device/MemCM/Device_MemCM.h
--- a/Engine/Plugins/TextureGraph/Source/TextureGraphEngine/Device/MemCM/Device_MemCM.h
+++ b/Engine/Plugins/TextureGraph/Source/TextureGraphEngine/Device/MemCM/Device_MemCM.h
@@ -4,17 +4,21 @@
 #include "CoreMinimal.h"
 #include "Device/Mem/Device_Mem.h"
 
-class TEXTUREGRAPHENGINE_API Device_MemCM : public Device_Mem
+#define UE_API TEXTUREGRAPHENGINE_API
+
+class Device_MemCM : public Device_Mem
 {
 public:
-									Device_MemCM();
-	virtual							~Device_MemCM() override;
+									UE_API Device_MemCM();
+	UE_API virtual							~Device_MemCM() override;
 
 	virtual FString					Name() const override { return "Device_MemCM"; }
-	virtual void					Update(float Delta) override;
+	UE_API virtual void					Update(float Delta) override;
 
 	//////////////////////////////////////////////////////////////////////////
 	/// Static functions
 	//////////////////////////////////////////////////////////////////////////
-	static Device_MemCM*			Get();
+	static UE_API Device_MemCM*			Get();
 };
+
+#undef UE_API
diff --git a/Engine/Plugins/TextureGraph/Source/TextureGraphEngine/Device/Null/DeviceBuffer_Null.h b/Engine/Plugins/TextureGraph/Source/TextureGraphEngine/Device/Null/DeviceBuffer_Null.h
--- a/Engine/Plugins/TextureGraph/Source/TextureGraphEngine/Device/Null/DeviceBuffer_Null.h
+++ b/Engine/Plugins/TextureGraph/Source/TextureGraphEngine/Device/Null/DeviceBuffer_Null.h
@@ -4,24 +4,28 @@
 #include "CoreMinimal.h"
 #include "Device/DeviceBuffer.h"
 
+#define UE_API TEXTUREGRAPHENGINE_API
+
 class Device_Null;
 
-class TEXTUREGRAPHENGINE_API DeviceBuffer_Null : public DeviceBuffer
+class DeviceBuffer_Null : public DeviceBuffer
 {
 protected:
-	virtual DeviceBuffer*			CreateFromRaw(RawBufferPtr rawBuffer) override;
-	virtual DeviceBuffer*			CreateFromDesc(BufferDescriptor desc, CHashPtr hash) override;
+	UE_API virtual DeviceBuffer*			CreateFromRaw(RawBufferPtr rawBuffer) override;
+	UE_API virtual DeviceBuffer*			CreateFromDesc(BufferDescriptor desc, CHashPtr hash) override;
 
-	virtual AsyncBufferResultPtr	TransferFrom(DeviceBufferRef& source);
+	UE_API virtual AsyncBufferResultPtr	TransferFrom(DeviceBufferRef& source);
 
 public:
-									DeviceBuffer_Null(Device_Null* device, BufferDescriptor desc, CHashPtr hash);
-	virtual							~DeviceBuffer_Null();
-
-	virtual RawBufferPtr			Raw_Now();
-	virtual size_t					MemSize() const;
-	virtual AsyncBufferResultPtr	Bind(const BlobTransform* transform, const ResourceBindInfo& bindInfo) override;
-	virtual bool					IsValid() const;
-	virtual bool					IsNull() const;
-	virtual bool					IsTransient() const override;
+									UE_API DeviceBuffer_Null(Device_Null* device, BufferDescriptor desc, CHashPtr hash);
+	UE_API virtual							~DeviceBuffer_Null();
+
+	UE_API virtual RawBufferPtr			Raw_Now();
+	UE_API virtual size_t					MemSize() const;
+	UE_API virtual AsyncBufferResultPtr	Bind(const BlobTransform* transform, const ResourceBindInfo& bindInfo) override;
+	UE_API virtual bool					IsValid() const;
+	UE_API virtual bool					IsNull() const;
+	UE_API virtual bool					IsTransient() const override;
 };
+
+#undef UE_API
diff --git a/Engine/Plugins/TextureGraph/Source/TextureGraphEngine/Device/Null/Device_Null.h b/Engine/Plugins/TextureGraph/Source/TextureGraphEngine/Device/Null/Device_Null.h
--- a/Engine/Plugins/TextureGraph/Source/TextureGraphEngine/Device/Null/Device_Null.h
+++ b/Engine/Plugins/TextureGraph/Source/TextureGraphEngine/Device/Null/Device_Null.h
@@ -4,20 +4,24 @@
 #include "CoreMinimal.h"
 #include "Device/Device.h"
 
-class TEXTUREGRAPHENGINE_API Device_Null : public Device
+#define UE_API TEXTUREGRAPHENGINE_API
+
+class Device_Null : public Device
 {
 protected:
 
-	virtual void					Collect(DeviceBuffer* Buffer) override;
+	UE_API virtual void					Collect(DeviceBuffer* Buffer) override;
 
 public:
-									Device_Null();
-	virtual							~Device_Null() override;
+									UE_API Device_Null();
+	UE_API virtual							~Device_Null() override;
 
 	virtual FString					Name() const override { return "Device_Null"; }
 
 	//////////////////////////////////////////////////////////////////////////
 	/// Static functions
 	//////////////////////////////////////////////////////////////////////////
-	static Device_Null*				Get();
+	static UE_API Device_Null*				Get();
 };
+
+#undef UE_API
diff --git a/Engine/Plugins/TextureGraph/Source/TextureGraphEngine/Export/T_ExportAsUAsset.h b/Engine/Plugins/TextureGraph/Source/TextureGraphEngine/Export/T_ExportAsUAsset.h
--- a/Engine/Plugins/TextureGraph/Source/TextureGraphEngine/Export/T_ExportAsUAsset.h
+++ b/Engine/Plugins/TextureGraph/Source/TextureGraphEngine/Export/T_ExportAsUAsset.h
@@ -5,15 +5,17 @@
 #include "Job/Job.h"
 #include "TextureExporter.h"
 
+#define UE_API TEXTUREGRAPHENGINE_API
+
 //////////////////////////////////////////////////////////////////////////
-class TEXTUREGRAPHENGINE_API Job_ExportAsUAsset : public Job
+class Job_ExportAsUAsset : public Job
 {
 private:
 	FExportMapSettings				Setting;
 	FString							OutFolder;
 
 public:
-									Job_ExportAsUAsset(UMixInterface* Mix, int32 TargetId, FExportMapSettings InSetting, FString OutputFolder, UObject* InErrorOwner = nullptr, uint16 Priority = (uint16)E_Priority::kHigh, uint64 id = 0);
+									UE_API Job_ExportAsUAsset(UMixInterface* Mix, int32 TargetId, FExportMapSettings InSetting, FString OutputFolder, UObject* InErrorOwner = nullptr, uint16 Priority = (uint16)E_Priority::kHigh, uint64 id = 0);
 
 protected:
 	virtual void					GetDependencies(JobPtrVec& prior, JobPtrVec& after, JobRunInfo runInfo) override { RunInfo = runInfo; Prev.clear(); }
@@ -25,18 +27,20 @@ protected:
 	virtual cti::continuable<int32>	UnbindArgs_All(JobRunInfo runInfo) override { return cti::make_ready_continuable(0); }
 	virtual bool					CanHandleTiles() const override { return false; }
 
-	virtual cti::continuable<int32>	PreExecAsync(ENamedThreads::Type execThread, ENamedThreads::Type returnThread) override;
-	virtual cti::continuable<int32>	ExecAsync(ENamedThreads::Type execThread, ENamedThreads::Type returnThread) override;
+	UE_API virtual cti::continuable<int32>	PreExecAsync(ENamedThreads::Type execThread, ENamedThreads::Type returnThread) override;
+	UE_API virtual cti::continuable<int32>	ExecAsync(ENamedThreads::Type execThread, ENamedThreads::Type returnThread) override;
 };
 
-class TEXTUREGRAPHENGINE_API T_ExportAsUAsset
+class T_ExportAsUAsset
 {
 public:
-									T_ExportAsUAsset();
-	virtual							~T_ExportAsUAsset();
+									UE_API T_ExportAsUAsset();
+	UE_API virtual							~T_ExportAsUAsset();
 
 	//////////////////////////////////////////////////////////////////////////
 	/// Static functions
 	//////////////////////////////////////////////////////////////////////////
-	static void						ExportAsUAsset(MixUpdateCyclePtr cycle, int32 targetId, FExportMapSettings exportMap, FString outFolder);
+	static UE_API void						ExportAsUAsset(MixUpdateCyclePtr cycle, int32 targetId, FExportMapSettings exportMap, FString outFolder);
 };
+
+#undef UE_API
diff --git a/Engine/Plugins/TextureGraph/Source/TextureGraphEngine/Export/TextureExporter.h b/Engine/Plugins/TextureGraph/Source/TextureGraphEngine/Export/TextureExporter.h
--- a/Engine/Plugins/TextureGraph/Source/TextureGraphEngine/Export/TextureExporter.h
+++ b/Engine/Plugins/TextureGraph/Source/TextureGraphEngine/Export/TextureExporter.h
@@ -18,7 +18,7 @@ THIRD_PARTY_INCLUDES_END
 
 DECLARE_LOG_CATEGORY_EXTERN(ExportLogs, Log, Verbose);
 
-struct TEXTUREGRAPHENGINE_API FExportMapSettings
+struct FExportMapSettings
 {
 	bool bIsEnabled = true;
 	FName Name;
@@ -40,7 +40,7 @@ struct TEXTUREGRAPHENGINE_API FExportMapSettings
 
 typedef std::vector<std::pair<FName, FExportMapSettings>> ExportMaps;
 
-class TEXTUREGRAPHENGINE_API FExportSettings
+class FExportSettings
 {
 public:
 	FString PresetName;
@@ -92,15 +92,15 @@ public:
 	}
 };
 
-struct TEXTUREGRAPHENGINE_API TextureExporter
+struct TextureExporter
 {
-	static FExportMapSettings		GetExportSettingsForTarget(FExportSettings& ExportSettings, TiledBlobPtr BlobObj, FName Name);
-	static AsyncBool				ExportAllMapsAsUAsset(UMixInterface* MixObj, const FString& OutputFolder, const FExportSettings& ExportSettings);
-	static AsyncBool				ExportAsUAsset(UMixInterface* MixObj, const FExportSettings& ExportSettings, FString C);
-	static void						ExportAsUAsset(UMixInterface* MixObj, FName AssetName, FString CurrentPath, const FExportSettings& Setting);
-	static AsyncInt					ExportRawAsUAsset(RawBufferPtr RawObj,const FExportMapSettings& Setting, const FString& CompletePath, const FName& TextureName);
-	static bool						IsPackageNameValid(FString Path, FString AssetName);
-	static bool 					IsFileNameValid(FName FileName, FText& Reason);
-	static bool 					IsFolderPathValid(FString FolderPath, FText& Reason);
-	static bool 					IsFilePathValid(const FName InFileName, const FName InFolderPath, FString& OutErrors);
+	static TEXTUREGRAPHENGINE_API FExportMapSettings		GetExportSettingsForTarget(FExportSettings& ExportSettings, TiledBlobPtr BlobObj, FName Name);
+	static TEXTUREGRAPHENGINE_API AsyncBool				ExportAllMapsAsUAsset(UMixInterface* MixObj, const FString& OutputFolder, const FExportSettings& ExportSettings);
+	static TEXTUREGRAPHENGINE_API AsyncBool				ExportAsUAsset(UMixInterface* MixObj, const FExportSettings& ExportSettings, FString C);
+	static TEXTUREGRAPHENGINE_API void						ExportAsUAsset(UMixInterface* MixObj, FName AssetName, FString CurrentPath, const FExportSettings& Setting);
+	static TEXTUREGRAPHENGINE_API AsyncInt					ExportRawAsUAsset(RawBufferPtr RawObj,const FExportMapSettings& Setting, const FString& CompletePath, const FName& TextureName);
+	static TEXTUREGRAPHENGINE_API bool						IsPackageNameValid(FString Path, FString AssetName);
+	static TEXTUREGRAPHENGINE_API bool 					IsFileNameValid(FName FileName, FText& Reason);
+	static TEXTUREGRAPHENGINE_API bool 					IsFolderPathValid(FString FolderPath, FText& Reason);
+	static TEXTUREGRAPHENGINE_API bool 					IsFilePathValid(const FName InFileName, const FName InFolderPath, FString& OutErrors);
 };
diff --git a/Engine/Plugins/TextureGraph/Source/TextureGraphEngine/FxMat/FxMaterial.h b/Engine/Plugins/TextureGraph/Source/TextureGraphEngine/FxMat/FxMaterial.h
--- a/Engine/Plugins/TextureGraph/Source/TextureGraphEngine/FxMat/FxMaterial.h
+++ b/Engine/Plugins/TextureGraph/Source/TextureGraphEngine/FxMat/FxMaterial.h
@@ -28,12 +28,14 @@
 
 #include <memory>
 
+#define UE_API TEXTUREGRAPHENGINE_API
+
 
 template <class SH_TypeParams> void SetupDefaultParameters(SH_TypeParams& params) {
 }
 
 //////////////////////////////////////////////////////////////////////////
-class TEXTUREGRAPHENGINE_API VSH_Base : public FGlobalShader
+class VSH_Base : public FGlobalShader
 {
 public:
 
@@ -47,7 +49,7 @@ public:
 };
 
 //////////////////////////////////////////////////////////////////////////
-class TEXTUREGRAPHENGINE_API FSH_Base : public FGlobalShader
+class FSH_Base : public FGlobalShader
 {
 public:
 	
@@ -62,7 +64,7 @@ public:
 
 //////////////////////////////////////////////////////////////////////////
 template <const int ThreadGroupSize_X = 16, const int ThreadGroupSize_Y = 16, const int ThreadGroupSize_Z = 1>
-class TEXTUREGRAPHENGINE_API CmpSH_Base : public FGlobalShader
+class CmpSH_Base : public FGlobalShader
 {
 public:
 
@@ -87,10 +89,10 @@ public:
 };
 
 ////////////////////////////////////////////////////////////////////////////
-class TEXTUREGRAPHENGINE_API VSH_Simple : public VSH_Base
+class VSH_Simple : public VSH_Base
 {
 public:
-	DECLARE_GLOBAL_SHADER(VSH_Simple);
+	DECLARE_EXPORTED_GLOBAL_SHADER(VSH_Simple, UE_API);
 	SHADER_USE_PARAMETER_STRUCT(VSH_Simple, VSH_Base);
 
 	BEGIN_SHADER_PARAMETER_STRUCT(FParameters, )
@@ -104,10 +106,10 @@ public:
 };
 
 ////////////////////////////////////////////////////////////////////////////
-class TEXTUREGRAPHENGINE_API FSH_Simple : public FSH_Base
+class FSH_Simple : public FSH_Base
 {
 public:
-	DECLARE_GLOBAL_SHADER(FSH_Simple);
+	DECLARE_EXPORTED_GLOBAL_SHADER(FSH_Simple, UE_API);
 	SHADER_USE_PARAMETER_STRUCT(FSH_Simple, FSH_Base);
 
 	BEGIN_SHADER_PARAMETER_STRUCT(FParameters, )
@@ -127,10 +129,10 @@ public:
 };
 
 //////////////////////////////////////////////////////////////////////////
-class TEXTUREGRAPHENGINE_API FSH_SimpleVT : public FSH_Base
+class FSH_SimpleVT : public FSH_Base
 {
 public:
-	DECLARE_GLOBAL_SHADER(FSH_SimpleVT);
+	DECLARE_EXPORTED_GLOBAL_SHADER(FSH_SimpleVT, UE_API);
 	SHADER_USE_PARAMETER_STRUCT(FSH_SimpleVT, FSH_Base);
 
 	BEGIN_SHADER_PARAMETER_STRUCT(FParameters, )
@@ -166,13 +168,13 @@ public:
 };
 
 //////////////////////////////////////////////////////////////////////////
-class TEXTUREGRAPHENGINE_API QuadScreenBuffer : public FVertexBuffer
+class QuadScreenBuffer : public FVertexBuffer
 {
-	void							InitRHI(FRHICommandListBase& RHICmdList) override;
+	UE_API void							InitRHI(FRHICommandListBase& RHICmdList) override;
 };
 
 //////////////////////////////////////////////////////////////////////////
-class TEXTUREGRAPHENGINE_API FxMaterial
+class FxMaterial
 {
 public:
 	struct MemberInfo
@@ -209,26 +211,26 @@ protected:
 
 	std::vector<BoundTextures>		Textures;			/// Textures that are bound with this material. We bind them in the end
 
-	std::unique_ptr<MemberLUT>&		GetParamsLUT();
-	MemberInfo						GetMember(FName MemberName);
+	UE_API std::unique_ptr<MemberLUT>&		GetParamsLUT();
+	UE_API MemberInfo						GetMember(FName MemberName);
 
-	void							BindTexturesForBlitting();
+	UE_API void							BindTexturesForBlitting();
 
 	/// Default constructor is only accessible to derived classes
 	FxMaterial() = default;
 
 public:
 
-	FxMaterial(TShaderRef<VSH_Base> VSH, TShaderRef<FSH_Base> FSH);
-	FxMaterial(TShaderRef<VSH_Base> VSH);
-	FxMaterial(FString VSHHashStr, FString FSHHashStr);
+	UE_API FxMaterial(TShaderRef<VSH_Base> VSH, TShaderRef<FSH_Base> FSH);
+	UE_API FxMaterial(TShaderRef<VSH_Base> VSH);
+	UE_API FxMaterial(FString VSHHashStr, FString FSHHashStr);
 
 	//////////////////////////////////////////////////////////////////////////
 	/// Global static data and functions
 	//////////////////////////////////////////////////////////////////////////
-	static TGlobalResource<QuadScreenBuffer> GQuadBuffer;
-	static void						InitPSO_Default(FGraphicsPipelineStateInitializer& PSO);
-	static void						InitPSO_Default(FGraphicsPipelineStateInitializer& PSO, FRHIVertexShader* VSH, FRHIPixelShader* FSH);
+	static UE_API TGlobalResource<QuadScreenBuffer> GQuadBuffer;
+	static UE_API void						InitPSO_Default(FGraphicsPipelineStateInitializer& PSO);
+	static UE_API void						InitPSO_Default(FGraphicsPipelineStateInitializer& PSO, FRHIVertexShader* VSH, FRHIPixelShader* FSH);
 
 	//////////////////////////////////////////////////////////////////////////
 	/// Must be overridden by the implementation
@@ -243,16 +245,16 @@ public:
 	//////////////////////////////////////////////////////////////////////////
 	/// Mimicing UMaterialInstanceDynamic
 	//////////////////////////////////////////////////////////////////////////
-	virtual void					SetArrayTextureParameterValue(FName Name, const ArrayTexture& Value);
-	virtual void					SetTextureParameterValue(FName Name, const UTexture* Value);
-	virtual void					SetScalarParameterValue(FName Name, float Value);
-	virtual void					SetScalarParameterValue(FName Name, int32 Value);
-	virtual void					SetVectorParameterValue(FName Name, const FLinearColor& Value);
-	virtual void					SetVectorParameterValue(FName Name, const FIntVector4& Value);
-	virtual void					SetStructParameterValue(FName Name, const char* Value, size_t StructSize); 
-	virtual void					SetArrayParameterValue(FName Name, const char* startAddress, size_t TypeSize, size_t ArraySize);
-	virtual void					SetMatrixParameterValue(FName Name, const FMatrix& Value);
-	virtual bool					DoesMemberExist(FName MemberName);
+	UE_API virtual void					SetArrayTextureParameterValue(FName Name, const ArrayTexture& Value);
+	UE_API virtual void					SetTextureParameterValue(FName Name, const UTexture* Value);
+	UE_API virtual void					SetScalarParameterValue(FName Name, float Value);
+	UE_API virtual void					SetScalarParameterValue(FName Name, int32 Value);
+	UE_API virtual void					SetVectorParameterValue(FName Name, const FLinearColor& Value);
+	UE_API virtual void					SetVectorParameterValue(FName Name, const FIntVector4& Value);
+	UE_API virtual void					SetStructParameterValue(FName Name, const char* Value, size_t StructSize); 
+	UE_API virtual void					SetArrayParameterValue(FName Name, const char* startAddress, size_t TypeSize, size_t ArraySize);
+	UE_API virtual void					SetMatrixParameterValue(FName Name, const FMatrix& Value);
+	UE_API virtual bool					DoesMemberExist(FName MemberName);
 	
 	//////////////////////////////////////////////////////////////////////////
 	/// Inline functions
@@ -507,3 +509,5 @@ static bool	ShouldCompilePermutation(const FGlobalShaderPermutationParameters& P
 
 #define TEXTUREGRAPH_ENGINE_DEFAULT_COMPILATION_ENV \
 static void						ModifyCompilationEnvironment(const FGlobalShaderPermutationParameters& params, FShaderCompilerEnvironment& env) {}
+
+#undef UE_API
diff --git a/Engine/Plugins/TextureGraph/Source/TextureGraphEngine/FxMat/FxMaterial_DrawMaterial.h b/Engine/Plugins/TextureGraph/Source/TextureGraphEngine/FxMat/FxMaterial_DrawMaterial.h
--- a/Engine/Plugins/TextureGraph/Source/TextureGraphEngine/FxMat/FxMaterial_DrawMaterial.h
+++ b/Engine/Plugins/TextureGraph/Source/TextureGraphEngine/FxMat/FxMaterial_DrawMaterial.h
@@ -21,7 +21,9 @@
 
 #include "Device/Device.h"
 
-class TEXTUREGRAPHENGINE_API FTextureGraphMaterialShaderVS : public FMaterialShader
+#define UE_API TEXTUREGRAPHENGINE_API
+
+class FTextureGraphMaterialShaderVS : public FMaterialShader
 {
 public:
 	DECLARE_SHADER_TYPE(FTextureGraphMaterialShaderVS, Material)
@@ -54,7 +56,7 @@ public:
 	}
 };
 
-class TEXTUREGRAPHENGINE_API FTextureGraphMaterialShaderPS : public FMaterialShader
+class FTextureGraphMaterialShaderPS : public FMaterialShader
 {
 public:
 	DECLARE_SHADER_TYPE(FTextureGraphMaterialShaderPS, Material)
@@ -66,7 +68,7 @@ public:
 
 	SHADER_USE_PARAMETER_STRUCT_WITH_LEGACY_BASE(FTextureGraphMaterialShaderPS, FMaterialShader);
 
-	static FName PSCONTROL_ARG;
+	static UE_API FName PSCONTROL_ARG;
 
 	static bool ShouldCompilePermutation(const FMaterialShaderPermutationParameters& Parameters)
 	{
@@ -227,3 +229,4 @@ typedef FxMaterial_DrawMaterial< FTextureGraphMaterialShaderVS, FTextureGraphMat
 typedef std::shared_ptr<FxMaterial_QuadDrawMaterial>	FxMaterialQuadDrawMaterialPtr;
 
 
+#undef UE_API
diff --git a/Engine/Plugins/TextureGraph/Source/TextureGraphEngine/FxMat/MaterialManager.h b/Engine/Plugins/TextureGraph/Source/TextureGraphEngine/FxMat/MaterialManager.h
--- a/Engine/Plugins/TextureGraph/Source/TextureGraphEngine/FxMat/MaterialManager.h
+++ b/Engine/Plugins/TextureGraph/Source/TextureGraphEngine/FxMat/MaterialManager.h
@@ -11,8 +11,10 @@
 #include "Model/ModelObject.h"
 #include "MaterialManager.generated.h"
 
-UCLASS()
-class TEXTUREGRAPHENGINE_API UMaterialManager : public UModelObject
+#define UE_API TEXTUREGRAPHENGINE_API
+
+UCLASS(MinimalAPI)
+class UMaterialManager : public UModelObject
 {
 	GENERATED_BODY()
 
@@ -20,20 +22,20 @@ private:
 	UPROPERTY(transient)
 	TMap<FString, TObjectPtr<UMaterial>> Materials;				/// The materials that have been loaded so far
 
-	void							InitEssentialMaterials();
+	UE_API void							InitEssentialMaterials();
 	
 
 public:
-									UMaterialManager();
-	virtual							~UMaterialManager() override;
+									UE_API UMaterialManager();
+	UE_API virtual							~UMaterialManager() override;
 	
-	class UMaterial*				LoadMaterial(const FString& Path);	
+	UE_API class UMaterial*				LoadMaterial(const FString& Path);	
 
-	RenderMaterial_BPPtr			CreateMaterial_BP(FString Name, const FString& Path);
-	RenderMaterial_BPPtr			CreateMaterial_BP(FString InName, UMaterial* InMaterial);
-	RenderMaterial_BP_NoTilePtr		CreateMaterial_BP_NoTile(FString Name, const FString& Path);
-	RenderMaterial_ThumbPtr			CreateMaterial_Thumbnail(FString Name, const FString& Path);
-	RenderMaterial_BP_TileArgsPtr	CreateMaterialOfType_BP_TileArgs(FString Name, const FString& Path);
+	UE_API RenderMaterial_BPPtr			CreateMaterial_BP(FString Name, const FString& Path);
+	UE_API RenderMaterial_BPPtr			CreateMaterial_BP(FString InName, UMaterial* InMaterial);
+	UE_API RenderMaterial_BP_NoTilePtr		CreateMaterial_BP_NoTile(FString Name, const FString& Path);
+	UE_API RenderMaterial_ThumbPtr			CreateMaterial_Thumbnail(FString Name, const FString& Path);
+	UE_API RenderMaterial_BP_TileArgsPtr	CreateMaterialOfType_BP_TileArgs(FString Name, const FString& Path);
 
 	//////////////////////////////////////////////////////////////////////////
 
@@ -101,3 +103,5 @@ public:
 		return std::make_shared<RenderMaterial_FX_Combined>(Name, std::static_pointer_cast<FxMaterial>(Mat));
 	}
 };
+
+#undef UE_API
diff --git a/Engine/Plugins/TextureGraph/Source/TextureGraphEngine/FxMat/RenderMaterial.h b/Engine/Plugins/TextureGraph/Source/TextureGraphEngine/FxMat/RenderMaterial.h
--- a/Engine/Plugins/TextureGraph/Source/TextureGraphEngine/FxMat/RenderMaterial.h
+++ b/Engine/Plugins/TextureGraph/Source/TextureGraphEngine/FxMat/RenderMaterial.h
@@ -5,6 +5,8 @@
 #include "Transform/BlobTransform.h"
 #include "Engine/TextureRenderTarget2D.h"
 
+#define UE_API TEXTUREGRAPHENGINE_API
+
 
 class Tex;
 class Device;
@@ -13,25 +15,25 @@ class FRHICommandListImmediate;
 class UTextureRenderTarget2D;
 class UTexture;
 
-class TEXTUREGRAPHENGINE_API RenderMaterial : public BlobTransform
+class RenderMaterial : public BlobTransform
 {
 protected:
 	bool							bCanHandleTiles = true;		/// Whether this material can handle tiles or not
 
 public:
-	explicit						RenderMaterial(FString Name);
-	virtual							~RenderMaterial() override;
+	UE_API explicit						RenderMaterial(FString Name);
+	UE_API virtual							~RenderMaterial() override;
 
 	virtual void					BlitTo(FRHICommandListImmediate& RHI, UTextureRenderTarget2D* Dst, const RenderMesh* MeshObj, int32 TargetId) const = 0;
-	void							BlitTo(FRHICommandListImmediate& RHI, UTextureRenderTarget2D* Dst) const;
+	UE_API void							BlitTo(FRHICommandListImmediate& RHI, UTextureRenderTarget2D* Dst) const;
 
 	//////////////////////////////////////////////////////////////////////////
 	/// BlobTransform overrides
 	//////////////////////////////////////////////////////////////////////////
-	virtual Device*					TargetDevice(size_t Index) const override;
-	virtual AsyncTransformResultPtr	Exec(const TransformArgs& Args) override;
-	virtual ENamedThreads::Type		ExecutionThread() const override;
-	virtual bool					CanHandleTiles() const override;
+	UE_API virtual Device*					TargetDevice(size_t Index) const override;
+	UE_API virtual AsyncTransformResultPtr	Exec(const TransformArgs& Args) override;
+	UE_API virtual ENamedThreads::Type		ExecutionThread() const override;
+	UE_API virtual bool					CanHandleTiles() const override;
 
 	virtual void					SetTexture(FName Name, const UTexture* Texture) const = 0;
 	virtual void					SetArrayTexture(FName Name, const std::vector<const UTexture*>& Textures) const = 0;
@@ -40,10 +42,10 @@ public:
 	virtual void					SetColor(FName Name, const FLinearColor& Value) const = 0;
 	virtual void					SetIntVector4(FName Name, const FIntVector4& Value) const = 0;
 	virtual void					SetMatrix(FName Name, const FMatrix& Value) const = 0;
-	void							SetSourceTexture(const Tex* Texture) const;
-	void							SetTexture(FName InName, std::shared_ptr<Tex> TexObj) const;
-	void							SetArrayTexture(FName InName, const std::vector<std::shared_ptr<Tex>>& TexObj) const;
-	void							BlitTo(FRHICommandListImmediate& RHI, Tex* Dst, const RenderMesh* MeshObj = nullptr, int32 TargetId = -1) const;
+	UE_API void							SetSourceTexture(const Tex* Texture) const;
+	UE_API void							SetTexture(FName InName, std::shared_ptr<Tex> TexObj) const;
+	UE_API void							SetArrayTexture(FName InName, const std::vector<std::shared_ptr<Tex>>& TexObj) const;
+	UE_API void							BlitTo(FRHICommandListImmediate& RHI, Tex* Dst, const RenderMesh* MeshObj = nullptr, int32 TargetId = -1) const;
 	//////////////////////////////////////////////////////////////////////////
 	/// Inline functions
 	//////////////////////////////////////////////////////////////////////////
@@ -52,3 +54,5 @@ public:
 };
 
 typedef std::shared_ptr<RenderMaterial>	RenderMaterialPtr;
+
+#undef UE_API
diff --git a/Engine/Plugins/TextureGraph/Source/TextureGraphEngine/FxMat/RenderMaterial_BP.h b/Engine/Plugins/TextureGraph/Source/TextureGraphEngine/FxMat/RenderMaterial_BP.h
--- a/Engine/Plugins/TextureGraph/Source/TextureGraphEngine/FxMat/RenderMaterial_BP.h
+++ b/Engine/Plugins/TextureGraph/Source/TextureGraphEngine/FxMat/RenderMaterial_BP.h
@@ -7,25 +7,27 @@
 #include "Materials/MaterialInstanceConstant.h"
 #include "UObject/StrongObjectPtr.h"
 
+#define UE_API TEXTUREGRAPHENGINE_API
+
 class FxMaterial;
 typedef std::shared_ptr< FxMaterial> FxMaterialPtr;
 
 // Helper container of MIC for the different static switch permutations of a Material
 // Used by the caller of a RenderMaterial_BP to pick a specific permutation different from the default set.
 // See TG_Expression_MaterialBase for a usage example
-struct TEXTUREGRAPHENGINE_API FMaterialInstanceStaticSwitchPermutationMap
+struct FMaterialInstanceStaticSwitchPermutationMap
 {
 	int32												DefaultKey = 0;
 	FStaticParameterSet									DefaultStaticParameterSet;
 	TMap<int32, TStrongObjectPtr<UMaterialInstanceConstant>>	PermutationsMap;
 
-	static TSharedPtr<FMaterialInstanceStaticSwitchPermutationMap> Create(UMaterialInterface* InMaterial);
-	UMaterialInterface*				GetRootMaterial();
-	int32							KeyFromStaticSwitchParameters(const TArray<FStaticSwitchParameter>& Parameters);
-	UMaterialInstanceConstant*		GetMaterialInstance(const TArray<FStaticSwitchParameter>& Parameters);
+	static UE_API TSharedPtr<FMaterialInstanceStaticSwitchPermutationMap> Create(UMaterialInterface* InMaterial);
+	UE_API UMaterialInterface*				GetRootMaterial();
+	UE_API int32							KeyFromStaticSwitchParameters(const TArray<FStaticSwitchParameter>& Parameters);
+	UE_API UMaterialInstanceConstant*		GetMaterialInstance(const TArray<FStaticSwitchParameter>& Parameters);
 };
 
-class TEXTUREGRAPHENGINE_API RenderMaterial_BP : public RenderMaterial
+class RenderMaterial_BP : public RenderMaterial
 {
 protected:
 
@@ -40,42 +42,42 @@ protected:
 
 	UCanvas*						Canvas = nullptr;
 
-	void							DrawMaterial(UMaterialInterface* RenderMaterial, FVector2D ScreenPosition, FVector2D ScreenSize,
+	UE_API void							DrawMaterial(UMaterialInterface* RenderMaterial, FVector2D ScreenPosition, FVector2D ScreenSize,
 									FVector2D CoordinatePosition, FVector2D CoordinateSize=FVector2D::UnitVector, float Rotation=0.f,
 									FVector2D PivotPoint=FVector2D(0.5f,0.5f)) const;
 public:
-									RenderMaterial_BP(FString Name, UMaterialInterface* InMaterial);
+									UE_API RenderMaterial_BP(FString Name, UMaterialInterface* InMaterial);
 
-	virtual							~RenderMaterial_BP() override;
+	UE_API virtual							~RenderMaterial_BP() override;
 
-	static bool						ValidateMaterialCompatible(UMaterialInterface* InMaterial);
+	static UE_API bool						ValidateMaterialCompatible(UMaterialInterface* InMaterial);
 
 
-	virtual AsyncPrepareResult		PrepareResources(const TransformArgs& Args) override;
+	UE_API virtual AsyncPrepareResult		PrepareResources(const TransformArgs& Args) override;
 
 	//////////////////////////////////////////////////////////////////////////
 	/// BlobTransform Implementation
 	//////////////////////////////////////////////////////////////////////////
-	virtual void					Bind(int32 Value, const ResourceBindInfo& BindInfo) override;
-	virtual void					Bind(float Value, const ResourceBindInfo& BindInfo) override;
-	virtual void					Bind(const FLinearColor& Value, const ResourceBindInfo& BindInfo) override;
-	virtual void					Bind(const FIntVector4& Value, const ResourceBindInfo& BindInfo) override;
-	virtual void					Bind(const FMatrix& Value, const ResourceBindInfo& BindInfo) override;
-	virtual void					BindStruct(const char* ValueAddress, size_t StructSize, const ResourceBindInfo& BindInfo) override;
-	virtual CHashPtr				Hash() const override;
-	virtual std::shared_ptr<BlobTransform> DuplicateInstance(FString InName) override;
+	UE_API virtual void					Bind(int32 Value, const ResourceBindInfo& BindInfo) override;
+	UE_API virtual void					Bind(float Value, const ResourceBindInfo& BindInfo) override;
+	UE_API virtual void					Bind(const FLinearColor& Value, const ResourceBindInfo& BindInfo) override;
+	UE_API virtual void					Bind(const FIntVector4& Value, const ResourceBindInfo& BindInfo) override;
+	UE_API virtual void					Bind(const FMatrix& Value, const ResourceBindInfo& BindInfo) override;
+	UE_API virtual void					BindStruct(const char* ValueAddress, size_t StructSize, const ResourceBindInfo& BindInfo) override;
+	UE_API virtual CHashPtr				Hash() const override;
+	UE_API virtual std::shared_ptr<BlobTransform> DuplicateInstance(FString InName) override;
 
 	//////////////////////////////////////////////////////////////////////////
 	/// RenderMaterial Implementation
 	//////////////////////////////////////////////////////////////////////////
-	virtual void					BlitTo(FRHICommandListImmediate& RHI, UTextureRenderTarget2D* DstRT, const RenderMesh* MeshObj, int32 TargetId) const override;
-	virtual void					SetTexture(FName InName, const UTexture* Texture) const override;
-	virtual void					SetArrayTexture(FName InName, const std::vector<const UTexture*>& Textures) const override;
-	virtual void					SetInt(FName InName, int32 Value) const override;
-	virtual void					SetFloat(FName InName, float Value) const override;
-	virtual void					SetColor(FName InName, const FLinearColor& Value) const override;
-	virtual void					SetIntVector4(FName InName, const FIntVector4& Value) const override;
-	virtual void					SetMatrix(FName InName, const FMatrix& Value) const override;
+	UE_API virtual void					BlitTo(FRHICommandListImmediate& RHI, UTextureRenderTarget2D* DstRT, const RenderMesh* MeshObj, int32 TargetId) const override;
+	UE_API virtual void					SetTexture(FName InName, const UTexture* Texture) const override;
+	UE_API virtual void					SetArrayTexture(FName InName, const std::vector<const UTexture*>& Textures) const override;
+	UE_API virtual void					SetInt(FName InName, int32 Value) const override;
+	UE_API virtual void					SetFloat(FName InName, float Value) const override;
+	UE_API virtual void					SetColor(FName InName, const FLinearColor& Value) const override;
+	UE_API virtual void					SetIntVector4(FName InName, const FIntVector4& Value) const override;
+	UE_API virtual void					SetMatrix(FName InName, const FMatrix& Value) const override;
 	
 	//////////////////////////////////////////////////////////////////////////
 	/// Inline functions
@@ -85,4 +87,6 @@ public:
 	FORCEINLINE const UMaterialInstanceConstant*	Instance() const { return MaterialInstance.Get(); }
 };
 
-typedef std::shared_ptr<RenderMaterial_BP> RenderMaterial_BPPtr;
\ No newline at end of file
+typedef std::shared_ptr<RenderMaterial_BP> RenderMaterial_BPPtr;
+
+#undef UE_API
diff --git a/Engine/Plugins/TextureGraph/Source/TextureGraphEngine/FxMat/RenderMaterial_BP_NoTile.h b/Engine/Plugins/TextureGraph/Source/TextureGraphEngine/FxMat/RenderMaterial_BP_NoTile.h
--- a/Engine/Plugins/TextureGraph/Source/TextureGraphEngine/FxMat/RenderMaterial_BP_NoTile.h
+++ b/Engine/Plugins/TextureGraph/Source/TextureGraphEngine/FxMat/RenderMaterial_BP_NoTile.h
@@ -4,10 +4,12 @@
 #include "CoreMinimal.h"
 #include "RenderMaterial_BP.h"
 
+#define UE_API TEXTUREGRAPHENGINE_API
+
 /**
  * 
  */
-class TEXTUREGRAPHENGINE_API RenderMaterial_BP_NoTile : public RenderMaterial_BP
+class RenderMaterial_BP_NoTile : public RenderMaterial_BP
 {
 
 public:
@@ -17,17 +19,19 @@ public:
 
 typedef std::shared_ptr<RenderMaterial_BP_NoTile> RenderMaterial_BP_NoTilePtr;
 
-class TEXTUREGRAPHENGINE_API RenderMaterial_BP_TileArgs : public RenderMaterial_BP
+class RenderMaterial_BP_TileArgs : public RenderMaterial_BP
 {
 private:
 	FName							MemberName = "TileInfo";
 
 public:
 	RenderMaterial_BP_TileArgs(FString InName, UMaterialInterface* InMaterial) : RenderMaterial_BP(InName, InMaterial) {}
-	void							AddTileArgs(TransformArgs& Args);
-	virtual AsyncPrepareResult		PrepareResources(const TransformArgs& Args) override;
+	UE_API void							AddTileArgs(TransformArgs& Args);
+	UE_API virtual AsyncPrepareResult		PrepareResources(const TransformArgs& Args) override;
 	virtual bool					CanHandleTiles() const override { return true; };
-	virtual std::shared_ptr<BlobTransform> DuplicateInstance(FString NewName) override;
+	UE_API virtual std::shared_ptr<BlobTransform> DuplicateInstance(FString NewName) override;
 };
 
-typedef std::shared_ptr<RenderMaterial_BP_TileArgs> RenderMaterial_BP_TileArgsPtr;
\ No newline at end of file
+typedef std::shared_ptr<RenderMaterial_BP_TileArgs> RenderMaterial_BP_TileArgsPtr;
+
+#undef UE_API
