Ran UnrealCodeFixup to fix dll storage

diff --git a/Engine/Plugins/TextureGraph/Source/TextureGraphEngine/FxMat/RenderMaterial_FX.h b/Engine/Plugins/TextureGraph/Source/TextureGraphEngine/FxMat/RenderMaterial_FX.h
--- a/Engine/Plugins/TextureGraph/Source/TextureGraphEngine/FxMat/RenderMaterial_FX.h
+++ b/Engine/Plugins/TextureGraph/Source/TextureGraphEngine/FxMat/RenderMaterial_FX.h
@@ -5,44 +5,46 @@
 #include "RenderMaterial.h"
 #include "SamplerStates_FX.h"
 
+#define UE_API TEXTUREGRAPHENGINE_API
+
 class FxMaterial;
 typedef std::shared_ptr<FxMaterial>	FxMaterialPtr;
 
-class TEXTUREGRAPHENGINE_API RenderMaterial_FX : public RenderMaterial
+class RenderMaterial_FX : public RenderMaterial
 {
 protected:
 	mutable FxMaterialPtr			FXMaterial;			/// The FX material from the shader
 	CHashPtr						HashValue;			/// The hash for the material
 	
 public:
-									RenderMaterial_FX(FString InName, FxMaterialPtr InFXMaterial);
-	virtual							~RenderMaterial_FX() override;
+									UE_API RenderMaterial_FX(FString InName, FxMaterialPtr InFXMaterial);
+	UE_API virtual							~RenderMaterial_FX() override;
 
 	//////////////////////////////////////////////////////////////////////////
 	/// BlobTransform implementation
 	//////////////////////////////////////////////////////////////////////////
-	virtual void					Bind(int32 Value, const ResourceBindInfo& BindInfo) override;
-	virtual void					Bind(float Value, const ResourceBindInfo& BindInfo) override;
-	virtual void					Bind(const FLinearColor& Value, const ResourceBindInfo& BindInfo) override;
-	virtual void					Bind(const FIntVector4& Value, const ResourceBindInfo& BindInfo) override;
-	virtual void					Bind(const FMatrix& Value, const ResourceBindInfo& BindInfo) override;
-	virtual void					BindStruct(const char* ValueAddress, size_t StructSize, const ResourceBindInfo& BindInfo) override;
-	virtual void					BindScalarArray(const char* StartingAddress, size_t TypeSize, size_t ArraySize, const ResourceBindInfo& BindInfo) override;
+	UE_API virtual void					Bind(int32 Value, const ResourceBindInfo& BindInfo) override;
+	UE_API virtual void					Bind(float Value, const ResourceBindInfo& BindInfo) override;
+	UE_API virtual void					Bind(const FLinearColor& Value, const ResourceBindInfo& BindInfo) override;
+	UE_API virtual void					Bind(const FIntVector4& Value, const ResourceBindInfo& BindInfo) override;
+	UE_API virtual void					Bind(const FMatrix& Value, const ResourceBindInfo& BindInfo) override;
+	UE_API virtual void					BindStruct(const char* ValueAddress, size_t StructSize, const ResourceBindInfo& BindInfo) override;
+	UE_API virtual void					BindScalarArray(const char* StartingAddress, size_t TypeSize, size_t ArraySize, const ResourceBindInfo& BindInfo) override;
 
-	virtual CHashPtr				Hash() const override;
-	virtual std::shared_ptr<BlobTransform> DuplicateInstance(FString NewName) override;
+	UE_API virtual CHashPtr				Hash() const override;
+	UE_API virtual std::shared_ptr<BlobTransform> DuplicateInstance(FString NewName) override;
 
 	//////////////////////////////////////////////////////////////////////////
 	/// RenderMaterial Implementation
 	//////////////////////////////////////////////////////////////////////////
-	virtual void					BlitTo(FRHICommandListImmediate& RHI, UTextureRenderTarget2D* RenderTarget, const RenderMesh* MeshObj, int32 TargetId) const override;
-	virtual void					SetTexture(FName InName, const UTexture* Texture) const override;
-	virtual void					SetArrayTexture(FName InName, const std::vector<const UTexture*>& Textures) const override;
-	virtual void					SetInt(FName InName, int32 Value) const override;
-	virtual void					SetFloat(FName InName, float Value) const override;
-	virtual void					SetColor(FName InName, const FLinearColor& Value) const override;
-	virtual void					SetIntVector4(FName InName, const FIntVector4& Value) const override;
-	virtual void					SetMatrix(FName InName, const FMatrix& Value) const override;
+	UE_API virtual void					BlitTo(FRHICommandListImmediate& RHI, UTextureRenderTarget2D* RenderTarget, const RenderMesh* MeshObj, int32 TargetId) const override;
+	UE_API virtual void					SetTexture(FName InName, const UTexture* Texture) const override;
+	UE_API virtual void					SetArrayTexture(FName InName, const std::vector<const UTexture*>& Textures) const override;
+	UE_API virtual void					SetInt(FName InName, int32 Value) const override;
+	UE_API virtual void					SetFloat(FName InName, float Value) const override;
+	UE_API virtual void					SetColor(FName InName, const FLinearColor& Value) const override;
+	UE_API virtual void					SetIntVector4(FName InName, const FIntVector4& Value) const override;
+	UE_API virtual void					SetMatrix(FName InName, const FMatrix& Value) const override;
 
 	//////////////////////////////////////////////////////////////////////////
 	/// Inline functions
@@ -53,3 +55,5 @@ public:
 
 
 typedef std::shared_ptr<RenderMaterial_FX> RenderMaterial_FXPtr;
+
+#undef UE_API
diff --git a/Engine/Plugins/TextureGraph/Source/TextureGraphEngine/FxMat/RenderMaterial_FX_Combined.h b/Engine/Plugins/TextureGraph/Source/TextureGraphEngine/FxMat/RenderMaterial_FX_Combined.h
--- a/Engine/Plugins/TextureGraph/Source/TextureGraphEngine/FxMat/RenderMaterial_FX_Combined.h
+++ b/Engine/Plugins/TextureGraph/Source/TextureGraphEngine/FxMat/RenderMaterial_FX_Combined.h
@@ -5,30 +5,34 @@
 #include "RenderMaterial_FX.h"
 #include "Data/TiledBlob.h"
 
+#define UE_API TEXTUREGRAPHENGINE_API
+
 DECLARE_LOG_CATEGORY_EXTERN(LogRenderMaterial_FX_NoTile, All, All);
 
 
 
 
-class TEXTUREGRAPHENGINE_API RenderMaterial_FX_Combined : public RenderMaterial_FX
+class RenderMaterial_FX_Combined : public RenderMaterial_FX
 {
 private:
 	
 	TArray<TiledBlobPtr>					ToBeCombined;		/// To save under process blobs. This array ensures that multiple SRVs can be passed to the shader
 
 public:
-	RenderMaterial_FX_Combined(FString InName, FxMaterialPtr InFXMaterial);
-	virtual									~RenderMaterial_FX_Combined() override;
+	UE_API RenderMaterial_FX_Combined(FString InName, FxMaterialPtr InFXMaterial);
+	UE_API virtual									~RenderMaterial_FX_Combined() override;
 
-	void									AddTileArgs(TransformArgs& Args);
-	virtual AsyncPrepareResult				PrepareResources(const TransformArgs& Args) override;
+	UE_API void									AddTileArgs(TransformArgs& Args);
+	UE_API virtual AsyncPrepareResult				PrepareResources(const TransformArgs& Args) override;
 	virtual bool							CanHandleTiles() const override { return true; };
-	virtual std::shared_ptr<BlobTransform>	DuplicateInstance(FString NewName) override;
-	virtual AsyncTransformResultPtr			Exec(const TransformArgs& Args) override;
-	TiledBlobPtr							AddBlobToCombine(TiledBlobPtr BlobToCombine);
-	void									FreeCombinedBlob();
+	UE_API virtual std::shared_ptr<BlobTransform>	DuplicateInstance(FString NewName) override;
+	UE_API virtual AsyncTransformResultPtr			Exec(const TransformArgs& Args) override;
+	UE_API TiledBlobPtr							AddBlobToCombine(TiledBlobPtr BlobToCombine);
+	UE_API void									FreeCombinedBlob();
 
-	static const FName						MemberName; /// Default struct parameter for tile
+	static UE_API const FName						MemberName; /// Default struct parameter for tile
 };
 
 typedef std::shared_ptr<RenderMaterial_FX_Combined> RenderMaterial_FX_CombinedPtr;
+
+#undef UE_API
diff --git a/Engine/Plugins/TextureGraph/Source/TextureGraphEngine/FxMat/RenderMaterial_FX_MinMax.h b/Engine/Plugins/TextureGraph/Source/TextureGraphEngine/FxMat/RenderMaterial_FX_MinMax.h
--- a/Engine/Plugins/TextureGraph/Source/TextureGraphEngine/FxMat/RenderMaterial_FX_MinMax.h
+++ b/Engine/Plugins/TextureGraph/Source/TextureGraphEngine/FxMat/RenderMaterial_FX_MinMax.h
@@ -5,12 +5,14 @@
 #include "RenderMaterial_FX.h"
 #include "Data/TiledBlob.h"
 
+#define UE_API TEXTUREGRAPHENGINE_API
+
 DECLARE_LOG_CATEGORY_EXTERN(LogRenderMaterial_FX_MinMax, All, All);
 
 class Tex;
 typedef std::shared_ptr<Tex>		TexPtr;
 
-class TEXTUREGRAPHENGINE_API RenderMaterial_FX_MinMax : public RenderMaterial_FX
+class RenderMaterial_FX_MinMax : public RenderMaterial_FX
 {
 private:
 	
@@ -19,13 +21,15 @@ private:
 	TArray<TexPtr>					DownsampledResultTargets;		/// used for creating intermediate targets.
 	
 public:
-	RenderMaterial_FX_MinMax(FString InName, FxMaterialPtr InMaterial, FxMaterialPtr InSecondPassMaterial);
-	virtual ~RenderMaterial_FX_MinMax() override;
+	UE_API RenderMaterial_FX_MinMax(FString InName, FxMaterialPtr InMaterial, FxMaterialPtr InSecondPassMaterial);
+	UE_API virtual ~RenderMaterial_FX_MinMax() override;
 
-	virtual AsyncPrepareResult		PrepareResources(const TransformArgs& Args) override;
-	virtual AsyncTransformResultPtr	Exec(const TransformArgs& Args) override;
+	UE_API virtual AsyncPrepareResult		PrepareResources(const TransformArgs& Args) override;
+	UE_API virtual AsyncTransformResultPtr	Exec(const TransformArgs& Args) override;
 	
 	void							SetDescriptor(const BufferDescriptor& InDesc) { SourceDesc = InDesc;}
 };
 
 typedef std::shared_ptr<RenderMaterial_FX_MinMax> RenderMaterial_FX_MinMaxPtr;
+
+#undef UE_API
diff --git a/Engine/Plugins/TextureGraph/Source/TextureGraphEngine/FxMat/RenderMaterial_Thumbnail.h b/Engine/Plugins/TextureGraph/Source/TextureGraphEngine/FxMat/RenderMaterial_Thumbnail.h
--- a/Engine/Plugins/TextureGraph/Source/TextureGraphEngine/FxMat/RenderMaterial_Thumbnail.h
+++ b/Engine/Plugins/TextureGraph/Source/TextureGraphEngine/FxMat/RenderMaterial_Thumbnail.h
@@ -6,8 +6,10 @@
 #include "UObject/StrongObjectPtr.h"
 #include "RenderMaterial_BP.h"
 
+#define UE_API TEXTUREGRAPHENGINE_API
+
 // Subset of RenderMaterial BP, This class will always render a thumbnail from its material instance
-class TEXTUREGRAPHENGINE_API RenderMaterial_Thumbnail : public RenderMaterial_BP
+class RenderMaterial_Thumbnail : public RenderMaterial_BP
 {
 public:
 	static constexpr uint32	GThumbWidth = 128;			/// Width of standard thumbnail image
@@ -17,12 +19,14 @@ private:
 	ERHIFeatureLevel::Type FeatureLevel;
 
 public:
-	RenderMaterial_Thumbnail(FString InName, UMaterialInterface* InMaterial);
+	UE_API RenderMaterial_Thumbnail(FString InName, UMaterialInterface* InMaterial);
 
-	virtual void BlitTo(FRHICommandListImmediate& RHI, UTextureRenderTarget2D* DstRT, const RenderMesh* MeshObj, int32 TargetId) const override;
+	UE_API virtual void BlitTo(FRHICommandListImmediate& RHI, UTextureRenderTarget2D* DstRT, const RenderMesh* MeshObj, int32 TargetId) const override;
 	virtual bool CanHandleTiles() const override { return true; };
-	virtual std::shared_ptr<BlobTransform> DuplicateInstance(FString InName) override;
-	virtual AsyncPrepareResult PrepareResources(const TransformArgs& Args) override;
+	UE_API virtual std::shared_ptr<BlobTransform> DuplicateInstance(FString InName) override;
+	UE_API virtual AsyncPrepareResult PrepareResources(const TransformArgs& Args) override;
 };
 
-typedef std::shared_ptr<RenderMaterial_Thumbnail> RenderMaterial_ThumbPtr;
\ No newline at end of file
+typedef std::shared_ptr<RenderMaterial_Thumbnail> RenderMaterial_ThumbPtr;
+
+#undef UE_API
diff --git a/Engine/Plugins/TextureGraph/Source/TextureGraphEngine/Helper/ColorUtil.h b/Engine/Plugins/TextureGraph/Source/TextureGraphEngine/Helper/ColorUtil.h
--- a/Engine/Plugins/TextureGraph/Source/TextureGraphEngine/Helper/ColorUtil.h
+++ b/Engine/Plugins/TextureGraph/Source/TextureGraphEngine/Helper/ColorUtil.h
@@ -4,38 +4,42 @@
 #include "CoreMinimal.h"
 #include "2D/TextureType.h"
 
-class TEXTUREGRAPHENGINE_API ColorUtil
+#define UE_API TEXTUREGRAPHENGINE_API
+
+class ColorUtil
 {
 public:
 	//Converting from Array to Functions to avoid static fiasco problem with CPP
-	static const FLinearColor			DefaultColor(TextureType Type);
-	static const FLinearColor			DefaultDiffuse();
-	static const FLinearColor			DefaultSpecular();
-	static const FLinearColor			DefaultAlbedo();
-	static const FLinearColor			DefaultMetalness();
-	static const FLinearColor			DefaultNormal();
-	static const FLinearColor			DefaultDisplacement();
-	static const FLinearColor			DefaultOpacity();
-	static const FLinearColor			DefaultRoughness();
-	static const FLinearColor			DefaultAO();
-	static const FLinearColor			DefaultCurvature();
-	static const FLinearColor			DefaultPreview();
+	static UE_API const FLinearColor			DefaultColor(TextureType Type);
+	static UE_API const FLinearColor			DefaultDiffuse();
+	static UE_API const FLinearColor			DefaultSpecular();
+	static UE_API const FLinearColor			DefaultAlbedo();
+	static UE_API const FLinearColor			DefaultMetalness();
+	static UE_API const FLinearColor			DefaultNormal();
+	static UE_API const FLinearColor			DefaultDisplacement();
+	static UE_API const FLinearColor			DefaultOpacity();
+	static UE_API const FLinearColor			DefaultRoughness();
+	static UE_API const FLinearColor			DefaultAO();
+	static UE_API const FLinearColor			DefaultCurvature();
+	static UE_API const FLinearColor			DefaultPreview();
 
-	static float						GetHue(const FColor& Color);
-	static float						GetSquaredDistance(FLinearColor Current, FLinearColor Match);
-	static FString						GetColorName(FLinearColor Color);
+	static UE_API float						GetHue(const FColor& Color);
+	static UE_API float						GetSquaredDistance(FLinearColor Current, FLinearColor Match);
+	static UE_API FString						GetColorName(FLinearColor Color);
 
-	static FLinearColor					HSV2RGB(FLinearColor C);
-	static FLinearColor					RGB2HSV(FLinearColor C);
-	static FLinearColor					HSVTweak(FLinearColor C, float H, float S, float V);
+	static UE_API FLinearColor					HSV2RGB(FLinearColor C);
+	static UE_API FLinearColor					RGB2HSV(FLinearColor C);
+	static UE_API FLinearColor					HSVTweak(FLinearColor C, float H, float S, float V);
 
-	static bool							IsColorBlack(const FLinearColor& Color, bool IgnoreAlpha = true);
-	static bool							IsColorWhite(const FLinearColor& Color, bool IgnoreAlpha = true);
-	static bool							IsColorGray(const FLinearColor& Color, bool IgnoreAlpha = true);
-	static bool							IsColorRed(const FLinearColor& Color, bool IgnoreAlpha = true);
-	static bool							IsColorGreen(const FLinearColor& Color, bool IgnoreAlpha = true);
-	static bool							IsColorBlue(const FLinearColor& Color, bool IgnoreAlpha = true);
-	static bool							IsColorYellow(const FLinearColor& Color, bool IgnoreAlpha = true);
-	static bool							IsColorMagenta(const FLinearColor& Color, bool IgnoreAlpha = true);
-	static bool							IsColorNear(const FLinearColor& Color, const FLinearColor& Ref, bool IgnoreAlpha = true);
+	static UE_API bool							IsColorBlack(const FLinearColor& Color, bool IgnoreAlpha = true);
+	static UE_API bool							IsColorWhite(const FLinearColor& Color, bool IgnoreAlpha = true);
+	static UE_API bool							IsColorGray(const FLinearColor& Color, bool IgnoreAlpha = true);
+	static UE_API bool							IsColorRed(const FLinearColor& Color, bool IgnoreAlpha = true);
+	static UE_API bool							IsColorGreen(const FLinearColor& Color, bool IgnoreAlpha = true);
+	static UE_API bool							IsColorBlue(const FLinearColor& Color, bool IgnoreAlpha = true);
+	static UE_API bool							IsColorYellow(const FLinearColor& Color, bool IgnoreAlpha = true);
+	static UE_API bool							IsColorMagenta(const FLinearColor& Color, bool IgnoreAlpha = true);
+	static UE_API bool							IsColorNear(const FLinearColor& Color, const FLinearColor& Ref, bool IgnoreAlpha = true);
 };
+
+#undef UE_API
diff --git a/Engine/Plugins/TextureGraph/Source/TextureGraphEngine/Helper/DataUtil.h b/Engine/Plugins/TextureGraph/Source/TextureGraphEngine/Helper/DataUtil.h
--- a/Engine/Plugins/TextureGraph/Source/TextureGraphEngine/Helper/DataUtil.h
+++ b/Engine/Plugins/TextureGraph/Source/TextureGraphEngine/Helper/DataUtil.h
@@ -26,7 +26,7 @@ DECLARE_LOG_CATEGORY_EXTERN(LogData, Log, All);
 #define MX_HASH_VAL(HashValue, Prime, val) ((HashValue * Prime) ^ (val))
 #define MX_HASH_VAL_DEF(val)		MX_HASH_VAL(DataUtil::GFNVInit, DataUtil::GFNVPrime, (HashType)val)
 
-struct TEXTUREGRAPHENGINE_API DataUtil
+struct DataUtil
 {
 	static constexpr HashType		GNullHash = 0; /// 0xDeadBeef;
 	//////////////////////////////////////////////////////////////////////////
@@ -39,15 +39,15 @@ struct TEXTUREGRAPHENGINE_API DataUtil
 	static const size_t				GMaxChunk = 16 * 1024;
 
 	/// Calculate HashValue for one chunk of data
-	static HashType					Hash(const uint8* data, size_t Length, HashType InitialHash = GFNVInit, HashType Prime = GFNVPrime);
-	static HashType					Hash_GenericString_Name(const FString& Name, HashType InitialHash = GFNVInit, HashType Prime = GFNVPrime);
-	static HashType					Hash(const HashTypeVec& SubHashes, HashType InitialValue = GFNVInit, HashType Prime = GFNVPrime);
-	static HashType					Hash(const CHashPtrVec& InSubHashes, HashType InitialValue = GFNVInit, HashType Prime = GFNVPrime);
+	static TEXTUREGRAPHENGINE_API HashType					Hash(const uint8* data, size_t Length, HashType InitialHash = GFNVInit, HashType Prime = GFNVPrime);
+	static TEXTUREGRAPHENGINE_API HashType					Hash_GenericString_Name(const FString& Name, HashType InitialHash = GFNVInit, HashType Prime = GFNVPrime);
+	static TEXTUREGRAPHENGINE_API HashType					Hash(const HashTypeVec& SubHashes, HashType InitialValue = GFNVInit, HashType Prime = GFNVPrime);
+	static TEXTUREGRAPHENGINE_API HashType					Hash(const CHashPtrVec& InSubHashes, HashType InitialValue = GFNVInit, HashType Prime = GFNVPrime);
 
 	/// Automatically decide whether to chunk the HashValue or not
-	static HashType					Hash_One(const uint8* Data, size_t Length, HashType InitialValue = GFNVInit, HashType Prime = GFNVPrime);
-	static HashType					Hash_Chunked(const uint8* data, size_t Length, HashType InitialValue = GFNVInit, HashType Prime = GFNVPrime);
-	static size_t					GetOptimalHashingSize(size_t Size);
+	static TEXTUREGRAPHENGINE_API HashType					Hash_One(const uint8* Data, size_t Length, HashType InitialValue = GFNVInit, HashType Prime = GFNVPrime);
+	static TEXTUREGRAPHENGINE_API HashType					Hash_Chunked(const uint8* data, size_t Length, HashType InitialValue = GFNVInit, HashType Prime = GFNVPrime);
+	static TEXTUREGRAPHENGINE_API size_t					GetOptimalHashingSize(size_t Size);
 
 	template <typename T>
 	static HashType					Hash_Simple(const T& Value, HashType InitialValue = GFNVInit, HashType Prime = GFNVPrime) 
@@ -145,7 +145,7 @@ struct TEXTUREGRAPHENGINE_API DataUtil
 //////////////////////////////////////////////////////////////////////////
 /// Hash data structure
 //////////////////////////////////////////////////////////////////////////
-class TEXTUREGRAPHENGINE_API CHash : public std::enable_shared_from_this<CHash>
+class CHash : public std::enable_shared_from_this<CHash>
 {
 private:
 	HashType						HashValue = DataUtil::GNullHash;/// The actual Value of the hash
@@ -167,24 +167,24 @@ private:
 																/// hash became finalised
 
 	FORCEINLINE bool				HasTempDependency_Internal() const { return (IsTemp() && !HashSources.empty() && TempHashValue->IsTemp()); }
-	void							AddLink(CHashPtrW Link);
-	void							CheckLinkCycles(std::unordered_set<CHashPtr>& Chain);
-	void							UpdateLinks();
-	void							HandleLinkUpdated(CHashPtr LinkUpdated);
+	TEXTUREGRAPHENGINE_API void							AddLink(CHashPtrW Link);
+	TEXTUREGRAPHENGINE_API void							CheckLinkCycles(std::unordered_set<CHashPtr>& Chain);
+	TEXTUREGRAPHENGINE_API void							UpdateLinks();
+	TEXTUREGRAPHENGINE_API void							HandleLinkUpdated(CHashPtr LinkUpdated);
 
 									/// DO NOT MAKE THIS C'TOR PUBLIC. Use ConstructFromSources instead
-	explicit						CHash(const CHashPtrVec& Sources);
+	TEXTUREGRAPHENGINE_API explicit						CHash(const CHashPtrVec& Sources);
 
 public:
-									CHash(HashType Value, bool bInIsFinal);
-	explicit						CHash(CHashPtr Temp);
+									TEXTUREGRAPHENGINE_API CHash(HashType Value, bool bInIsFinal);
+	TEXTUREGRAPHENGINE_API explicit						CHash(CHashPtr Temp);
 
-	bool							TryFinalise(HashType FinalHash = DataUtil::GNullHash, bool UpdateBlobber = true);
-	bool							operator == (const CHash& RHS) const;
-	HashTypeVec						GetIntermediateHashes() const;
+	TEXTUREGRAPHENGINE_API bool							TryFinalise(HashType FinalHash = DataUtil::GNullHash, bool UpdateBlobber = true);
+	TEXTUREGRAPHENGINE_API bool							operator == (const CHash& RHS) const;
+	TEXTUREGRAPHENGINE_API HashTypeVec						GetIntermediateHashes() const;
 
-	static CHashPtr					ConstructFromSources(const CHashPtrVec& Sources);
-	static CHashPtr					UpdateHash(CHashPtr NewHash, CHashPtr PrevHash);
+	static TEXTUREGRAPHENGINE_API CHashPtr					ConstructFromSources(const CHashPtrVec& Sources);
+	static TEXTUREGRAPHENGINE_API CHashPtr					UpdateHash(CHashPtr NewHash, CHashPtr PrevHash);
 
 	//////////////////////////////////////////////////////////////////////////
 	/// Inline functions
@@ -211,7 +211,7 @@ typedef cti::continuable<CHashPtr>	AscynCHashPtr;
 
 //////////////////////////////////////////////////////////////////////////
 template <typename Type>
-class TEXTUREGRAPHENGINE_API T_Tiles
+class T_Tiles
 {
 public:
 	typedef std::vector<Type>		TypeVec;
diff --git a/Engine/Plugins/TextureGraph/Source/TextureGraphEngine/Helper/GraphicsUtil.h b/Engine/Plugins/TextureGraph/Source/TextureGraphEngine/Helper/GraphicsUtil.h
--- a/Engine/Plugins/TextureGraph/Source/TextureGraphEngine/Helper/GraphicsUtil.h
+++ b/Engine/Plugins/TextureGraph/Source/TextureGraphEngine/Helper/GraphicsUtil.h
@@ -6,6 +6,8 @@
 #include "GraphicsDefs.h"
 #include <array>
 
+#define UE_API TEXTUREGRAPHENGINE_API
+
 enum class BoundsCorners
 {
 	L0_BottomLeft,
@@ -19,27 +21,27 @@ enum class BoundsCorners
 	L1_BottomRight
 };
 
-class TEXTUREGRAPHENGINE_API GraphicsUtil
+class GraphicsUtil
 {
 public:
-	static Vector2				CalculateBarycentricCoordinates(const Vector2* UV, const Vector2& P);
-	static Vector2				CalculateBarycentricCoordinates(const Vector3& A, const Vector3& B, const Vector3& C, const Vector3& P);
+	static UE_API Vector2				CalculateBarycentricCoordinates(const Vector2* UV, const Vector2& P);
+	static UE_API Vector2				CalculateBarycentricCoordinates(const Vector3& A, const Vector3& B, const Vector3& C, const Vector3& P);
 
-	static bool					IsPointInTriangle(const Vector3& A, const Vector3& B, const Vector3& C, const Vector3& P);
-	static bool					IsPointInTriangle(float u, float v, float tolerance = 0.001f);
-	static bool					IsPointInTriangle(const Vector2& uv, float tolerance = 0.001f);
+	static UE_API bool					IsPointInTriangle(const Vector3& A, const Vector3& B, const Vector3& C, const Vector3& P);
+	static UE_API bool					IsPointInTriangle(float u, float v, float tolerance = 0.001f);
+	static UE_API bool					IsPointInTriangle(const Vector2& uv, float tolerance = 0.001f);
 
 	//////////////////////////////////////////////////////////////////////////
-	static Vector3				s_boundsAxis[6];
-	static float				SquaredDistance(const Vector3& p0, const Vector3& p1);
-	static float				SquaredLength(const Vector3& pt);
+	static UE_API Vector3				s_boundsAxis[6];
+	static UE_API float				SquaredDistance(const Vector3& p0, const Vector3& p1);
+	static UE_API float				SquaredLength(const Vector3& pt);
 
-	static std::array<Vector3, 8> GetBoundsCorners(const FBox& bounds);
+	static UE_API std::array<Vector3, 8> GetBoundsCorners(const FBox& bounds);
 
-	static Vector3				GetBoundsPoint(const FBox& bounds, BoundsCorners pti);
-	static FBox					GetChildBounds_QuadTree(const FBox& bounds, BoundsCorners ci);
-	static FBox					GetChildBounds_Octree(const FBox& bounds, BoundsCorners ci);
-	static bool					CheckRayTriangleIntersection(const FRay& ray, const Vector3& p1, const Vector3& p2, const Vector3& p3, float& distance, Vector3& ip);
+	static UE_API Vector3				GetBoundsPoint(const FBox& bounds, BoundsCorners pti);
+	static UE_API FBox					GetChildBounds_QuadTree(const FBox& bounds, BoundsCorners ci);
+	static UE_API FBox					GetChildBounds_Octree(const FBox& bounds, BoundsCorners ci);
+	static UE_API bool					CheckRayTriangleIntersection(const FRay& ray, const Vector3& p1, const Vector3& p2, const Vector3& p3, float& distance, Vector3& ip);
 
 	static FORCEINLINE Vector4	ConvertTangent(const Tangent& t) { return Vector4(t.TangentX.X, t.TangentX.Y, t.TangentX.Z, t.bFlipTangentY ? -1.0f : 1.0f); }
 	static FORCEINLINE FLinearColor TangentAsLinearColor(const Tangent& t) { return FLinearColor(ConvertTangent(t)); }
@@ -58,3 +60,5 @@ public:
 		return v0 + u * (v1 - v0) + v * (v2 - v0);
 	}
 };
+
+#undef UE_API
diff --git a/Engine/Plugins/TextureGraph/Source/TextureGraphEngine/Helper/MathUtils.h b/Engine/Plugins/TextureGraph/Source/TextureGraphEngine/Helper/MathUtils.h
--- a/Engine/Plugins/TextureGraph/Source/TextureGraphEngine/Helper/MathUtils.h
+++ b/Engine/Plugins/TextureGraph/Source/TextureGraphEngine/Helper/MathUtils.h
@@ -2,26 +2,30 @@
 #pragma once
 
 #include "CoreMinimal.h"
+
+#define UE_API TEXTUREGRAPHENGINE_API
 /**
  * 
  */
-class TEXTUREGRAPHENGINE_API MathUtils
+class MathUtils
 {
 public:
 	MathUtils() = delete;
 
-	static const float					GMeterToCm;
+	static UE_API const float					GMeterToCm;
 
-	static float						MinFloat(); 
-	static float						MaxFloat(); 
+	static UE_API float						MinFloat(); 
+	static UE_API float						MaxFloat(); 
 	
-	static FVector						MinFVector();
-	static FVector						MaxFVector();
-	static FVector2f					MinFVector2();
-	static FVector2f					MaxFVector2();
-	static void							UpdateBounds(FBox& bounds, const FVector& point);
-	static void							EncapsulateBound(FBox& bounds, FBox& otherBounds);
-	static FVector						GetDirection(float yzAngle, float xAngle, int xSign = 1);
-	static FBox							GetCombinedBounds(TArray<FBox> inputBounds);
-	static float						Step(float Y, float X);
+	static UE_API FVector						MinFVector();
+	static UE_API FVector						MaxFVector();
+	static UE_API FVector2f					MinFVector2();
+	static UE_API FVector2f					MaxFVector2();
+	static UE_API void							UpdateBounds(FBox& bounds, const FVector& point);
+	static UE_API void							EncapsulateBound(FBox& bounds, FBox& otherBounds);
+	static UE_API FVector						GetDirection(float yzAngle, float xAngle, int xSign = 1);
+	static UE_API FBox							GetCombinedBounds(TArray<FBox> inputBounds);
+	static UE_API float						Step(float Y, float X);
 };
+
+#undef UE_API
diff --git a/Engine/Plugins/TextureGraph/Source/TextureGraphEngine/Helper/Promise.h b/Engine/Plugins/TextureGraph/Source/TextureGraphEngine/Helper/Promise.h
--- a/Engine/Plugins/TextureGraph/Source/TextureGraphEngine/Helper/Promise.h
+++ b/Engine/Plugins/TextureGraph/Source/TextureGraphEngine/Helper/Promise.h
@@ -9,10 +9,12 @@ THIRD_PARTY_INCLUDES_END
 #include "Async/TaskGraphInterfaces.h"
 #include "Util.h"
 
+#define UE_API TEXTUREGRAPHENGINE_API
+
 typedef cti::continuable<int32> AsyncInt;
 
 // This is temporary and will eventually go away
-struct TEXTUREGRAPHENGINE_API ActionResult
+struct ActionResult
 {
 	std::exception_ptr				ExInner;					/// Original exception that was raised by the action
 	int32							ErrorCode = 0;				/// What is the error code
@@ -29,12 +31,12 @@ typedef std::shared_ptr<ActionResult>		ActionResultPtr;
 typedef cti::continuable<ActionResultPtr>	AsyncActionResultPtr;
 
 //////////////////////////////////////////////////////////////////////////
-struct TEXTUREGRAPHENGINE_API PromiseUtil 
+struct PromiseUtil 
 {
 	////////////////////////////////////////////////////////////////////////////
 	///// Static functions
 	////////////////////////////////////////////////////////////////////////////
-	static auto						OnThread(ENamedThreads::Type ThreadType);
+	static UE_API auto						OnThread(ENamedThreads::Type ThreadType);
 
 	//////////////////////////////////////////////////////////////////////////
 	/// Inline functions
@@ -117,3 +119,5 @@ struct TEXTUREGRAPHENGINE_API PromiseUtil
 };
 
 #define FWD_PROMISE(p) p = std::forward<decltype(p)>(p)
+
+#undef UE_API
diff --git a/Engine/Plugins/TextureGraph/Source/TextureGraphEngine/Helper/Util.h b/Engine/Plugins/TextureGraph/Source/TextureGraphEngine/Helper/Util.h
--- a/Engine/Plugins/TextureGraph/Source/TextureGraphEngine/Helper/Util.h
+++ b/Engine/Plugins/TextureGraph/Source/TextureGraphEngine/Helper/Util.h
@@ -7,6 +7,8 @@
 #include <UObject/ObjectMacros.h>
 #include <UObject/ReflectedTypeAccessors.h>
 
+#define UE_API TEXTUREGRAPHENGINE_API
+
 //////////////////////////////////////////////////////////////////////////
 enum class E_Priority
 {
@@ -49,7 +51,7 @@ class UMaterial;
 class UPackage;
 struct FWorldContext;
 class FRHICommandListImmediate;
-struct TEXTUREGRAPHENGINE_API Util
+struct Util
 {
 	static constexpr int32			GDefaultWidth = 1024;
 	static constexpr int32			GDefaultHeight = 1024;
@@ -57,11 +59,11 @@ struct TEXTUREGRAPHENGINE_API Util
 	//////////////////////////////////////////////////////////////////////////
 	/// Package related utility functions
 	//////////////////////////////////////////////////////////////////////////
-	static UObject*					GetModelPackage();
-	static UObject*					GetTexturesPackage();
-	static UObject*					GetRenderTargetPackage();
-	static UObject*					GetMegascansTexturesPackage();
-	static UObject*					GetMaterialsPackage();
+	static UE_API UObject*					GetModelPackage();
+	static UE_API UObject*					GetTexturesPackage();
+	static UE_API UObject*					GetRenderTargetPackage();
+	static UE_API UObject*					GetMegascansTexturesPackage();
+	static UE_API UObject*					GetMaterialsPackage();
 
 	//////////////////////////////////////////////////////////////////////////
 	/// Path related
@@ -71,30 +73,30 @@ struct TEXTUREGRAPHENGINE_API Util
 	static FORCEINLINE FString		GetRuntimeTexturesPackagePath(FString suffix) { return GetRuntimeTexturesPackagePath() + suffix; }
 	static FORCEINLINE FString		GetRuntimeTexturesPackagePath(const char* suffix) { return GetRuntimeTexturesPackagePath() + FString(suffix); }
 
-	static bool						IsFileValid(const FString& FileName);
-	static bool						IsDirectory(const FString& Path);
-	static bool						IsDirectoryEmpty(const FString& Path);
+	static UE_API bool						IsFileValid(const FString& FileName);
+	static UE_API bool						IsDirectory(const FString& Path);
+	static UE_API bool						IsDirectoryEmpty(const FString& Path);
 
 	//////////////////////////////////////////////////////////////////////////
 	/// Threading related
 	//////////////////////////////////////////////////////////////////////////
-	static void						OnGameThread(TUniqueFunction<void()> Callback);
-	static void						OnRenderingThread(TUniqueFunction<void(FRHICommandListImmediate&)> Callback);
-	static void						OnBackgroundThread(TUniqueFunction<void()> Callback);
-	static void						OnAnyThread(TUniqueFunction<void()> Callback);
-	static void						OnThread(ENamedThreads::Type thread, TUniqueFunction<void()> Callback);
-	static size_t					GetCurrentThreadId();
+	static UE_API void						OnGameThread(TUniqueFunction<void()> Callback);
+	static UE_API void						OnRenderingThread(TUniqueFunction<void(FRHICommandListImmediate&)> Callback);
+	static UE_API void						OnBackgroundThread(TUniqueFunction<void()> Callback);
+	static UE_API void						OnAnyThread(TUniqueFunction<void()> Callback);
+	static UE_API void						OnThread(ENamedThreads::Type thread, TUniqueFunction<void()> Callback);
+	static UE_API size_t					GetCurrentThreadId();
 
-	static FRHICommandListImmediate& RHIImmediate();
+	static UE_API FRHICommandListImmediate& RHIImmediate();
 
 	//////////////////////////////////////////////////////////////////////////
 	/// Game objects and engine related
 	//////////////////////////////////////////////////////////////////////////
-	static const FWorldContext*		GetGameWorldContext();
-	static UWorld*					GetGameWorld();
-	static void*					GetOSWindowHandle();
-	static FString					RandomID();
-	static FString					ToPascalCase(const FString& strToConvert);
+	static UE_API const FWorldContext*		GetGameWorldContext();
+	static UE_API UWorld*					GetGameWorld();
+	static UE_API void*					GetOSWindowHandle();
+	static UE_API FString					RandomID();
+	static UE_API FString					ToPascalCase(const FString& strToConvert);
 	static FORCEINLINE double		Time() { return FPlatformTime::Seconds() * 1000.0; }
 	static FORCEINLINE double		TimeDelta(double start) { return Time() - start; }
 
@@ -121,4 +123,6 @@ struct TEXTUREGRAPHENGINE_API Util
 		return Name;
 	}
 
-};
\ No newline at end of file
+};
+
+#undef UE_API
diff --git a/Engine/Plugins/TextureGraph/Source/TextureGraphEngine/Job/BlobHasherService.h b/Engine/Plugins/TextureGraph/Source/TextureGraphEngine/Job/BlobHasherService.h
--- a/Engine/Plugins/TextureGraph/Source/TextureGraphEngine/Job/BlobHasherService.h
+++ b/Engine/Plugins/TextureGraph/Source/TextureGraphEngine/Job/BlobHasherService.h
@@ -5,15 +5,19 @@
 #include "BlobHelperService.h"
 #include <vector>
 
-class TEXTUREGRAPHENGINE_API BlobHasherService : public BlobHelperService
+#define UE_API TEXTUREGRAPHENGINE_API
+
+class BlobHasherService : public BlobHelperService
 {
 public:
-									BlobHasherService();
-	virtual							~BlobHasherService() override;
+									UE_API BlobHasherService();
+	UE_API virtual							~BlobHasherService() override;
 
-	virtual AsyncJobResultPtr		Tick() override;
-	virtual void					Add(BlobRef BlobObj) override;
+	UE_API virtual AsyncJobResultPtr		Tick() override;
+	UE_API virtual void					Add(BlobRef BlobObj) override;
 };
 
 typedef std::shared_ptr<BlobHasherService>	BlobHasherServicePtr;
 typedef std::weak_ptr<BlobHasherService>	BlobHasherServicePtrW;
+
+#undef UE_API
diff --git a/Engine/Plugins/TextureGraph/Source/TextureGraphEngine/Job/BlobHelperService.h b/Engine/Plugins/TextureGraph/Source/TextureGraphEngine/Job/BlobHelperService.h
--- a/Engine/Plugins/TextureGraph/Source/TextureGraphEngine/Job/BlobHelperService.h
+++ b/Engine/Plugins/TextureGraph/Source/TextureGraphEngine/Job/BlobHelperService.h
@@ -5,16 +5,20 @@
 #include "IdleService.h"
 #include <vector>
 
-class TEXTUREGRAPHENGINE_API BlobHelperService : public IdleService
+#define UE_API TEXTUREGRAPHENGINE_API
+
+class BlobHelperService : public IdleService
 {
 protected:
 	std::vector<BlobPtr>			Blobs;					/// The blobs that we want to hash
 
 public:
-	explicit						BlobHelperService(const FString& InName);
-	virtual							~BlobHelperService() override;
+	UE_API explicit						BlobHelperService(const FString& InName);
+	UE_API virtual							~BlobHelperService() override;
 
-	virtual void					Stop() override;
+	UE_API virtual void					Stop() override;
 
-	virtual void					Add(BlobRef BlobObj);
+	UE_API virtual void					Add(BlobRef BlobObj);
 };
+
+#undef UE_API
diff --git a/Engine/Plugins/TextureGraph/Source/TextureGraphEngine/Job/DeviceTransferService.h b/Engine/Plugins/TextureGraph/Source/TextureGraphEngine/Job/DeviceTransferService.h
--- a/Engine/Plugins/TextureGraph/Source/TextureGraphEngine/Job/DeviceTransferService.h
+++ b/Engine/Plugins/TextureGraph/Source/TextureGraphEngine/Job/DeviceTransferService.h
@@ -5,9 +5,11 @@
 #include "IdleService.h"
 #include <unordered_map>
 
+#define UE_API TEXTUREGRAPHENGINE_API
+
 DECLARE_LOG_CATEGORY_EXTERN(LogTransfer_Svc, Log, Verbose);
 
-class TEXTUREGRAPHENGINE_API DeviceTransferService : public IdleService
+class DeviceTransferService : public IdleService
 {
 protected:
 	static constexpr int			GMaxFlush = 32;		/// Size is chosen so that TransferInfo
@@ -28,16 +30,18 @@ protected:
 									Transfers;				/// The current transfers that we're working on
 
 public:
-									DeviceTransferService();
-	virtual							~DeviceTransferService() override;
+									UE_API DeviceTransferService();
+	UE_API virtual							~DeviceTransferService() override;
 
-	virtual AsyncJobResultPtr		Tick() override;
-	virtual void					Stop() override;
+	UE_API virtual AsyncJobResultPtr		Tick() override;
+	UE_API virtual void					Stop() override;
 	virtual bool					ShouldRunDuringBusyUpdate() const override { return true; }
 
-	virtual	void					QueueTransfer(BlobPtr BlobObj, Device* TargetDevice);
-	virtual void					AbortTransfer(HashType HashValue);
+	UE_API virtual	void					QueueTransfer(BlobPtr BlobObj, Device* TargetDevice);
+	UE_API virtual void					AbortTransfer(HashType HashValue);
 };
 
 typedef std::shared_ptr<DeviceTransferService> DeviceTransferServicePtr;
 typedef std::weak_ptr<DeviceTransferService>	DeviceTransferServicePtrW;
+
+#undef UE_API
diff --git a/Engine/Plugins/TextureGraph/Source/TextureGraphEngine/Job/HistogramService.h b/Engine/Plugins/TextureGraph/Source/TextureGraphEngine/Job/HistogramService.h
--- a/Engine/Plugins/TextureGraph/Source/TextureGraphEngine/Job/HistogramService.h
+++ b/Engine/Plugins/TextureGraph/Source/TextureGraphEngine/Job/HistogramService.h
@@ -5,6 +5,8 @@
 #include "IdleService.h"
 #include "Model/Mix/MixUpdateCycle.h"
 
+#define UE_API TEXTUREGRAPHENGINE_API
+
 class UModelObject;
 typedef std::shared_ptr<class Job>		JobPtr;
 typedef std::unique_ptr<class Job>		JobUPtr;
@@ -12,25 +14,26 @@ typedef std::weak_ptr<class Job>		JobPtrW;
 
 class UMixInterface;
 
-class TEXTUREGRAPHENGINE_API HistogramService : public IdleService
+class HistogramService : public IdleService
 {
 private:
 	JobBatchPtr Batch;
 	bool bCaptureNextBatch = false;
 
 public:
-	HistogramService();
-	virtual	~HistogramService() override;
+	UE_API HistogramService();
+	UE_API virtual	~HistogramService() override;
 
-	virtual AsyncJobResultPtr Tick() override;
-	virtual void Stop() override;
+	UE_API virtual AsyncJobResultPtr Tick() override;
+	UE_API virtual void Stop() override;
 
-	void AddHistogramJob(MixUpdateCyclePtr Cycle,JobUPtr JobToAdd, int32 TargetID, UMixInterface* Mix);
-	JobBatchPtr GetOrCreateNewBatch(UMixInterface* mix);
+	UE_API void AddHistogramJob(MixUpdateCyclePtr Cycle,JobUPtr JobToAdd, int32 TargetID, UMixInterface* Mix);
+	UE_API JobBatchPtr GetOrCreateNewBatch(UMixInterface* mix);
 
-	void CaptureNextBatch();
+	UE_API void CaptureNextBatch();
 };
 
 typedef std::shared_ptr<HistogramService>	HistogramServicePtr;
 typedef std::weak_ptr<HistogramService>	HistogramServicePtrW;
 
+#undef UE_API
diff --git a/Engine/Plugins/TextureGraph/Source/TextureGraphEngine/Job/IdleService.h b/Engine/Plugins/TextureGraph/Source/TextureGraphEngine/Job/IdleService.h
--- a/Engine/Plugins/TextureGraph/Source/TextureGraphEngine/Job/IdleService.h
+++ b/Engine/Plugins/TextureGraph/Source/TextureGraphEngine/Job/IdleService.h
@@ -4,9 +4,11 @@
 #include "CoreMinimal.h"
 #include "JobBatch.h"
 
+#define UE_API TEXTUREGRAPHENGINE_API
+
 DECLARE_LOG_CATEGORY_EXTERN(LogIdle_Svc, Log, Verbose);
 
-class TEXTUREGRAPHENGINE_API IdleService
+class IdleService
 {
 public:
 	struct Stats
@@ -32,7 +34,7 @@ public:
 	virtual void					Stop() = 0;
 	virtual bool					ShouldRunDuringBusyUpdate() const { return false; }
 
-	void							UpdateStats(double StartTime, double EndTime, bool DidOffendTimeLimit);
+	UE_API void							UpdateStats(double StartTime, double EndTime, bool DidOffendTimeLimit);
 
 	//////////////////////////////////////////////////////////////////////////
 	/// Inline functions
@@ -42,3 +44,5 @@ public:
 };
 
 typedef std::shared_ptr<IdleService>	IdleServicePtr;
+
+#undef UE_API
diff --git a/Engine/Plugins/TextureGraph/Source/TextureGraphEngine/Job/Job.h b/Engine/Plugins/TextureGraph/Source/TextureGraphEngine/Job/Job.h
--- a/Engine/Plugins/TextureGraph/Source/TextureGraphEngine/Job/Job.h
+++ b/Engine/Plugins/TextureGraph/Source/TextureGraphEngine/Job/Job.h
@@ -8,6 +8,8 @@
 #include "JobArgs.h"
 #include "Data/TiledBlob.h"
 
+#define UE_API TEXTUREGRAPHENGINE_API
+
 class Job;
 class Scheduler;
 class JobBatch;
@@ -18,7 +20,7 @@ DECLARE_LOG_CATEGORY_EXTERN(LogJob, Log, All);
 class RenderMesh;
 
 //////////////////////////////////////////////////////////////////////////
-class TEXTUREGRAPHENGINE_API Job : public DeviceNativeTask
+class Job : public DeviceNativeTask
 {
 public:
 	struct JobStats
@@ -82,96 +84,96 @@ protected:
 	typedef TMap<HashType, IndexPair> JobLUT;
 
 	/// Binds/Unbinds args for individual tiles or the entire blob if its a non-tiled job
-	virtual AsyncJobArgResultPtr	BindArgs(JobRunInfo NewRunInfo, BlobTransformPtr TransformObj, int32 RowId, int32 ColId);
-	virtual AsyncJobArgResultPtr	UnbindArgs(JobRunInfo NewRunInfo, BlobTransformPtr TransformObj, int32 RowId, int32 ColId);
+	UE_API virtual AsyncJobArgResultPtr	BindArgs(JobRunInfo NewRunInfo, BlobTransformPtr TransformObj, int32 RowId, int32 ColId);
+	UE_API virtual AsyncJobArgResultPtr	UnbindArgs(JobRunInfo NewRunInfo, BlobTransformPtr TransformObj, int32 RowId, int32 ColId);
 
 	/// Binds/Unbinds args for all the tiles at once and returns a promise
-	virtual AsyncInt				BindArgs_All(JobRunInfo RunInfo);
-	virtual AsyncInt				UnbindArgs_All(JobRunInfo RunInfo);
+	UE_API virtual AsyncInt				BindArgs_All(JobRunInfo RunInfo);
+	UE_API virtual AsyncInt				UnbindArgs_All(JobRunInfo RunInfo);
 
 	typedef std::function			<AsyncJobArgResultPtr(JobRunInfo, BlobTransformPtr, int32, int32)> Bind_Unbind_Func;
-	virtual AsyncInt				Bind_Or_Unbind_All_Generic(JobRunInfo NewRunInfo, Bind_Unbind_Func BFunc);
+	UE_API virtual AsyncInt				Bind_Or_Unbind_All_Generic(JobRunInfo NewRunInfo, Bind_Unbind_Func BFunc);
 
-	virtual AsyncJobResultPtr		RunTile(JobRunInfo RunInfo, int32 RowId, int32 ColId);
-	virtual AsyncJobResultPtr		RunSingle(JobRunInfo RunInfo);
-	virtual AsyncJobResultPtr		FinaliseTiles(JobRunInfo RunInfo);
-	void							SetTileResult(int32 RowId, int32 ColId, JobResultPtr TileResult);
-	virtual AsyncPrepareResult		PrepareTargets(JobBatch* batch);
+	UE_API virtual AsyncJobResultPtr		RunTile(JobRunInfo RunInfo, int32 RowId, int32 ColId);
+	UE_API virtual AsyncJobResultPtr		RunSingle(JobRunInfo RunInfo);
+	UE_API virtual AsyncJobResultPtr		FinaliseTiles(JobRunInfo RunInfo);
+	UE_API void							SetTileResult(int32 RowId, int32 ColId, JobResultPtr TileResult);
+	UE_API virtual AsyncPrepareResult		PrepareTargets(JobBatch* batch);
 
-	AsyncTransformResultPtr			ExecTransform(JobRunInfo RunInfo, BlobTransformPtr TransformObj , int32 RowId, int32 ColId, CHashPtr jobHash);
-	virtual void					MarkJobDone();
-	virtual bool					IsDiscard() const;
+	UE_API AsyncTransformResultPtr			ExecTransform(JobRunInfo RunInfo, BlobTransformPtr TransformObj , int32 RowId, int32 ColId, CHashPtr jobHash);
+	UE_API virtual void					MarkJobDone();
+	UE_API virtual bool					IsDiscard() const;
 
 	//////////////////////////////////////////////////////////////////////////
 	/// These methods can be used to execute a Job then and there
 	//////////////////////////////////////////////////////////////////////////
-	virtual AsyncJobResultPtr		Run(JobRunInfo RunInfo);
-	virtual AsyncPrepareResult		PrepareResources(JobBatch* batch);
+	UE_API virtual AsyncJobResultPtr		Run(JobRunInfo RunInfo);
+	UE_API virtual AsyncPrepareResult		PrepareResources(JobBatch* batch);
 
-	virtual bool					CheckCached();
-	BufferDescriptor				GetResultDesc() const;
-	BufferDescriptor				GetCombinedDesc(BufferDescriptor& ArgsDescCombined, size_t& Count) const;
+	UE_API virtual bool					CheckCached();
+	UE_API BufferDescriptor				GetResultDesc() const;
+	UE_API BufferDescriptor				GetCombinedDesc(BufferDescriptor& ArgsDescCombined, size_t& Count) const;
 
 public:
-									Job(int32 TargetId, BlobTransformPtr InTransform, UObject* ErrorOwner = nullptr, uint16 InPriority = (uint16)E_Priority::kNormal, uint64 InId = -1);
-									Job(UMixInterface* InMix, int32 TargetId, BlobTransformPtr Transform, UObject* ErrorOwner = nullptr, uint16 priority = (uint16)E_Priority::kNormal, uint64 id = -1);
-	virtual							~Job() override;
+									UE_API Job(int32 TargetId, BlobTransformPtr InTransform, UObject* ErrorOwner = nullptr, uint16 InPriority = (uint16)E_Priority::kNormal, uint64 InId = -1);
+									UE_API Job(UMixInterface* InMix, int32 TargetId, BlobTransformPtr Transform, UObject* ErrorOwner = nullptr, uint16 priority = (uint16)E_Priority::kNormal, uint64 id = -1);
+	UE_API virtual							~Job() override;
 
 
-	virtual Job*					SetArgs(const TArray<JobArgPtr>& NewArgs);
-	virtual Job*					AddArg(JobArgPtr Arg);
-	uint32_t						NumArgs() const;
-	JobArgPtr						GetArg(uint32_t Index) const;
-	virtual void					AddResultToBlobber();
+	UE_API virtual Job*					SetArgs(const TArray<JobArgPtr>& NewArgs);
+	UE_API virtual Job*					AddArg(JobArgPtr Arg);
+	UE_API uint32_t						NumArgs() const;
+	UE_API JobArgPtr						GetArg(uint32_t Index) const;
+	UE_API virtual void					AddResultToBlobber();
 
-	virtual CHashPtr				Hash() const;
-	virtual CHashPtr				TileHash(int32 RowId, int32 ColId) const;
-	virtual CHashPtr				CalcTileHash(int32 RowId, int32 ColId) const;
-	virtual bool					CanHandleTiles() const;
-	virtual bool					CheckDefaultArgs() const;
+	UE_API virtual CHashPtr				Hash() const;
+	UE_API virtual CHashPtr				TileHash(int32 RowId, int32 ColId) const;
+	UE_API virtual CHashPtr				CalcTileHash(int32 RowId, int32 ColId) const;
+	UE_API virtual bool					CanHandleTiles() const;
+	UE_API virtual bool					CheckDefaultArgs() const;
 
-	virtual TiledBlobPtr			InitResult(FString NewName, const BufferDescriptor* InDesiredDesc, int32 NumTilesX = 0, int32 NumTilesY = 0);
+	UE_API virtual TiledBlobPtr			InitResult(FString NewName, const BufferDescriptor* InDesiredDesc, int32 NumTilesX = 0, int32 NumTilesY = 0);
 
-	void							ResetForReplay(bool noCache); /// For debug purpose, reset the job so it can be replayed
+	UE_API void							ResetForReplay(bool noCache); /// For debug purpose, reset the job so it can be replayed
 
 	//////////////////////////////////////////////////////////////////////////
 	/// These methods are used to queue up the job as a DeviceNativeTask.
 	/// You cannot call the DeviceNativeTask implementation directly. That is
 	/// done by the Device itself.
 	//////////////////////////////////////////////////////////////////////////
-	virtual AsyncInt				BeginNative(JobRunInfo RunInfo);
-	virtual AsyncJobResultPtr		EndNative();
+	UE_API virtual AsyncInt				BeginNative(JobRunInfo RunInfo);
+	UE_API virtual AsyncJobResultPtr		EndNative();
 
 	//////////////////////////////////////////////////////////////////////////
 	/// BEGIN: DeviceNativeTask implementation
 	//////////////////////////////////////////////////////////////////////////
 protected:
-	virtual TiledBlobPtr			InitLateBoundResult(FString NewName, BufferDescriptor DesiredDesc, size_t NumInputBlobs);
+	UE_API virtual TiledBlobPtr			InitLateBoundResult(FString NewName, BufferDescriptor DesiredDesc, size_t NumInputBlobs);
 
-	virtual AsyncInt				PreExecAsync(ENamedThreads::Type execThread, ENamedThreads::Type ReturnThread) override;
-	virtual int32					Exec() override;
-	virtual void					PostExec() override;
-	virtual AsyncInt				ExecAsync(ENamedThreads::Type execThread, ENamedThreads::Type ReturnThread) override;
-	virtual JobPtrVec				GetArgDependencies();
+	UE_API virtual AsyncInt				PreExecAsync(ENamedThreads::Type execThread, ENamedThreads::Type ReturnThread) override;
+	UE_API virtual int32					Exec() override;
+	UE_API virtual void					PostExec() override;
+	UE_API virtual AsyncInt				ExecAsync(ENamedThreads::Type execThread, ENamedThreads::Type ReturnThread) override;
+	UE_API virtual JobPtrVec				GetArgDependencies();
 
 public:
-	virtual Device*					GetTargetDevice() const override;
-	virtual ENamedThreads::Type		GetExecutionThread() const override;
-	virtual bool					IsAsync() const override;
-	bool							GetTileInvalidation(int32 RowId, int32 ColId) const;
+	UE_API virtual Device*					GetTargetDevice() const override;
+	UE_API virtual ENamedThreads::Type		GetExecutionThread() const override;
+	UE_API virtual bool					IsAsync() const override;
+	UE_API bool							GetTileInvalidation(int32 RowId, int32 ColId) const;
 
-	virtual void					GetDependencies(JobPtrVec& Prior, JobPtrVec& After, JobRunInfo RunInfo);
+	UE_API virtual void					GetDependencies(JobPtrVec& Prior, JobPtrVec& After, JobRunInfo RunInfo);
 
-	virtual FString					GetName() const override;
+	UE_API virtual FString					GetName() const override;
 
-	virtual bool					DebugCompleteCheck() override;
+	UE_API virtual bool					DebugCompleteCheck() override;
 
 	//////////////////////////////////////////////////////////////////////////
 	/// END: DeviceNativeTask implementation
 	//////////////////////////////////////////////////////////////////////////
-	virtual bool					CheckCulled(JobRunInfo RunInfo);
-	virtual FString					GetRunTimings(double BatchStartTime) const;
-	virtual FString					GetDebugName() const override;
+	UE_API virtual bool					CheckCulled(JobRunInfo RunInfo);
+	UE_API virtual FString					GetRunTimings(double BatchStartTime) const;
+	UE_API virtual FString					GetDebugName() const override;
 
 	//////////////////////////////////////////////////////////////////////////
 	/// Inline functions
@@ -215,3 +217,5 @@ public:
 };
 
 //////////////////////////////////////////////////////////////////////////
+
+#undef UE_API
diff --git a/Engine/Plugins/TextureGraph/Source/TextureGraphEngine/Job/JobArgs.h b/Engine/Plugins/TextureGraph/Source/TextureGraphEngine/Job/JobArgs.h
--- a/Engine/Plugins/TextureGraph/Source/TextureGraphEngine/Job/JobArgs.h
+++ b/Engine/Plugins/TextureGraph/Source/TextureGraphEngine/Job/JobArgs.h
@@ -9,6 +9,8 @@
 #include "Data/TiledBlob.h"
 #include "Profiling/StatGroup.h"
 
+#define UE_API TEXTUREGRAPHENGINE_API
+
 class Job;
 class JobBatch;
 class Scheduler;
@@ -16,7 +18,7 @@ class Scheduler;
 DECLARE_CYCLE_STAT(TEXT("JobArg_SimpleType_Bind"), STAT_JobArg_SimpleType_Bind, STATGROUP_TextureGraphEngine);
 DECLARE_CYCLE_STAT(TEXT("JobArg_SimpleType_UnBind"), STAT_JobArg_SimpleType_UnBind, STATGROUP_TextureGraphEngine);
 //////////////////////////////////////////////////////////////////////////
-struct TEXTUREGRAPHENGINE_API JobArgResult
+struct JobArgResult
 {
 	std::exception_ptr				ExInner;			/// Original exception that was raised by the action
 	int32							ErrorCode = 0;		/// What is the error code
@@ -26,7 +28,7 @@ typedef std::shared_ptr<JobArgResult> JobArgResultPtr;
 typedef cti::continuable<JobArgResultPtr> AsyncJobArgResultPtr;
 
 //////////////////////////////////////////////////////////////////////////
-struct TEXTUREGRAPHENGINE_API JobArgBindInfo
+struct JobArgBindInfo
 {
 	Job*							JobObj = nullptr;	/// The job that this is associated with
 	const JobBatch*					Batch = nullptr;	/// What is the batch that we're going to be part of (can be nullptr)
@@ -39,7 +41,7 @@ struct TEXTUREGRAPHENGINE_API JobArgBindInfo
 };
 
 //////////////////////////////////////////////////////////////////////////
-class TEXTUREGRAPHENGINE_API JobArg
+class JobArg
 {
 protected:
 	bool							bIgnoreDesc = false;/// Whether to ignore the descriptor for this argument
@@ -52,8 +54,8 @@ public:
 	virtual							~JobArg() {}
 	virtual AsyncJobArgResultPtr	Bind(JobArgBindInfo JobBindInfo) = 0;
 	virtual AsyncJobArgResultPtr	Unbind(JobArgBindInfo JobBindInfo) = 0;
-	virtual bool					CanHandleTiles() const;
-	virtual bool					ForceNonTiledTransform() const;
+	UE_API virtual bool					CanHandleTiles() const;
+	UE_API virtual bool					ForceNonTiledTransform() const;
 
 	virtual CHashPtr				Hash() const = 0;
 	virtual CHashPtr				TileHash(uint32 TileX, uint32 TileY) const = 0;
@@ -84,7 +86,7 @@ typedef std::shared_ptr<JobArg>		JobArgPtr;
 //////////////////////////////////////////////////////////////////////////	
 /// Generic CPU/GPU resource e.g. textures, samplers etc.
 //////////////////////////////////////////////////////////////////////////
-class TEXTUREGRAPHENGINE_API JobArg_Resource : public JobArg
+class JobArg_Resource : public JobArg
 {
 protected:
 	ResourceBindInfo				ArgBindInfo;			/// Binding information for this resource
@@ -95,7 +97,7 @@ public:
 };
 
 //////////////////////////////////////////////////////////////////////////
-class TEXTUREGRAPHENGINE_API JobArg_Blob : public JobArg_Resource
+class JobArg_Blob : public JobArg_Resource
 {
 protected:
 	mutable TiledBlobRef			BlobObjRef;							/// Blob Cached between Bind and Unbind calls only
@@ -110,40 +112,40 @@ protected:
 	bool							bBindArrayOfTiles = false;			/// Bind ALL the tiles contained in this arg, filling the array of UTextures before execution of the job.
 
 
-	TiledBlobPtr					GetRootBlob(JobArgBindInfo JobBindInfo) const;
+	UE_API TiledBlobPtr					GetRootBlob(JobArgBindInfo JobBindInfo) const;
 
 public:
-									JobArg_Blob(TiledBlobPtr blob, const ResourceBindInfo& BindInfo);
-									JobArg_Blob(TiledBlobPtr BlobObj, const char* TargetName);
+									UE_API JobArg_Blob(TiledBlobPtr blob, const ResourceBindInfo& BindInfo);
+									UE_API JobArg_Blob(TiledBlobPtr BlobObj, const char* TargetName);
 
-	virtual							~JobArg_Blob() override;
-	virtual void					SetHandleTiles(bool bInCanHandleTiles);
-	virtual bool					CanHandleTiles() const override;
-	JobArg_Blob&					WithNotHandleTiles(); // equivalent to call  SetHandleTiles(false)
+	UE_API virtual							~JobArg_Blob() override;
+	UE_API virtual void					SetHandleTiles(bool bInCanHandleTiles);
+	UE_API virtual bool					CanHandleTiles() const override;
+	UE_API JobArg_Blob&					WithNotHandleTiles(); // equivalent to call  SetHandleTiles(false)
 
 
-	virtual void					SetForceNonTiledTransform(bool bInForceNonTiledTransform);
-	virtual bool					ForceNonTiledTransform() const override;
+	UE_API virtual void					SetForceNonTiledTransform(bool bInForceNonTiledTransform);
+	UE_API virtual bool					ForceNonTiledTransform() const override;
 
-	JobArg_Blob&					WithDownsampled4To1();
-	bool							IsDownsampled4To1() const;
+	UE_API JobArg_Blob&					WithDownsampled4To1();
+	UE_API bool							IsDownsampled4To1() const;
 
-	JobArg_Blob&					WithNeighborTiles();
-	bool							IsNeighborTiles() const;
+	UE_API JobArg_Blob&					WithNeighborTiles();
+	UE_API bool							IsNeighborTiles() const;
 
-	JobArg_Blob&					WithArrayOfTiles();
-	bool							IsArrayOfTiles() const;
+	UE_API JobArg_Blob&					WithArrayOfTiles();
+	UE_API bool							IsArrayOfTiles() const;
 
-	virtual AsyncJobArgResultPtr	Bind(JobArgBindInfo JobBindInfo) override;
-	virtual AsyncJobArgResultPtr	Unbind(JobArgBindInfo JobBindInfo) override;
+	UE_API virtual AsyncJobArgResultPtr	Bind(JobArgBindInfo JobBindInfo) override;
+	UE_API virtual AsyncJobArgResultPtr	Unbind(JobArgBindInfo JobBindInfo) override;
 
-	virtual CHashPtr				Hash() const override;
-	virtual CHashPtr				TileHash(uint32 TileX, uint32 TileY) const override;
+	UE_API virtual CHashPtr				Hash() const override;
+	UE_API virtual CHashPtr				TileHash(uint32 TileX, uint32 TileY) const override;
 
-	virtual bool					IsLateBound(uint32 TileX, uint32 TileY) const override;
+	UE_API virtual bool					IsLateBound(uint32 TileX, uint32 TileY) const override;
 
-	virtual JobPtrW					GeneratingJob() const override;
-	virtual const BufferDescriptor* GetDescriptor() const override;
+	UE_API virtual JobPtrW					GeneratingJob() const override;
+	UE_API virtual const BufferDescriptor* GetDescriptor() const override;
 
 	FORCEINLINE FString				Target() const { return ArgBindInfo.Target; }
 	FORCEINLINE TiledBlobRef		GetBlob() const { return BlobObjRef; }
@@ -209,7 +211,7 @@ public:
 
 //////////////////////////////////////////////////////////////////////////
 
-class TEXTUREGRAPHENGINE_API JobArg_String : public JobArg_SimpleType<FString>
+class JobArg_String : public JobArg_SimpleType<FString>
 {
 public:
 									JobArg_String(FString InValue, const ResourceBindInfo& BindInfo) : JobArg_SimpleType<FString>(InValue, BindInfo) {}
@@ -226,7 +228,7 @@ public:
 
 //////////////////////////////////////////////////////////////////////////
 class RenderMesh;
-class TEXTUREGRAPHENGINE_API JobArg_Mesh : public JobArg
+class JobArg_Mesh : public JobArg
 {
 private:
 	std::shared_ptr<RenderMesh>		Mesh;			/// The mesh that we're binding
@@ -235,15 +237,15 @@ private:
 public:
 									JobArg_Mesh(std::shared_ptr<RenderMesh> mesh, int32 targetId) : Mesh(mesh), _targetId(targetId) { check(Mesh); }
 
-	virtual AsyncJobArgResultPtr	Bind(JobArgBindInfo JobBindInfo) override;
-	virtual AsyncJobArgResultPtr	Unbind(JobArgBindInfo JobBindInfo) override;
-	virtual CHashPtr				Hash() const override;
-	virtual CHashPtr				TileHash(uint32 TileX, uint32 TileY) const override;
+	UE_API virtual AsyncJobArgResultPtr	Bind(JobArgBindInfo JobBindInfo) override;
+	UE_API virtual AsyncJobArgResultPtr	Unbind(JobArgBindInfo JobBindInfo) override;
+	UE_API virtual CHashPtr				Hash() const override;
+	UE_API virtual CHashPtr				TileHash(uint32 TileX, uint32 TileY) const override;
 };
 
 //////////////////////////////////////////////////////////////////////////
 template <typename Type>
-class TEXTUREGRAPHENGINE_API JobArg_Array : public JobArg_Resource
+class JobArg_Array : public JobArg_Resource
 {
 protected:
 	TArray<Type>					Value;				/// The underlying InValue to bind
@@ -287,7 +289,7 @@ public:
 //////////////////////////////////////////////////////////////////////////
 
 template <typename StructType>
-class TEXTUREGRAPHENGINE_API JobArg_Struct: public JobArg_Resource
+class JobArg_Struct: public JobArg_Resource
 {
 
 protected:
@@ -330,20 +332,20 @@ public:
 	}
 };
 
-class TEXTUREGRAPHENGINE_API JobArg_TileInfo : public JobArg_Struct<FTileInfo>
+class JobArg_TileInfo : public JobArg_Struct<FTileInfo>
 {
 
 public:
 	JobArg_TileInfo(FTileInfo InValue, const ResourceBindInfo& BindInfo) : JobArg_Struct<FTileInfo>(InValue, BindInfo) {}
 	JobArg_TileInfo(FTileInfo InValue, const char* TargetName) : JobArg_Struct<FTileInfo>(InValue, TargetName) {}
 
-	virtual AsyncJobArgResultPtr	Bind(JobArgBindInfo JobBindInfo) override;
+	UE_API virtual AsyncJobArgResultPtr	Bind(JobArgBindInfo JobBindInfo) override;
 
-	virtual AsyncJobArgResultPtr	Unbind(JobArgBindInfo JobBindInfo) override;
+	UE_API virtual AsyncJobArgResultPtr	Unbind(JobArgBindInfo JobBindInfo) override;
 
-	virtual CHashPtr				Hash() const override;
+	UE_API virtual CHashPtr				Hash() const override;
 
-	virtual CHashPtr				TileHash(uint32 TileX, uint32 TileY) const override;
+	UE_API virtual CHashPtr				TileHash(uint32 TileX, uint32 TileY) const override;
 };
 
 //////////////////////////////////////////////////////////////////////////
@@ -355,7 +357,7 @@ typedef JobArg_SimpleType<FMatrix>		JobArg_Matrix;
 
 //////////////////////////////////////////////////////////////////////////
 
-class TEXTUREGRAPHENGINE_API JobArg_ForceTiling : public JobArg
+class JobArg_ForceTiling : public JobArg
 {
 protected:
 public:
@@ -369,8 +371,8 @@ public:
 		return cti::make_ready_continuable(std::make_shared<JobArgResult>());
 	};
 
-	virtual CHashPtr				Hash() const override;
-	virtual CHashPtr				TileHash(uint32 TileX, uint32 TileY) const override;
+	UE_API virtual CHashPtr				Hash() const override;
+	UE_API virtual CHashPtr				TileHash(uint32 TileX, uint32 TileY) const override;
 };
 //////////////////////////////////////////////////////////////////////////
 
@@ -417,3 +419,5 @@ FORCEINLINE JobArgPtr				WithUnbounded(JobArgPtr Arg, bool bUnbounded = true)
 	Arg->WithUnbounded(bUnbounded);
 	return Arg;
 }
+
+#undef UE_API
diff --git a/Engine/Plugins/TextureGraph/Source/TextureGraphEngine/Job/JobBatch.h b/Engine/Plugins/TextureGraph/Source/TextureGraphEngine/Job/JobBatch.h
--- a/Engine/Plugins/TextureGraph/Source/TextureGraphEngine/Job/JobBatch.h
+++ b/Engine/Plugins/TextureGraph/Source/TextureGraphEngine/Job/JobBatch.h
@@ -7,6 +7,8 @@
 #include "Device/Device.h"
 #include "JobCommon.h"
 
+#define UE_API TEXTUREGRAPHENGINE_API
+
 class										MixUpdateCycle;
 class										Job;
 typedef std::unique_ptr<Job>				JobUPtr;
@@ -26,10 +28,10 @@ DECLARE_LOG_CATEGORY_EXTERN(LogBatch, Log, All);
 /// JobBatch: A series of jobs put together that need to execute at the 
 /// same time
 //////////////////////////////////////////////////////////////////////////
-class TEXTUREGRAPHENGINE_API JobBatch
+class JobBatch
 {
 public:
-	static uint64					GBatchId;				/// For easier debugging
+	static UE_API uint64					GBatchId;				/// For easier debugging
 protected:
 
 	mutable FCriticalSection		Mutex;						/// Mutex for the JobObj queue
@@ -56,40 +58,40 @@ protected:
 	bool							bIsFinished = false;		/// Whether the batch has been finished or not
 	bool							bIsAsync = true;			/// Whether the job batch is asynchronous or not
 
-	AsyncJobResultPtr				BeginQueue_Native(JobPriorityQueue& InQueue);
+	UE_API AsyncJobResultPtr				BeginQueue_Native(JobPriorityQueue& InQueue);
 
-	void							EndBatch();	
+	UE_API void							EndBatch();	
 
 	friend class Job;
-	void							OnJobDone(Job* JobObj, int64 JobId);	/// Once a JobObj is TRULY done then it reports to the batch,
-	void							AddNativeJob_Now(JobPtr JobObj, bool AddToAllJobs);
+	UE_API void							OnJobDone(Job* JobObj, int64 JobId);	/// Once a JobObj is TRULY done then it reports to the batch,
+	UE_API void							AddNativeJob_Now(JobPtr JobObj, bool AddToAllJobs);
 
 	std::atomic<uint32>				NumJobsRunning = 0;		/// A count down of the jobs still running, decremented by OnJobDone
 	OnAllJobsDoneCallbackVec		OnAllJobsDoneCallbacks;		/// If the call back has been assigned, that 's what will get called when all the jobs are done.
 
 	std::vector<bool>				JobsFinished;				/// Jobs that have been finished
 
-	void							PrintJobTimings() const;
+	UE_API void							PrintJobTimings() const;
 
 public:
-	explicit						JobBatch(const FInvalidationDetails& Details);
-	virtual							~JobBatch();
+	UE_API explicit						JobBatch(const FInvalidationDetails& Details);
+	UE_API virtual							~JobBatch();
 
-	virtual AsyncJobResultPtr		Exec(OnAllJobsDoneCallback Callback = nullptr);
+	UE_API virtual AsyncJobResultPtr		Exec(OnAllJobsDoneCallback Callback = nullptr);
 	
-	virtual JobPtrW					AddJob(JobUPtr JobObj);
-	virtual void					RemoveJob(JobPtrW JobObj);
+	UE_API virtual JobPtrW					AddJob(JobUPtr JobObj);
+	UE_API virtual void					RemoveJob(JobPtrW JobObj);
 
-	bool							IsLocked() const;
+	UE_API bool							IsLocked() const;
 
-	void							ResetForReplay(int32 JobId = -1); /// For debug purpose we can reset the state of a batch for a replay
+	UE_API void							ResetForReplay(int32 JobId = -1); /// For debug purpose we can reset the state of a batch for a replay
 																	  /// A JobObj index can be specified to only reset that JobObj during replay
 
-	void							SetCaptureRenderDoc(bool CaptureRenderDoc = true);
-	void							DebugDumpUnfinishedJobs();
+	UE_API void							SetCaptureRenderDoc(bool CaptureRenderDoc = true);
+	UE_API void							DebugDumpUnfinishedJobs();
 
-	void							Terminate();
-	void							OnDone(OnAllJobsDoneCallback Callback);
+	UE_API void							Terminate();
+	UE_API void							OnDone(OnAllJobsDoneCallback Callback);
 
 	//////////////////////////////////////////////////////////////////////////
 	/// Inline functions
@@ -118,7 +120,8 @@ public:
 	/// Static functions
 	//////////////////////////////////////////////////////////////////////////
 	static FORCEINLINE uint64		CurrentBatchId() { return GBatchId; }
-	static std::shared_ptr<JobBatch> Create(const FInvalidationDetails& details);
+	static UE_API std::shared_ptr<JobBatch> Create(const FInvalidationDetails& details);
 };
 
 
+#undef UE_API
diff --git a/Engine/Plugins/TextureGraph/Source/TextureGraphEngine/Job/JobCommon.h b/Engine/Plugins/TextureGraph/Source/TextureGraphEngine/Job/JobCommon.h
--- a/Engine/Plugins/TextureGraph/Source/TextureGraphEngine/Job/JobCommon.h
+++ b/Engine/Plugins/TextureGraph/Source/TextureGraphEngine/Job/JobCommon.h
@@ -5,6 +5,8 @@
 #include "Helper/DataUtil.h"
 #include "Data/Blobber.h"
 
+#define UE_API TEXTUREGRAPHENGINE_API
+
 //////////////////////////////////////////////////////////////////////////
 class Device;
 class Scheduler;
@@ -23,7 +25,7 @@ typedef std::shared_ptr<BlobTransform> BlobTransformPtr;
 
 class Job;
 
-struct TEXTUREGRAPHENGINE_API JobRunInfo
+struct JobRunInfo
 {
 	Scheduler*						JobScheduler = nullptr;	/// The scheduler instance that's scheduled this job
 	Device*							Dev = nullptr;			/// The device on which the job must be executed
@@ -33,9 +35,9 @@ struct TEXTUREGRAPHENGINE_API JobRunInfo
 };
 
 //////////////////////////////////////////////////////////////////////////
-struct TEXTUREGRAPHENGINE_API JobResult
+struct JobResult
 {
-	static std::shared_ptr<JobResult> NullResult;		/// Null result
+	static UE_API std::shared_ptr<JobResult> NullResult;		/// Null result
 
 	std::exception_ptr				ExInner;			/// Original exception that was raised by the action
 	int32							ErrorCode = 0;		/// What is the error code
@@ -45,7 +47,7 @@ struct TEXTUREGRAPHENGINE_API JobResult
 
 	JobResult() : BlobObj(BlobRef()) {}
 	JobResult(BlobRef NewBlobObj, CHashPtr NewHashValue) : BlobObj(NewBlobObj), HashValue(NewHashValue) {}
-	~JobResult();
+	UE_API ~JobResult();
 };
 
 typedef std::shared_ptr<JobResult>	JobResultPtr;
@@ -57,3 +59,4 @@ typedef std::vector<JobPtr>			JobPtrVec;
 typedef TArray<JobPtr>				JobPtrArray;
 typedef TArray<Job*>				JobPArray;
 
+#undef UE_API
diff --git a/Engine/Plugins/TextureGraph/Source/TextureGraphEngine/Job/MinMaxService.h b/Engine/Plugins/TextureGraph/Source/TextureGraphEngine/Job/MinMaxService.h
--- a/Engine/Plugins/TextureGraph/Source/TextureGraphEngine/Job/MinMaxService.h
+++ b/Engine/Plugins/TextureGraph/Source/TextureGraphEngine/Job/MinMaxService.h
@@ -7,21 +7,24 @@
 #include <vector>
 #include <unordered_map>
 
+#define UE_API TEXTUREGRAPHENGINE_API
+
 class UModelObject;
 typedef std::shared_ptr<class Job>		JobPtr;
 typedef std::weak_ptr<class Job>		JobPtrW;
 
 class UMixInterface;
 
-class TEXTUREGRAPHENGINE_API MinMaxService : public BlobHelperService
+class MinMaxService : public BlobHelperService
 {
 public:
-									MinMaxService();
-	virtual							~MinMaxService() override;
+									UE_API MinMaxService();
+	UE_API virtual							~MinMaxService() override;
 
-	virtual AsyncJobResultPtr		Tick() override;
+	UE_API virtual AsyncJobResultPtr		Tick() override;
 };
 
 typedef std::shared_ptr<MinMaxService>	MinMaxServicePtr;
 typedef std::weak_ptr<MinMaxService>	MinMaxServicePtrW;
 
+#undef UE_API
diff --git a/Engine/Plugins/TextureGraph/Source/TextureGraphEngine/Job/MipMapService.h b/Engine/Plugins/TextureGraph/Source/TextureGraphEngine/Job/MipMapService.h
--- a/Engine/Plugins/TextureGraph/Source/TextureGraphEngine/Job/MipMapService.h
+++ b/Engine/Plugins/TextureGraph/Source/TextureGraphEngine/Job/MipMapService.h
@@ -7,21 +7,24 @@
 #include <vector>
 #include <unordered_map>
 
+#define UE_API TEXTUREGRAPHENGINE_API
+
 class UModelObject;
 typedef std::shared_ptr<class Job>		JobPtr;
 typedef std::weak_ptr<class Job>		JobPtrW;
 
 class UMixInterface;
 
-class TEXTUREGRAPHENGINE_API MipMapService : public BlobHelperService
+class MipMapService : public BlobHelperService
 {
 public:
-									MipMapService();
-	virtual							~MipMapService() override;
+									UE_API MipMapService();
+	UE_API virtual							~MipMapService() override;
 
-	virtual AsyncJobResultPtr		Tick() override;
+	UE_API virtual AsyncJobResultPtr		Tick() override;
 };
 
 typedef std::shared_ptr<MipMapService>	MipMapServicePtr;
 typedef std::weak_ptr<MipMapService>	MipMapServicePtrW;
 
+#undef UE_API
diff --git a/Engine/Plugins/TextureGraph/Source/TextureGraphEngine/Job/Scheduler.h b/Engine/Plugins/TextureGraph/Source/TextureGraphEngine/Job/Scheduler.h
--- a/Engine/Plugins/TextureGraph/Source/TextureGraphEngine/Job/Scheduler.h
+++ b/Engine/Plugins/TextureGraph/Source/TextureGraphEngine/Job/Scheduler.h
@@ -14,11 +14,13 @@
 #include <list>
 #include "Profiling/StatGroup.h"
 
+#define UE_API TEXTUREGRAPHENGINE_API
+
 class JobBatch;
 typedef std::shared_ptr<JobBatch>		JobBatchPtr;
 typedef std::weak_ptr<JobBatch>			JobBatchPtrW;
 
-class TEXTUREGRAPHENGINE_API SchedulerObserverSource
+class SchedulerObserverSource
 {
 protected:
 	/// Protected interface of emitters called by the scheduler to notify the observers
@@ -39,13 +41,13 @@ public:
 };
 typedef std::shared_ptr<SchedulerObserverSource> SchedulerObserverSourcePtr;
 
-class TEXTUREGRAPHENGINE_API Scheduler
+class Scheduler
 {
 private:
-	static const double					IdleTimeInterval;				/// Start an Idle Batch after <this many> milliseconds
-	static const double					IdleBusyTimeInterval;			/// Start an Idle Batch (during busy times) after <this many> milliseconds
-	static const double					IdleBatchTimeLimit;				/// Start an Idle Batch after <this many> milliseconds
-	static const double					CurrentBatchWarningLimit;		/// Scheduler starts spewing out warnings if the current batch doesn't finish ater <this many> milliseconds
+	static UE_API const double					IdleTimeInterval;				/// Start an Idle Batch after <this many> milliseconds
+	static UE_API const double					IdleBusyTimeInterval;			/// Start an Idle Batch (during busy times) after <this many> milliseconds
+	static UE_API const double					IdleBatchTimeLimit;				/// Start an Idle Batch after <this many> milliseconds
+	static UE_API const double					CurrentBatchWarningLimit;		/// Scheduler starts spewing out warnings if the current batch doesn't finish ater <this many> milliseconds
 
 	bool								bIsRunning = false;				/// Whether the scheduler is running or not
 
@@ -90,26 +92,26 @@ private:
 
 	uint64								BatchIndex = 0;					/// Index of the batches
 
-	void								Start();
-	void								Stop();
-	void								StopServices();
+	UE_API void								Start();
+	UE_API void								Stop();
+	UE_API void								StopServices();
 
-	AsyncJobResultPtr					UpdateIdleBatch(size_t index, bool bIsBusyUpdate);
-	void								UpdateIdle(bool bForce, bool bIsBusyUpdate = false);
+	UE_API AsyncJobResultPtr					UpdateIdleBatch(size_t index, bool bIsBusyUpdate);
+	UE_API void								UpdateIdle(bool bForce, bool bIsBusyUpdate = false);
 
 public:
-										Scheduler();
-										~Scheduler();
+										UE_API Scheduler();
+										UE_API ~Scheduler();
 
-	void								Update(float Delta);
-	void								AddBatch(JobBatchPtr Batch);
-	void								AddIdleService(IdleServicePtr batch);
+	UE_API void								Update(float Delta);
+	UE_API void								AddBatch(JobBatchPtr Batch);
+	UE_API void								AddIdleService(IdleServicePtr batch);
 	
-	void								CaptureRenderDocLastRunBatch();
+	UE_API void								CaptureRenderDocLastRunBatch();
 
-	void								RegisterObserverSource(const SchedulerObserverSourcePtr& observerSource); ///
-	void								SetCaptureRenderDocNextBatch(bool capture = true);
-	void								ClearCache();
+	UE_API void								RegisterObserverSource(const SchedulerObserverSourcePtr& observerSource); ///
+	UE_API void								SetCaptureRenderDocNextBatch(bool capture = true);
+	UE_API void								ClearCache();
 
 	//////////////////////////////////////////////////////////////////////////
 	/// Inline functions 
@@ -141,3 +143,4 @@ public:
 
 typedef std::unique_ptr<Scheduler>		SchedulerPtr;
 
+#undef UE_API
diff --git a/Engine/Plugins/TextureGraph/Source/TextureGraphEngine/Job/TempHashService.h b/Engine/Plugins/TextureGraph/Source/TextureGraphEngine/Job/TempHashService.h
--- a/Engine/Plugins/TextureGraph/Source/TextureGraphEngine/Job/TempHashService.h
+++ b/Engine/Plugins/TextureGraph/Source/TextureGraphEngine/Job/TempHashService.h
@@ -5,22 +5,25 @@
 #include "IdleService.h"
 #include <vector>
 
-class TEXTUREGRAPHENGINE_API TempHashService : public IdleService
+#define UE_API TEXTUREGRAPHENGINE_API
+
+class TempHashService : public IdleService
 {
 protected:
 	FCriticalSection				HashesMutex;			/// Mutex for the hashes
 	std::vector<CHashPtr>			Hashes;					/// The hashes that we're waiting to resolve
 
 public:
-									TempHashService();
-	virtual							~TempHashService() override;
+									UE_API TempHashService();
+	UE_API virtual							~TempHashService() override;
 
-	virtual AsyncJobResultPtr		Tick() override;
-	virtual void					Stop() override;
+	UE_API virtual AsyncJobResultPtr		Tick() override;
+	UE_API virtual void					Stop() override;
 
-	void							Add(CHashPtr HashValue);
+	UE_API void							Add(CHashPtr HashValue);
 };
 
 typedef std::shared_ptr<TempHashService>	TempHashServicePtr;
 typedef std::weak_ptr<TempHashService>		TempHashServicePtrW;
 
+#undef UE_API
diff --git a/Engine/Plugins/TextureGraph/Source/TextureGraphEngine/Job/ThumbnailsService.h b/Engine/Plugins/TextureGraph/Source/TextureGraphEngine/Job/ThumbnailsService.h
--- a/Engine/Plugins/TextureGraph/Source/TextureGraphEngine/Job/ThumbnailsService.h
+++ b/Engine/Plugins/TextureGraph/Source/TextureGraphEngine/Job/ThumbnailsService.h
@@ -7,6 +7,8 @@
 #include <vector>
 #include <unordered_map>
 
+#define UE_API TEXTUREGRAPHENGINE_API
+
 class UModelObject;
 typedef std::shared_ptr<class Job>		JobPtr;
 typedef std::unique_ptr<class Job>		JobUPtr;
@@ -14,7 +16,7 @@ typedef std::weak_ptr<class Job>		JobPtrW;
 
 class UMixInterface;
 
-class TEXTUREGRAPHENGINE_API ThumbnailsService : public IdleService
+class ThumbnailsService : public IdleService
 {
 protected:
 	JobBatchPtr						Batch;					/// The thumbnail rendering batch
@@ -29,18 +31,18 @@ protected:
 
 	std::unordered_map<UObject*, AddedJobs> Handled;		/// The components that we have already added jobs for
 
-	JobBatchPtr						CreateNewUpdateCycle(UMixInterface* Mix);
-	JobBatchPtr						GetNextUpdateCycle();
+	UE_API JobBatchPtr						CreateNewUpdateCycle(UMixInterface* Mix);
+	UE_API JobBatchPtr						GetNextUpdateCycle();
 
 public:
 	DECLARE_MULTICAST_DELEGATE_OneParam(FOnThumbnailBatchDoneDelegate, JobBatchPtr Batch)
 	FOnThumbnailBatchDoneDelegate	OnUpdateThumbnailDelegate;
-									ThumbnailsService();
-	virtual							~ThumbnailsService() override;
+									UE_API ThumbnailsService();
+	UE_API virtual							~ThumbnailsService() override;
 
-	virtual AsyncJobResultPtr		Tick() override;
-	virtual void					Stop() override;
-	void							AddUniqueJobToCycle(UObject* CreatorComponent, UMixInterface* Mix, JobUPtr ThumbnailJob, int32 TargetId, bool bForceExec = false);	//Use this to always generate thumbnail from last update cycle
+	UE_API virtual AsyncJobResultPtr		Tick() override;
+	UE_API virtual void					Stop() override;
+	UE_API void							AddUniqueJobToCycle(UObject* CreatorComponent, UMixInterface* Mix, JobUPtr ThumbnailJob, int32 TargetId, bool bForceExec = false);	//Use this to always generate thumbnail from last update cycle
 
 	//////////////////////////////////////////////////////////////////////////
 	/// Inline functions
@@ -51,3 +53,4 @@ public:
 typedef std::shared_ptr<ThumbnailsService>	ThumbnailsServicePtr;
 typedef std::weak_ptr<ThumbnailsService>	ThumbnailsServicePtrW;
 
+#undef UE_API
diff --git a/Engine/Plugins/TextureGraph/Source/TextureGraphEngine/Model/Mix/Mix.h b/Engine/Plugins/TextureGraph/Source/TextureGraphEngine/Model/Mix/Mix.h
--- a/Engine/Plugins/TextureGraph/Source/TextureGraphEngine/Model/Mix/Mix.h
+++ b/Engine/Plugins/TextureGraph/Source/TextureGraphEngine/Model/Mix/Mix.h
@@ -9,6 +9,8 @@
 
 #include "Mix.generated.h"
 
+#define UE_API TEXTUREGRAPHENGINE_API
+
 DECLARE_LOG_CATEGORY_EXTERN(LogMix, All, All);
 
 class UMixParameters;
@@ -16,8 +18,8 @@ class UMixParameters;
 class RenderMesh;
 typedef std::shared_ptr<RenderMesh>	RenderMeshPtr;
 
-UCLASS(Blueprintable, BlueprintType)
-class TEXTUREGRAPHENGINE_API UMix : public UMixInterface
+UCLASS(MinimalAPI, Blueprintable, BlueprintType)
+class UMix : public UMixInterface
 {
 	GENERATED_BODY()
 
@@ -25,15 +27,17 @@ public:
 	
 
 private:
-	static UMix*						GNullMix;					/// This is a generic mix that is used in various places and rendering actions that
+	static UE_API UMix*						GNullMix;					/// This is a generic mix that is used in various places and rendering actions that
 																	/// often require a Mix object. This is initialised in the first call to the public 
 																	/// static function below. 
-	static const FString				GRootSectionName;
+	static UE_API const FString				GRootSectionName;
 
 public:
-	static UMix*						NullMix();
+	static UE_API UMix*						NullMix();
 
 public:
-	virtual								~UMix() override;
-	virtual void						Update(MixUpdateCyclePtr Cycle) override;
+	UE_API virtual								~UMix() override;
+	UE_API virtual void						Update(MixUpdateCyclePtr Cycle) override;
 };
+
+#undef UE_API
diff --git a/Engine/Plugins/TextureGraph/Source/TextureGraphEngine/Model/Mix/MixInterface.h b/Engine/Plugins/TextureGraph/Source/TextureGraphEngine/Model/Mix/MixInterface.h
--- a/Engine/Plugins/TextureGraph/Source/TextureGraphEngine/Model/Mix/MixInterface.h
+++ b/Engine/Plugins/TextureGraph/Source/TextureGraphEngine/Model/Mix/MixInterface.h
@@ -7,6 +7,8 @@
 
 #include "MixInterface.generated.h"
 
+#define UE_API TEXTUREGRAPHENGINE_API
+
 class UMix;
 class ULayerStack;
 class UMixSettings;
@@ -23,20 +25,20 @@ typedef std::shared_ptr<MeshAsset>		MeshAssetPtr;
 /// To support this functionality, we introduced the new class named MixInterface. 
 /// This will be the base class for both Mix and MixInstance.
 //////////////////////////////////////////////////////////////////////////
-UCLASS(Abstract)
-class TEXTUREGRAPHENGINE_API UMixInterface : public UModelObject
+UCLASS(MinimalAPI, Abstract)
+class UMixInterface : public UModelObject
 {
 	GENERATED_BODY()
 
 private:
-	static const TMap<FString, FString> s_uriAlias;						/// Contains the list of URI aliases map to keep the data meaningful, short and precise.
-	static const TMap<FString, FString> s_uriAliasDecryptor;			/// We need alias decryptor to convert it so we can extract the object from the data. 
+	static UE_API const TMap<FString, FString> s_uriAlias;						/// Contains the list of URI aliases map to keep the data meaningful, short and precise.
+	static UE_API const TMap<FString, FString> s_uriAliasDecryptor;			/// We need alias decryptor to convert it so we can extract the object from the data. 
 
-	static TMap<FString, FString>		InitURIAlias();
-	static TMap<FString, FString>		InitURIAliasDecryptor();
+	static UE_API TMap<FString, FString>		InitURIAlias();
+	static UE_API TMap<FString, FString>		InitURIAliasDecryptor();
 
-	FString								GetObjectAlias(FString object);
-	FString								GetAliasToObject(FString Alias);
+	UE_API FString								GetObjectAlias(FString object);
+	UE_API FString								GetAliasToObject(FString Alias);
 
 protected:
 	UPROPERTY()
@@ -54,10 +56,10 @@ protected:
 	bool								bEnableLOD = true;	/// Enable LOD-ing in the system
 
 	virtual UModelObject*				RootURI() { checkNoEntry() return nullptr; };
-	virtual void						Invalidate(FModelInvalidateInfo InvalidateInfo);
+	UE_API virtual void						Invalidate(FModelInvalidateInfo InvalidateInfo);
 
 public:
-	virtual								~UMixInterface() override;
+	UE_API virtual								~UMixInterface() override;
 
 
 
@@ -66,36 +68,36 @@ public:
 
 	DECLARE_DELEGATE_TwoParams(FOnRenderDone, UMixInterface*, const FInvalidationDetails*);
 	FOnRenderDone						OnRenderDone;
-	virtual void						PostMeshLoad();
+	UE_API virtual void						PostMeshLoad();
 
 	virtual bool						CanEdit() { return false; }
 
-	virtual void						SetMesh(RenderMeshPtr MeshObj, int MeshType, FVector Scale, FVector2D Dimension);
+	UE_API virtual void						SetMesh(RenderMeshPtr MeshObj, int MeshType, FVector Scale, FVector2D Dimension);
 
 #if WITH_EDITOR
-	virtual AsyncActionResultPtr		SetEditorMesh(AActor* Actor);
-	virtual AsyncActionResultPtr		SetEditorMesh(UStaticMeshComponent* MeshComponent, UWorld* World);
+	UE_API virtual AsyncActionResultPtr		SetEditorMesh(AActor* Actor);
+	UE_API virtual AsyncActionResultPtr		SetEditorMesh(UStaticMeshComponent* MeshComponent, UWorld* World);
 #endif
 	
-	virtual void						Update(MixUpdateCyclePtr cycle);	
+	UE_API virtual void						Update(MixUpdateCyclePtr cycle);	
 
-	virtual RenderMeshPtr				GetMesh() const; // Get the mesh assigned in Settings
+	UE_API virtual RenderMeshPtr				GetMesh() const; // Get the mesh assigned in Settings
 
-	virtual int32						Width() const; 
-	virtual int32						Height() const;
+	UE_API virtual int32						Width() const; 
+	UE_API virtual int32						Height() const;
 	
-	virtual int32						GetNumXTiles() const;
-	virtual int32						GetNumYTiles()const;
+	UE_API virtual int32						GetNumXTiles() const;
+	UE_API virtual int32						GetNumYTiles()const;
 	
-	virtual UMixSettings*				GetSettings() const;
+	UE_API virtual UMixSettings*				GetSettings() const;
 
-	bool								IsHigherPriorityThan(const UMixInterface* RHS) const;
+	UE_API bool								IsHigherPriorityThan(const UMixInterface* RHS) const;
 
-	virtual void						InvalidateWithDetails(const FInvalidationDetails& Details);
-	virtual void						InvalidateAll();
+	UE_API virtual void						InvalidateWithDetails(const FInvalidationDetails& Details);
+	UE_API virtual void						InvalidateAll();
 
-	virtual void						BroadcastOnRenderingDone(const FInvalidationDetails* Details);
-	virtual void						BroadcastOnBatchQueued(const FInvalidationDetails* Details);
+	UE_API virtual void						BroadcastOnRenderingDone(const FInvalidationDetails* Details);
+	UE_API virtual void						BroadcastOnBatchQueued(const FInvalidationDetails* Details);
 	
 	//////////////////////////////////////////////////////////////////////////
 	/// Inline functions
@@ -107,3 +109,4 @@ public:
 };
 
 
+#undef UE_API
diff --git a/Engine/Plugins/TextureGraph/Source/TextureGraphEngine/Model/Mix/MixManager.h b/Engine/Plugins/TextureGraph/Source/TextureGraphEngine/Model/Mix/MixManager.h
--- a/Engine/Plugins/TextureGraph/Source/TextureGraphEngine/Model/Mix/MixManager.h
+++ b/Engine/Plugins/TextureGraph/Source/TextureGraphEngine/Model/Mix/MixManager.h
@@ -6,6 +6,8 @@
 #include "Helper/Util.h"
 #include "Model/ModelObject.h"
 
+#define UE_API TEXTUREGRAPHENGINE_API
+
 class UMixInterface;
 class UActionManager;
 
@@ -36,7 +38,7 @@ typedef GenericTaskPriorityQueue<MixInvalidateInfo, MixInterface_ComparePriority
 /// in the correct priority with actively visible mixes getting top priority
 /// compared with background mixes, that might still be updating. 
 //////////////////////////////////////////////////////////////////////////
-class TEXTUREGRAPHENGINE_API MixManager
+class MixManager
 {
 private:
 	std::atomic_bool				bIsSuspended = false;			/// Whether the mix manager is suspended or not
@@ -44,16 +46,16 @@ private:
 	std::vector<UMixInterface*>		MixesToFlush;						/// The mixes to flush
 	
 public:
-									MixManager();
-	virtual							~MixManager();
+									UE_API MixManager();
+	UE_API virtual							~MixManager();
 
-	void							Update(float Delta);
-	void							Suspend();
-	void							Resume();
-	void							Exit();
+	UE_API void							Update(float Delta);
+	UE_API void							Suspend();
+	UE_API void							Resume();
+	UE_API void							Exit();
 
-	void							InvalidateMix(UMixInterface* MixObj, const FInvalidationDetails &Details, int32 Priority = (int32)E_Priority::kNormal);
-	void							FlushMix(UMixInterface* MixObj);
+	UE_API void							InvalidateMix(UMixInterface* MixObj, const FInvalidationDetails &Details, int32 Priority = (int32)E_Priority::kNormal);
+	UE_API void							FlushMix(UMixInterface* MixObj);
 	
 	//////////////////////////////////////////////////////////////////////////
 	/// Inline functions
@@ -62,3 +64,4 @@ public:
 
 typedef std::unique_ptr<MixManager>		MixManagerUPtr;
 
+#undef UE_API
diff --git a/Engine/Plugins/TextureGraph/Source/TextureGraphEngine/Model/Mix/MixSettings.h b/Engine/Plugins/TextureGraph/Source/TextureGraphEngine/Model/Mix/MixSettings.h
--- a/Engine/Plugins/TextureGraph/Source/TextureGraphEngine/Model/Mix/MixSettings.h
+++ b/Engine/Plugins/TextureGraph/Source/TextureGraphEngine/Model/Mix/MixSettings.h
@@ -8,6 +8,8 @@
 #include "Helper/Promise.h"
 #include "MixSettings.generated.h"
 
+#define UE_API TEXTUREGRAPHENGINE_API
+
 UENUM()
 enum class PBRWorkflowMode
 {
@@ -105,8 +107,8 @@ class UStaticMesh;
 /// mix. These can be specified on a per mix basis. However, they can 
 /// be over-ridden on a per mix instance basis as well
 //////////////////////////////////////////////////////////////////////////
-UCLASS(Blueprintable, BlueprintType)
-class TEXTUREGRAPHENGINE_API UMixSettings : public UObject
+UCLASS(MinimalAPI, Blueprintable, BlueprintType)
+class UMixSettings : public UObject
 {
 	GENERATED_BODY()
 
@@ -174,32 +176,32 @@ private:
 
 	template<typename RenderMeshClass>
 	void							SetMeshInternal(RenderMeshPtr mesh, int meshType, FVector scale, FVector2D dimension);
-	void							SetMesh(RenderMeshPtr mesh);
+	UE_API void							SetMesh(RenderMeshPtr mesh);
 	
 protected:
 
 #if WITH_EDITOR
-	virtual void					PostEditChangeProperty(FPropertyChangedEvent& PropertyChangedEvent) override;
-	virtual void					PostEditChangeChainProperty(FPropertyChangedChainEvent& PropertyChangedEvent) override; // For TArray type properties in the component.
+	UE_API virtual void					PostEditChangeProperty(FPropertyChangedEvent& PropertyChangedEvent) override;
+	UE_API virtual void					PostEditChangeChainProperty(FPropertyChangedChainEvent& PropertyChangedEvent) override; // For TArray type properties in the component.
 #endif
 
 public:
-	virtual							~UMixSettings() override;
+	UE_API virtual							~UMixSettings() override;
 
-	void							FreeTargets();
-	void							InitTargets(size_t count);
-	void							SetTarget(size_t index, TargetTextureSetPtr& target);
-	virtual void					Free();
+	UE_API void							FreeTargets();
+	UE_API void							InitTargets(size_t count);
+	UE_API void							SetTarget(size_t index, TargetTextureSetPtr& target);
+	UE_API virtual void					Free();
 	
 #if WITH_EDITOR
-	AsyncActionResultPtr			SetEditorMesh(AActor* actor);
-	AsyncActionResultPtr			SetEditorMesh(UStaticMeshComponent* meshComponent,UWorld* world);
+	UE_API AsyncActionResultPtr			SetEditorMesh(AActor* actor);
+	UE_API AsyncActionResultPtr			SetEditorMesh(UStaticMeshComponent* meshComponent,UWorld* world);
 #endif
 
-	RenderMeshPtr					GetMesh();
-	UMixInterface*					Mix() const;
+	UE_API RenderMeshPtr					GetMesh();
+	UE_API UMixInterface*					Mix() const;
 	
-	FViewportSettings&				GetViewportSettings();
+	UE_API FViewportSettings&				GetViewportSettings();
 
 	DECLARE_MULTICAST_DELEGATE(FPreviewMeshUpdateEvent)
 	FPreviewMeshUpdateEvent OnPreviewMeshChangedEvent;
@@ -245,3 +247,4 @@ public:
 	}
 };
 
+#undef UE_API
diff --git a/Engine/Plugins/TextureGraph/Source/TextureGraphEngine/Model/Mix/MixUpdateCycle.h b/Engine/Plugins/TextureGraph/Source/TextureGraphEngine/Model/Mix/MixUpdateCycle.h
--- a/Engine/Plugins/TextureGraph/Source/TextureGraphEngine/Model/Mix/MixUpdateCycle.h
+++ b/Engine/Plugins/TextureGraph/Source/TextureGraphEngine/Model/Mix/MixUpdateCycle.h
@@ -7,6 +7,8 @@
 #include "Job/JobCommon.h"
 #include "Model/ModelObject.h"
 
+#define UE_API TEXTUREGRAPHENGINE_API
+
 class JobBatch;
 
 class Job;
@@ -20,7 +22,7 @@ class UMixInterface;
 /// This class doesn't need to create or manage textures. 
 /// its here to refer to textures created by other layers. 
 //////////////////////////////////////////////////////////////////////////
-class TEXTUREGRAPHENGINE_API MixTextureSet : public TextureSet
+class MixTextureSet : public TextureSet
 {
 public:
 	MixTextureSet() : TextureSet(1, 1) {}
@@ -29,7 +31,7 @@ public:
 typedef std::vector<MixTextureSet>	SceneTextureSetVec;
 
 //////////////////////////////////////////////////////////////////////////
-class TEXTUREGRAPHENGINE_API MixTargetUpdate
+class MixTargetUpdate
 {
 private:
 	TWeakObjectPtr<UMixInterface>	MixObj;							/// What is the mix that we're dealing with
@@ -40,13 +42,13 @@ private:
 	JobPtrW							LastAddedJob;					/// Last job that was added
 
 public:
-									MixTargetUpdate(TWeakObjectPtr<UMixInterface> InMixObj, int32 InTargetId);
-									MixTargetUpdate(TWeakObjectPtr<UMixInterface> InMixObj, const TileInvalidateMatrix& InInvalidationMatrix, int32 InTargetId);
-									~MixTargetUpdate();
+									UE_API MixTargetUpdate(TWeakObjectPtr<UMixInterface> InMixObj, int32 InTargetId);
+									UE_API MixTargetUpdate(TWeakObjectPtr<UMixInterface> InMixObj, const TileInvalidateMatrix& InInvalidationMatrix, int32 InTargetId);
+									UE_API ~MixTargetUpdate();
 
-	void							InvalidateTile(int32 Row, int32 Col);
-	void							InvalidateAllTiles();
-	void							InvalidateNoTiles();
+	UE_API void							InvalidateTile(int32 Row, int32 Col);
+	UE_API void							InvalidateAllTiles();
+	UE_API void							InvalidateNoTiles();
 
 	//////////////////////////////////////////////////////////////////////////
 	/// Inline functions
@@ -69,7 +71,7 @@ class ULayerComponent;
 
 class UMaskStack;
 
-class TEXTUREGRAPHENGINE_API MixUpdateCycle : public FGCObject
+class MixUpdateCycle : public FGCObject
 {
 private:
 	const TWeakObjectPtr<UMixInterface>	MixObj;				/// What is the mix that we're dealing with
@@ -88,29 +90,29 @@ private:
 	//////////////////////////////////////////////////////////////////////////
 	/// FGCObject overrides
 	//////////////////////////////////////////////////////////////////////////
-	virtual void					AddReferencedObjects(FReferenceCollector& collector) override;
-	virtual FString					GetReferencerName() const override;
+	UE_API virtual void					AddReferencedObjects(FReferenceCollector& collector) override;
+	UE_API virtual FString					GetReferencerName() const override;
 
 public:
-	explicit						MixUpdateCycle(const FInvalidationDetails& InDetails);
-	virtual							~MixUpdateCycle() override;
+	UE_API explicit						MixUpdateCycle(const FInvalidationDetails& InDetails);
+	UE_API virtual							~MixUpdateCycle() override;
 
-	JobPtrW							AddJob(int32 InTargetId, JobUPtr JobObj);
-	void							AddTarget(SceneTargetUpdatePtr Target);
+	UE_API JobPtrW							AddJob(int32 InTargetId, JobUPtr JobObj);
+	UE_API void							AddTarget(SceneTargetUpdatePtr Target);
 
-	void							Begin();
-	void							End();
+	UE_API void							Begin();
+	UE_API void							End();
 	
-	void							PushMix(UMixInterface* mix);
-	void							PopMix();
-	UMixInterface*					TopMix();
-	bool							ContainsMix(UMixInterface* InMixObj) const;
+	UE_API void							PushMix(UMixInterface* mix);
+	UE_API void							PopMix();
+	UE_API UMixInterface*					TopMix();
+	UE_API bool							ContainsMix(UMixInterface* InMixObj) const;
 	
-	JobPtrW							LastAddedJob(int32 InTargetId) const; 
-	void							SetBatch(std::shared_ptr<JobBatch> InBatch);
-	uint32							LODLevel() const;
-	bool							NoCache() const;
-	void							MergeDetails(const FInvalidationDetails& InDetails);
+	UE_API JobPtrW							LastAddedJob(int32 InTargetId) const; 
+	UE_API void							SetBatch(std::shared_ptr<JobBatch> InBatch);
+	UE_API uint32							LODLevel() const;
+	UE_API bool							NoCache() const;
+	UE_API void							MergeDetails(const FInvalidationDetails& InDetails);
 	SceneTargetUpdatePtr			GetTarget(size_t targetId) { check(targetId < Targets.size()); return Targets[targetId]; }
 
 	//////////////////////////////////////////////////////////////////////////
@@ -126,3 +128,5 @@ public:
 
 typedef std::shared_ptr<MixUpdateCycle>		MixUpdateCyclePtr;
 typedef std::weak_ptr<MixUpdateCycle>		SceneUpdateCyclePtrW;
+
+#undef UE_API
diff --git a/Engine/Plugins/TextureGraph/Source/TextureGraphEngine/Model/Mix/ViewportSettings.h b/Engine/Plugins/TextureGraph/Source/TextureGraphEngine/Model/Mix/ViewportSettings.h
--- a/Engine/Plugins/TextureGraph/Source/TextureGraphEngine/Model/Mix/ViewportSettings.h
+++ b/Engine/Plugins/TextureGraph/Source/TextureGraphEngine/Model/Mix/ViewportSettings.h
@@ -4,10 +4,12 @@
 #include "CoreMinimal.h" 
 #include "ViewportSettings.generated.h"
 
+#define UE_API TEXTUREGRAPHENGINE_API
+
 class UMaterial;
 
 USTRUCT()
-struct TEXTUREGRAPHENGINE_API FMaterialMappingInfo
+struct FMaterialMappingInfo
 {
 	GENERATED_USTRUCT_BODY()
 
@@ -21,7 +23,7 @@ struct TEXTUREGRAPHENGINE_API FMaterialMappingInfo
 };
 
 USTRUCT()
-struct TEXTUREGRAPHENGINE_API FViewportSettings
+struct FViewportSettings
 {
 	GENERATED_BODY()
 	
@@ -31,20 +33,20 @@ struct TEXTUREGRAPHENGINE_API FViewportSettings
 	UPROPERTY(EditAnywhere, EditFixedSize, Category = "Viewport Material")
 	TArray<FMaterialMappingInfo> MaterialMappingInfos;
 	
-	void InitDefaultSettings(FName InitialTargetName);
-	void SetDefaultTarget(FName DefaultTargetName);
+	UE_API void InitDefaultSettings(FName InitialTargetName);
+	UE_API void SetDefaultTarget(FName DefaultTargetName);
 	
-	UMaterial* GetDefaultMaterial();
+	UE_API UMaterial* GetDefaultMaterial();
 	
-	FName GetMaterialName() const;
-	FName GetMaterialMappingInfo(const FName MaterialInput);
-	bool ContainsMaterialMappingInfo(const FName InMaterialInput);
+	UE_API FName GetMaterialName() const;
+	UE_API FName GetMaterialMappingInfo(const FName MaterialInput);
+	UE_API bool ContainsMaterialMappingInfo(const FName InMaterialInput);
 
-	bool RemoveMaterialMappingForTarget(FName OutputNode);
-	void OnMaterialUpdate();
-	void OnTargetRename(const FName OldName, const FName NewName);
+	UE_API bool RemoveMaterialMappingForTarget(FName OutputNode);
+	UE_API void OnMaterialUpdate();
+	UE_API void OnTargetRename(const FName OldName, const FName NewName);
 
-	int  NumAssignedTargets();
+	UE_API int  NumAssignedTargets();
 	
 	DECLARE_MULTICAST_DELEGATE(FViewportSettingsUpdateEvent)
 	FViewportSettingsUpdateEvent OnViewportMaterialChangedEvent;
@@ -52,3 +54,5 @@ struct TEXTUREGRAPHENGINE_API FViewportSettings
 	DECLARE_MULTICAST_DELEGATE(FMaterialMappingChangedEvent)
 	FMaterialMappingChangedEvent OnMaterialMappingChangedEvent;
 };
+
+#undef UE_API
diff --git a/Engine/Plugins/TextureGraph/Source/TextureGraphEngine/Model/ModelObject.h b/Engine/Plugins/TextureGraph/Source/TextureGraphEngine/Model/ModelObject.h
--- a/Engine/Plugins/TextureGraph/Source/TextureGraphEngine/Model/ModelObject.h
+++ b/Engine/Plugins/TextureGraph/Source/TextureGraphEngine/Model/ModelObject.h
@@ -11,11 +11,13 @@
 #include "Materials/MaterialInstanceDynamic.h"
 #include "ModelObject.generated.h"
 
+#define UE_API TEXTUREGRAPHENGINE_API
+
 class UMixInterface;
 
 //////////////////////////////////////////////////////////////////////////
 USTRUCT(BlueprintType)
-struct TEXTUREGRAPHENGINE_API FInvalidationDetails
+struct FInvalidationDetails
 {
 	GENERATED_BODY()
 
@@ -38,22 +40,22 @@ struct TEXTUREGRAPHENGINE_API FInvalidationDetails
 
 	mutable CHashPtr					HashValue = nullptr;								/// The hash for this invalidation 
 	
-	FInvalidationDetails&				All();
-	FInvalidationDetails&				None();
-	FInvalidationDetails&				Merge(const FInvalidationDetails& other); /// Merge the other details with this one,
+	UE_API FInvalidationDetails&				All();
+	UE_API FInvalidationDetails&				None();
+	UE_API FInvalidationDetails&				Merge(const FInvalidationDetails& other); /// Merge the other details with this one,
 
 	/// also merge the other.OnDone delegate, grouped in this._internal_onDoneMerged
 
-	CHashPtr							Hash() const;
-	bool								IsDiscard() const;
+	UE_API CHashPtr							Hash() const;
+	UE_API bool								IsDiscard() const;
 
-	void								BroadcastOnDone() const;	/// Trigger the broadcast of the OnDone delegate AND all 
+	UE_API void								BroadcastOnDone() const;	/// Trigger the broadcast of the OnDone delegate AND all 
 																	/// the other accumulated delegates in _internal_onDoneMerged
 
-	void								BroadcastOnQueued() const;
+	UE_API void								BroadcastOnQueued() const;
 
 	FInvalidationDetails() : Mix(nullptr) {}
-	explicit FInvalidationDetails(const TWeakObjectPtr<UMixInterface> InMixObj);
+	UE_API explicit FInvalidationDetails(const TWeakObjectPtr<UMixInterface> InMixObj);
 
 
 	template <typename ListType>
@@ -88,7 +90,7 @@ private:
 
 //////////////////////////////////////////////////////////////////////////
 USTRUCT(Blueprintable, BlueprintType)
-struct TEXTUREGRAPHENGINE_API FModelInvalidateInfo
+struct FModelInvalidateInfo
 {
 	GENERATED_BODY()
 
@@ -98,7 +100,7 @@ struct TEXTUREGRAPHENGINE_API FModelInvalidateInfo
 	FInvalidationDetails		Details;						/// details about the invalidation
 	mutable CHashPtr			HashValue = nullptr;					/// The hash for this invalidation 
 
-	CHashPtr					Hash() const;
+	UE_API CHashPtr					Hash() const;
 
 	FORCEINLINE UMixInterface*	GetMix() const
 	{
@@ -115,20 +117,20 @@ struct TEXTUREGRAPHENGINE_API FModelInvalidateInfo
 /// This is the base class for all models that need to be loaded from or
 /// saved to the mix
 //////////////////////////////////////////////////////////////////////////
-UCLASS(Blueprintable, BlueprintType)
-class TEXTUREGRAPHENGINE_API UModelObject : public UObject
+UCLASS(MinimalAPI, Blueprintable, BlueprintType)
+class UModelObject : public UObject
 {
 	GENERATED_BODY()
 
 public:
-	virtual						~UModelObject() override;
-	virtual HashType			HashValue_Simple() const;
-	virtual FString				ToString() const;
+	UE_API virtual						~UModelObject() override;
+	UE_API virtual HashType			HashValue_Simple() const;
+	UE_API virtual FString				ToString() const;
 	
 	//////////////////////////////////////////////////////////////////////////
 	// Static methods 
 	//////////////////////////////////////////////////////////////////////////
-	static UObject*				LoadObjectFromPath(const FString& distilledPath);
+	static UE_API UObject*				LoadObjectFromPath(const FString& distilledPath);
 
 	template <typename ModelClass>
 	static ModelClass* CreateNew(EObjectFlags Flags = RF_NoFlags)
@@ -137,3 +139,4 @@ public:
 	}
 };
 
+#undef UE_API
diff --git a/Engine/Plugins/TextureGraph/Source/TextureGraphEngine/Model/StaticImageResource.h b/Engine/Plugins/TextureGraph/Source/TextureGraphEngine/Model/StaticImageResource.h
--- a/Engine/Plugins/TextureGraph/Source/TextureGraphEngine/Model/StaticImageResource.h
+++ b/Engine/Plugins/TextureGraph/Source/TextureGraphEngine/Model/StaticImageResource.h
@@ -7,11 +7,13 @@
 #include "Mix/MixUpdateCycle.h"
 #include "StaticImageResource.generated.h"
 
+#define UE_API TEXTUREGRAPHENGINE_API
+
 //////////////////////////////////////////////////////////////////////////
 /// Base static image resource class
 ////////////////////////////////////////////////////////////////////////// 
-UCLASS(Blueprintable, BlueprintType)
-class TEXTUREGRAPHENGINE_API UStaticImageResource : public UModelObject
+UCLASS(MinimalAPI, Blueprintable, BlueprintType)
+class UStaticImageResource : public UModelObject
 {
 public:
 	GENERATED_BODY()
@@ -29,14 +31,14 @@ private:
 	/*Is loading directly from the filesystem*/
 	bool bIsFilesystem = false;
 
-	virtual AsyncTiledBlobRef Load(MixUpdateCyclePtr Cycle);
+	UE_API virtual AsyncTiledBlobRef Load(MixUpdateCyclePtr Cycle);
 
-	FDateTime GetAssetTimeStamp();
+	UE_API FDateTime GetAssetTimeStamp();
 	
 public:
-	virtual ~UStaticImageResource() override;
+	UE_API virtual ~UStaticImageResource() override;
 	
-	virtual TiledBlobPtr GetBlob(MixUpdateCyclePtr Cycle, BufferDescriptor* DesiredDesc, int32 TargetId);
+	UE_API virtual TiledBlobPtr GetBlob(MixUpdateCyclePtr Cycle, BufferDescriptor* DesiredDesc, int32 TargetId);
 
 	//////////////////////////////////////////////////////////////////////////
 	/// Inline functions
@@ -47,3 +49,4 @@ public:
 	FORCEINLINE bool IsFileSystem() const { return bIsFilesystem; }
 };
 
+#undef UE_API
