Ran UnrealCodeFixup to fix dll storage

diff --git a/Engine/Plugins/TextureGraph/Source/TextureGraphEngine/Profiling/RenderDoc/RenderDocManager.h b/Engine/Plugins/TextureGraph/Source/TextureGraphEngine/Profiling/RenderDoc/RenderDocManager.h
--- a/Engine/Plugins/TextureGraph/Source/TextureGraphEngine/Profiling/RenderDoc/RenderDocManager.h
+++ b/Engine/Plugins/TextureGraph/Source/TextureGraphEngine/Profiling/RenderDoc/RenderDocManager.h
@@ -4,24 +4,28 @@
 #include "CoreMinimal.h"
 #include <memory>
 
+#define UE_API TEXTUREGRAPHENGINE_API
+
 DECLARE_LOG_CATEGORY_EXTERN(LogRenderDocTextureGraph, Log, All);
 
 namespace TextureGraphEditor
 {
-	class TEXTUREGRAPHENGINE_API RenderDocManager
+	class RenderDocManager
 	{
 
 	public:
-		RenderDocManager();
-		~RenderDocManager();
-		void												Initialize();
+		UE_API RenderDocManager();
+		UE_API ~RenderDocManager();
+		UE_API void												Initialize();
 
-		void												CaptureNextBatch();
-		void												CapturePreviousBatch();
-		void												BeginCapture();
-		void												EndCapture();
-		void												CaptureNextBatchHistogram();
+		UE_API void												CaptureNextBatch();
+		UE_API void												CapturePreviousBatch();
+		UE_API void												BeginCapture();
+		UE_API void												EndCapture();
+		UE_API void												CaptureNextBatchHistogram();
 
 };
 	typedef std::unique_ptr<RenderDocManager>				RenderDocManagerPtr;
-}
\ No newline at end of file
+}
+
+#undef UE_API
diff --git a/Engine/Plugins/TextureGraph/Source/TextureGraphEngine/TextureGraphEngine.h b/Engine/Plugins/TextureGraph/Source/TextureGraphEngine/TextureGraphEngine.h
--- a/Engine/Plugins/TextureGraph/Source/TextureGraphEngine/TextureGraphEngine.h
+++ b/Engine/Plugins/TextureGraph/Source/TextureGraphEngine/TextureGraphEngine.h
@@ -11,6 +11,8 @@
 #include "Modules/ModuleManager.h"
 #include "TextureGraphErrorReporter.h"
 
+#define UE_API TEXTUREGRAPHENGINE_API
+
 DECLARE_LOG_CATEGORY_EXTERN(LogTextureGraphEngine, Log, All);
 class UMixInterface;
 class UAppState;
@@ -52,7 +54,7 @@ public:
 };
 
 
-class TEXTUREGRAPHENGINE_API EngineObserverSource
+class EngineObserverSource
 {
 protected:
 	/// Protected interface of emitters called by the engine to notify the observers
@@ -75,21 +77,21 @@ typedef std::shared_ptr<FTextureGraphErrorReporter>			ErrorReporterPtr;
 /// This is just a place for accessing the state of the application
 /// and NOT meant to be a place for doing a lot of the things.
 //////////////////////////////////////////////////////////////////////////
-class TEXTUREGRAPHENGINE_API TextureGraphEngine
+class TextureGraphEngine
 {
 private:
-	static TextureGraphEngine*							GInstance;					/// Static instance. Use Create to create a new instance. Destroy the current instance first
-	static std::atomic<bool>							bGIsEngineDestroying;		/// Whether the Engine is in a state of being destroyed
+	static UE_API TextureGraphEngine*							GInstance;					/// Static instance. Use Create to create a new instance. Destroy the current instance first
+	static UE_API std::atomic<bool>							bGIsEngineDestroying;		/// Whether the Engine is in a state of being destroyed
 	
 	bool												bIsEngineInitDone = false;	/// Whether the Engine init is actually done. This is to peg it to the editor
 
-	explicit											TextureGraphEngine(bool bInIsTestMode);
+	UE_API explicit											TextureGraphEngine(bool bInIsTestMode);
 
-	void												InitEngineInternal();		/// This is to allow modules of Engine to access s_instance
+	UE_API void												InitEngineInternal();		/// This is to allow modules of Engine to access s_instance
 
-														~TextureGraphEngine();
+														UE_API ~TextureGraphEngine();
 
-	static EngineObserverSourcePtr						GObserverSource;			/// Engine observer source interface where public signals will be emitted from. 
+	static UE_API EngineObserverSourcePtr						GObserverSource;			/// Engine observer source interface where public signals will be emitted from. 
 																					/// Default observerSource is the default class implementation which is a no-op
 	
 	TStrongObjectPtr<UMaterialManager>					MaterialMgrObj;				/// The material manager
@@ -115,7 +117,7 @@ private:
 	std::mutex*											LockWaitMutex = nullptr;	/// Mutex used to wait if the engine is locked and asked to be destroyed
 	std::condition_variable*							LockWaitCVar = nullptr;		/// Conditional variable to be used in conjunction with the mutex
 
-	void												FirstRunInit();
+	UE_API void												FirstRunInit();
 
 public:
 	//////////////////////////////////////////////////////////////////////////
@@ -123,26 +125,26 @@ public:
 	//////////////////////////////////////////////////////////////////////////
 	/// Init can be called again and again to reset the engine
 	/// and recreate the data structures
-	static void											Create(bool bIsTestMode = false);
+	static UE_API void											Create(bool bIsTestMode = false);
 
 	/// Destroy the current instance
-	static void											Destroy();
+	static UE_API void											Destroy();
 	
-	static void											InitTests();
-	static void											DoneTests();
+	static UE_API void											InitTests();
+	static UE_API void											DoneTests();
 
 	/// Update the components in the engine
-	static void											Update(float dt);
+	static UE_API void											Update(float dt);
 
 	/// Register the observer source (only one active)
-	static void											RegisterObserverSource(const EngineObserverSourcePtr& observerSource);
+	static UE_API void											RegisterObserverSource(const EngineObserverSourcePtr& observerSource);
 
 	/// Register the Error Reporter
-	static void											RegisterErrorReporter(const UMixInterface* MixKey, const ErrorReporterPtr& ErrorReporter);
+	static UE_API void											RegisterErrorReporter(const UMixInterface* MixKey, const ErrorReporterPtr& ErrorReporter);
 
-	static void											Lock();
-	static void											Unlock();
-	static void											SetRunEngine();
+	static UE_API void											Lock();
+	static UE_API void											Unlock();
+	static UE_API void											SetRunEngine();
 
 	//////////////////////////////////////////////////////////////////////////
 	/// Inline functions
@@ -167,3 +169,4 @@ public:
 };
 
 
+#undef UE_API
diff --git a/Engine/Plugins/TextureGraph/Source/TextureGraphEngine/TextureGraphErrorReporter.h b/Engine/Plugins/TextureGraph/Source/TextureGraphEngine/TextureGraphErrorReporter.h
--- a/Engine/Plugins/TextureGraph/Source/TextureGraphEngine/TextureGraphErrorReporter.h
+++ b/Engine/Plugins/TextureGraph/Source/TextureGraphEngine/TextureGraphErrorReporter.h
@@ -2,6 +2,8 @@
 #pragma once
 #include "CoreMinimal.h"
 
+#define UE_API TEXTUREGRAPHENGINE_API
+
 DECLARE_LOG_CATEGORY_EXTERN(LogTextureGraphError, Log, All);
 struct FTextureGraphErrorReport
 {
@@ -25,7 +27,7 @@ enum class ETextureGraphErrorType: int32
 	NODE_WARNING,
 	INPUT_ARRAY_WARNING
 };
-class TEXTUREGRAPHENGINE_API FTextureGraphErrorReporter
+class FTextureGraphErrorReporter
 {
 public:
 	FTextureGraphErrorReporter()
@@ -34,11 +36,11 @@ public:
 	
 	virtual ~FTextureGraphErrorReporter() { }
 	
-	virtual FTextureGraphErrorReport ReportLog(int32 ErrorId, const FString& ErrorMsg, UObject* ReferenceObj = nullptr);
+	UE_API virtual FTextureGraphErrorReport ReportLog(int32 ErrorId, const FString& ErrorMsg, UObject* ReferenceObj = nullptr);
 
-	virtual FTextureGraphErrorReport ReportWarning(int32 ErrorId, const FString& ErrorMsg, UObject* ReferenceObj = nullptr);
+	UE_API virtual FTextureGraphErrorReport ReportWarning(int32 ErrorId, const FString& ErrorMsg, UObject* ReferenceObj = nullptr);
 
-	virtual FTextureGraphErrorReport ReportError(int32 ErrorId, const FString& ErrorMsg, UObject* ReferenceObj = nullptr);
+	UE_API virtual FTextureGraphErrorReport ReportError(int32 ErrorId, const FString& ErrorMsg, UObject* ReferenceObj = nullptr);
 	
 	virtual void Clear()
 	{
@@ -53,3 +55,5 @@ public:
 		return CompilationErrors;
 	}
 };
+
+#undef UE_API
diff --git a/Engine/Plugins/TextureGraph/Source/TextureGraphEngine/Transform/BlobTransform.h b/Engine/Plugins/TextureGraph/Source/TextureGraphEngine/Transform/BlobTransform.h
--- a/Engine/Plugins/TextureGraph/Source/TextureGraphEngine/Transform/BlobTransform.h
+++ b/Engine/Plugins/TextureGraph/Source/TextureGraphEngine/Transform/BlobTransform.h
@@ -7,6 +7,8 @@
 #include "GraphicsDefs.h"
 #include <vector>
 
+#define UE_API TEXTUREGRAPHENGINE_API
+
 class Device;
 class Blob;
 class DeviceNativeTask;
@@ -55,7 +57,7 @@ struct TransformArgs
 };
 
 //////////////////////////////////////////////////////////////////////////
-struct TEXTUREGRAPHENGINE_API TransformResult
+struct TransformResult
 {
 	std::exception_ptr				ExInner;					/// Original exception that was raised by the action
 	int32							ErrorCode = 0;				/// What is the error code
@@ -76,7 +78,7 @@ typedef cti::continuable<int32>				AsyncPrepareResult;
 /// of having asynchronous function is left upto the Job class, which 
 /// is responsible for using the BlobTransform functions.
 //////////////////////////////////////////////////////////////////////////
-class TEXTUREGRAPHENGINE_API BlobTransform
+class BlobTransform
 {
 public:
 	
@@ -86,20 +88,20 @@ protected:
 	mutable CHashPtr				HashValue;						/// The hash for this BlobObj transform
 
 public:
-									explicit BlobTransform(FString Name);
-	virtual							~BlobTransform();
+									UE_API explicit BlobTransform(FString Name);
+	UE_API virtual							~BlobTransform();
 
 	virtual size_t					NumTargetDevices() const { return 1; }
 	virtual Device*					TargetDevice(size_t Index) const = 0;
 	virtual AsyncTransformResultPtr	Exec(const TransformArgs& Args) = 0;
-	virtual CHashPtr				Hash() const;
+	UE_API virtual CHashPtr				Hash() const;
 	virtual std::shared_ptr<BlobTransform> DuplicateInstance(FString TransformName) { return nullptr; }
 
 	//virtual BlobPtr					Target();
 	//virtual void					SetTarget(BlobPtr target);
 
-	virtual AsyncBufferResultPtr	Bind(BlobPtr BlobObj, const ResourceBindInfo& BindInfo);
-	virtual AsyncBufferResultPtr	Unbind(BlobPtr BlobObj, const ResourceBindInfo& BindInfo);
+	UE_API virtual AsyncBufferResultPtr	Bind(BlobPtr BlobObj, const ResourceBindInfo& BindInfo);
+	UE_API virtual AsyncBufferResultPtr	Unbind(BlobPtr BlobObj, const ResourceBindInfo& BindInfo);
 
 	virtual void					Bind(int32 InValue, const ResourceBindInfo& BindInfo) {}
 	virtual void					Bind(float InValue, const ResourceBindInfo& BindInfo) {}
@@ -122,15 +124,15 @@ public:
 		const char* valueAddress = reinterpret_cast<const char*>(&InValue);
 		BindStruct(valueAddress, sizeof(StructType), BindInfo);
 	}
-	virtual void					Bind(const Vector3& InValue, const ResourceBindInfo& BindInfo);
-	virtual void					Bind(const Vector4& InValue, const ResourceBindInfo& BindInfo);
-	virtual void					Bind(const Tangent& InValue, const ResourceBindInfo& BindInfo);
+	UE_API virtual void					Bind(const Vector3& InValue, const ResourceBindInfo& BindInfo);
+	UE_API virtual void					Bind(const Vector4& InValue, const ResourceBindInfo& BindInfo);
+	UE_API virtual void					Bind(const Tangent& InValue, const ResourceBindInfo& BindInfo);
 	virtual void					Bind(const FString& InValue, const ResourceBindInfo& BindInfo) {}
 
-	virtual bool					GeneratesData() const;
-	virtual bool					CanHandleTiles() const;
+	UE_API virtual bool					GeneratesData() const;
+	UE_API virtual bool					CanHandleTiles() const;
 
-	virtual AsyncPrepareResult		PrepareResources(const TransformArgs& Args);
+	UE_API virtual AsyncPrepareResult		PrepareResources(const TransformArgs& Args);
 
 
 	/// No-Ops
@@ -139,13 +141,13 @@ public:
 	virtual void					Unbind(const FLinearColor& InValue, const ResourceBindInfo& BindInfo) {}
 	virtual void					Unbind(const FIntVector4& InValue, const ResourceBindInfo& BindInfo) {}
 	virtual void					Unbind(const FMatrix& InValue, const ResourceBindInfo& BindInfo) {}
-	virtual void					Unbind(const Vector3& InValue, const ResourceBindInfo& BindInfo);
-	virtual void					Unbind(const Vector4& InValue, const ResourceBindInfo& BindInfo);
-	virtual void					Unbind(const Tangent& InValue, const ResourceBindInfo& BindInfo);
+	UE_API virtual void					Unbind(const Vector3& InValue, const ResourceBindInfo& BindInfo);
+	UE_API virtual void					Unbind(const Vector4& InValue, const ResourceBindInfo& BindInfo);
+	UE_API virtual void					Unbind(const Tangent& InValue, const ResourceBindInfo& BindInfo);
 	virtual void					Unbind(const FString& InValue, const ResourceBindInfo& BindInfo) {}
 
-	virtual ENamedThreads::Type		ExecutionThread() const;
-	virtual bool					IsAsync() const;
+	UE_API virtual ENamedThreads::Type		ExecutionThread() const;
+	UE_API virtual bool					IsAsync() const;
 	virtual FString					GetName() const { return Name; }
 
 
@@ -164,7 +166,7 @@ public:
 /// a transform. This is for completeness so that we don't have to 
 /// put checks in places that are expecting a transform
 //////////////////////////////////////////////////////////////////////////
-class TEXTUREGRAPHENGINE_API Null_Transform : public BlobTransform
+class Null_Transform : public BlobTransform
 {
 private:
 	Device*							Dev = nullptr;			/// The device that this transform belongs to
@@ -172,14 +174,16 @@ private:
 	bool							bGeneratesData = false;	/// The null transform doesn't generate data by default
 
 public:
-									Null_Transform(Device* Dev, FString Name, bool bIsAsync_ = false, bool bGeneratesData_ = false);
-	virtual AsyncBufferResultPtr	Bind(BlobPtr InValue, const ResourceBindInfo& BindInfo) override;
-	virtual AsyncBufferResultPtr	Unbind(BlobPtr BlobObj, const ResourceBindInfo& BindInfo) override;
-	virtual bool					IsAsync() const override;
-	virtual bool					GeneratesData() const override;
-	virtual Device*					TargetDevice(size_t Index) const override;
-	virtual AsyncTransformResultPtr	Exec(const TransformArgs& Args) override;
-	virtual std::shared_ptr<BlobTransform> DuplicateInstance(FString TransformName) override;
+									UE_API Null_Transform(Device* Dev, FString Name, bool bIsAsync_ = false, bool bGeneratesData_ = false);
+	UE_API virtual AsyncBufferResultPtr	Bind(BlobPtr InValue, const ResourceBindInfo& BindInfo) override;
+	UE_API virtual AsyncBufferResultPtr	Unbind(BlobPtr BlobObj, const ResourceBindInfo& BindInfo) override;
+	UE_API virtual bool					IsAsync() const override;
+	UE_API virtual bool					GeneratesData() const override;
+	UE_API virtual Device*					TargetDevice(size_t Index) const override;
+	UE_API virtual AsyncTransformResultPtr	Exec(const TransformArgs& Args) override;
+	UE_API virtual std::shared_ptr<BlobTransform> DuplicateInstance(FString TransformName) override;
 };
 
 typedef std::shared_ptr<BlobTransform> BlobTransformPtr;
+
+#undef UE_API
diff --git a/Engine/Plugins/TextureGraph/Source/TextureGraphEngine/Transform/Expressions/T_ArrayGrid.h b/Engine/Plugins/TextureGraph/Source/TextureGraphEngine/Transform/Expressions/T_ArrayGrid.h
--- a/Engine/Plugins/TextureGraph/Source/TextureGraphEngine/Transform/Expressions/T_ArrayGrid.h
+++ b/Engine/Plugins/TextureGraph/Source/TextureGraphEngine/Transform/Expressions/T_ArrayGrid.h
@@ -9,11 +9,14 @@
 #include "Model/Mix/MixUpdateCycle.h"
 #include <DataDrivenShaderPlatformInfo.h>
 
+#define UE_API TEXTUREGRAPHENGINE_API
+
 //////////////////////////////////////////////////////////////////////////
-class TEXTUREGRAPHENGINE_API T_ArrayGrid
+class T_ArrayGrid
 {
 public:
-	static TiledBlobPtr Create(MixUpdateCyclePtr Cycle, BufferDescriptor& DesiredOutputDesc, TArray<TiledBlobPtr> Inputs, 
+	static UE_API TiledBlobPtr Create(MixUpdateCyclePtr Cycle, BufferDescriptor& DesiredOutputDesc, TArray<TiledBlobPtr> Inputs, 
 		int32 NumRows, int32 NumCols, FLinearColor BackgroundColor, int32 TargetId = 0);
 };
 
+#undef UE_API
diff --git a/Engine/Plugins/TextureGraph/Source/TextureGraphEngine/Transform/Expressions/T_Blend.h b/Engine/Plugins/TextureGraph/Source/TextureGraphEngine/Transform/Expressions/T_Blend.h
--- a/Engine/Plugins/TextureGraph/Source/TextureGraphEngine/Transform/Expressions/T_Blend.h
+++ b/Engine/Plugins/TextureGraph/Source/TextureGraphEngine/Transform/Expressions/T_Blend.h
@@ -9,11 +9,13 @@
 #include "Model/Mix/MixUpdateCycle.h"
 #include <DataDrivenShaderPlatformInfo.h>
 
+#define UE_API TEXTUREGRAPHENGINE_API
+
 //////////////////////////////////////////////////////////////////////////
 /// Shader
 //////////////////////////////////////////////////////////////////////////
 
-class TEXTUREGRAPHENGINE_API FSH_BlendBase : public FSH_Base
+class FSH_BlendBase : public FSH_Base
 {
 public:
 	SHADER_USE_PARAMETER_STRUCT(FSH_BlendBase, FSH_Base);
@@ -43,119 +45,119 @@ public:
 //////////////////////////////////////////////////////////////////////////
 /// Normal
 //////////////////////////////////////////////////////////////////////////
-class TEXTUREGRAPHENGINE_API FSH_BlendNormal : public FSH_BlendBase
+class FSH_BlendNormal : public FSH_BlendBase
 {
 public:
-	DECLARE_GLOBAL_SHADER(FSH_BlendNormal);
+	DECLARE_EXPORTED_GLOBAL_SHADER(FSH_BlendNormal, UE_API);
 	SHADER_USE_PARAMETER_STRUCT(FSH_BlendNormal, FSH_BlendBase);
 };
 
 //////////////////////////////////////////////////////////////////////////
 /// Add
 //////////////////////////////////////////////////////////////////////////
-class TEXTUREGRAPHENGINE_API FSH_BlendAdd : public FSH_BlendBase
+class FSH_BlendAdd : public FSH_BlendBase
 {
 public:
-	DECLARE_GLOBAL_SHADER(FSH_BlendAdd);
+	DECLARE_EXPORTED_GLOBAL_SHADER(FSH_BlendAdd, UE_API);
 	SHADER_USE_PARAMETER_STRUCT(FSH_BlendAdd, FSH_BlendBase);
 };
 
 //////////////////////////////////////////////////////////////////////////
 /// Add
 //////////////////////////////////////////////////////////////////////////
-class TEXTUREGRAPHENGINE_API FSH_BlendSubtract : public FSH_BlendBase
+class FSH_BlendSubtract : public FSH_BlendBase
 {
 public:
-	DECLARE_GLOBAL_SHADER(FSH_BlendSubtract);
+	DECLARE_EXPORTED_GLOBAL_SHADER(FSH_BlendSubtract, UE_API);
 	SHADER_USE_PARAMETER_STRUCT(FSH_BlendSubtract, FSH_BlendBase);
 };
 
 //////////////////////////////////////////////////////////////////////////
 /// Multiply
 //////////////////////////////////////////////////////////////////////////
-class TEXTUREGRAPHENGINE_API FSH_BlendMultiply : public FSH_BlendBase
+class FSH_BlendMultiply : public FSH_BlendBase
 {
 public:
-	DECLARE_GLOBAL_SHADER(FSH_BlendMultiply);
+	DECLARE_EXPORTED_GLOBAL_SHADER(FSH_BlendMultiply, UE_API);
 	SHADER_USE_PARAMETER_STRUCT(FSH_BlendMultiply, FSH_BlendBase);
 };
 
 //////////////////////////////////////////////////////////////////////////
 /// Divide
 //////////////////////////////////////////////////////////////////////////
-class TEXTUREGRAPHENGINE_API FSH_BlendDivide : public FSH_BlendBase
+class FSH_BlendDivide : public FSH_BlendBase
 {
 public:
-	DECLARE_GLOBAL_SHADER(FSH_BlendDivide);
+	DECLARE_EXPORTED_GLOBAL_SHADER(FSH_BlendDivide, UE_API);
 	SHADER_USE_PARAMETER_STRUCT(FSH_BlendDivide, FSH_BlendBase);
 };
 //////////////////////////////////////////////////////////////////////////
 /// Difference
 //////////////////////////////////////////////////////////////////////////
-class TEXTUREGRAPHENGINE_API FSH_BlendDifference : public FSH_BlendBase
+class FSH_BlendDifference : public FSH_BlendBase
 {
 public:
-	DECLARE_GLOBAL_SHADER(FSH_BlendDifference);
+	DECLARE_EXPORTED_GLOBAL_SHADER(FSH_BlendDifference, UE_API);
 	SHADER_USE_PARAMETER_STRUCT(FSH_BlendDifference, FSH_BlendBase);
 };
 
 //////////////////////////////////////////////////////////////////////////
 /// Max
 //////////////////////////////////////////////////////////////////////////
-class TEXTUREGRAPHENGINE_API FSH_BlendMax : public FSH_BlendBase
+class FSH_BlendMax : public FSH_BlendBase
 {
 public:
-	DECLARE_GLOBAL_SHADER(FSH_BlendMax);
+	DECLARE_EXPORTED_GLOBAL_SHADER(FSH_BlendMax, UE_API);
 	SHADER_USE_PARAMETER_STRUCT(FSH_BlendMax, FSH_BlendBase);
 };
 
 //////////////////////////////////////////////////////////////////////////
 /// Min
 //////////////////////////////////////////////////////////////////////////
-class TEXTUREGRAPHENGINE_API FSH_BlendMin : public FSH_BlendBase
+class FSH_BlendMin : public FSH_BlendBase
 {
 public:
-	DECLARE_GLOBAL_SHADER(FSH_BlendMin);
+	DECLARE_EXPORTED_GLOBAL_SHADER(FSH_BlendMin, UE_API);
 	SHADER_USE_PARAMETER_STRUCT(FSH_BlendMin, FSH_BlendBase);
 };
 
 //////////////////////////////////////////////////////////////////////////
 /// Step
 //////////////////////////////////////////////////////////////////////////
-class TEXTUREGRAPHENGINE_API FSH_BlendStep : public FSH_BlendBase
+class FSH_BlendStep : public FSH_BlendBase
 {
 public:
-	DECLARE_GLOBAL_SHADER(FSH_BlendStep);
+	DECLARE_EXPORTED_GLOBAL_SHADER(FSH_BlendStep, UE_API);
 	SHADER_USE_PARAMETER_STRUCT(FSH_BlendStep, FSH_BlendBase);
 };
 
 //////////////////////////////////////////////////////////////////////////
 /// Overlay
 //////////////////////////////////////////////////////////////////////////
-class TEXTUREGRAPHENGINE_API FSH_BlendOverlay : public FSH_BlendBase
+class FSH_BlendOverlay : public FSH_BlendBase
 {
 public:
-	DECLARE_GLOBAL_SHADER(FSH_BlendOverlay);
+	DECLARE_EXPORTED_GLOBAL_SHADER(FSH_BlendOverlay, UE_API);
 	SHADER_USE_PARAMETER_STRUCT(FSH_BlendOverlay, FSH_BlendBase);
 };
 
 //////////////////////////////////////////////////////////////////////////
 /// Distort
 //////////////////////////////////////////////////////////////////////////
-class TEXTUREGRAPHENGINE_API FSH_BlendDistort : public FSH_BlendBase
+class FSH_BlendDistort : public FSH_BlendBase
 {
 public:
-	DECLARE_GLOBAL_SHADER(FSH_BlendDistort);
+	DECLARE_EXPORTED_GLOBAL_SHADER(FSH_BlendDistort, UE_API);
 	SHADER_USE_PARAMETER_STRUCT(FSH_BlendDistort, FSH_BlendBase);
 };
 /**
  * 
  */
-class TEXTUREGRAPHENGINE_API T_Blend
+class T_Blend
 {
 public:
-	T_Blend();
-	~T_Blend();
+	UE_API T_Blend();
+	UE_API ~T_Blend();
 
 	struct FBlendSettings
 	{
@@ -170,16 +172,18 @@ public:
 	/// Static functions
 	//////////////////////////////////////////////////////////////////////////
 
-	static TiledBlobPtr				Create(MixUpdateCyclePtr InCycle, BufferDescriptor DesiredDesc, int InTargetId, EBlendModes InBlendMode, const FBlendSettings* InBlendSettings);
-	static TiledBlobPtr				CreateNormal(MixUpdateCyclePtr InCycle, BufferDescriptor DesiredDesc, int InTargetId, const FBlendSettings* InBlendSettings);
-	static TiledBlobPtr				CreateAdd(MixUpdateCyclePtr InCycle, BufferDescriptor DesiredDesc, int InTargetId, const FBlendSettings* InBlendSettings);
-	static TiledBlobPtr				CreateSubtract(MixUpdateCyclePtr InCycle, BufferDescriptor DesiredDesc, int InTargetId, const FBlendSettings* InBlendSettings);
-	static TiledBlobPtr				CreateMultiply(MixUpdateCyclePtr InCycle, BufferDescriptor DesiredDesc, int InTargetId, const FBlendSettings* InBlendSettings);
-	static TiledBlobPtr				CreateDivide(MixUpdateCyclePtr InCycle, BufferDescriptor DesiredDesc, int InTargetId, const FBlendSettings* InBlendSettings);
-	static TiledBlobPtr				CreateDifference(MixUpdateCyclePtr InCycle, BufferDescriptor DesiredDesc, int InTargetId, const FBlendSettings* InBlendSettings);
-	static TiledBlobPtr				CreateMax(MixUpdateCyclePtr InCycle, BufferDescriptor DesiredDesc, int InTargetId, const FBlendSettings* InBlendSettings);
-	static TiledBlobPtr				CreateMin(MixUpdateCyclePtr InCycle, BufferDescriptor DesiredDesc, int InTargetId, const FBlendSettings* InBlendSettings);
-	static TiledBlobPtr				CreateStep(MixUpdateCyclePtr InCycle, BufferDescriptor DesiredDesc, int InTargetId, const FBlendSettings* InBlendSettings);
-	static TiledBlobPtr				CreateOverlay(MixUpdateCyclePtr InCycle, BufferDescriptor DesiredDesc, int InTargetId, const FBlendSettings* InBlendSettings);
-	static TiledBlobPtr				CreateDistort(MixUpdateCyclePtr InCycle, BufferDescriptor DesiredDesc, int InTargetId, const FBlendSettings* InBlendSettings);
+	static UE_API TiledBlobPtr				Create(MixUpdateCyclePtr InCycle, BufferDescriptor DesiredDesc, int InTargetId, EBlendModes InBlendMode, const FBlendSettings* InBlendSettings);
+	static UE_API TiledBlobPtr				CreateNormal(MixUpdateCyclePtr InCycle, BufferDescriptor DesiredDesc, int InTargetId, const FBlendSettings* InBlendSettings);
+	static UE_API TiledBlobPtr				CreateAdd(MixUpdateCyclePtr InCycle, BufferDescriptor DesiredDesc, int InTargetId, const FBlendSettings* InBlendSettings);
+	static UE_API TiledBlobPtr				CreateSubtract(MixUpdateCyclePtr InCycle, BufferDescriptor DesiredDesc, int InTargetId, const FBlendSettings* InBlendSettings);
+	static UE_API TiledBlobPtr				CreateMultiply(MixUpdateCyclePtr InCycle, BufferDescriptor DesiredDesc, int InTargetId, const FBlendSettings* InBlendSettings);
+	static UE_API TiledBlobPtr				CreateDivide(MixUpdateCyclePtr InCycle, BufferDescriptor DesiredDesc, int InTargetId, const FBlendSettings* InBlendSettings);
+	static UE_API TiledBlobPtr				CreateDifference(MixUpdateCyclePtr InCycle, BufferDescriptor DesiredDesc, int InTargetId, const FBlendSettings* InBlendSettings);
+	static UE_API TiledBlobPtr				CreateMax(MixUpdateCyclePtr InCycle, BufferDescriptor DesiredDesc, int InTargetId, const FBlendSettings* InBlendSettings);
+	static UE_API TiledBlobPtr				CreateMin(MixUpdateCyclePtr InCycle, BufferDescriptor DesiredDesc, int InTargetId, const FBlendSettings* InBlendSettings);
+	static UE_API TiledBlobPtr				CreateStep(MixUpdateCyclePtr InCycle, BufferDescriptor DesiredDesc, int InTargetId, const FBlendSettings* InBlendSettings);
+	static UE_API TiledBlobPtr				CreateOverlay(MixUpdateCyclePtr InCycle, BufferDescriptor DesiredDesc, int InTargetId, const FBlendSettings* InBlendSettings);
+	static UE_API TiledBlobPtr				CreateDistort(MixUpdateCyclePtr InCycle, BufferDescriptor DesiredDesc, int InTargetId, const FBlendSettings* InBlendSettings);
 };
+
+#undef UE_API
diff --git a/Engine/Plugins/TextureGraph/Source/TextureGraphEngine/Transform/Expressions/T_Blur.h b/Engine/Plugins/TextureGraph/Source/TextureGraphEngine/Transform/Expressions/T_Blur.h
--- a/Engine/Plugins/TextureGraph/Source/TextureGraphEngine/Transform/Expressions/T_Blur.h
+++ b/Engine/Plugins/TextureGraph/Source/TextureGraphEngine/Transform/Expressions/T_Blur.h
@@ -9,13 +9,15 @@
 #include "Model/Mix/MixUpdateCycle.h"
 #include <DataDrivenShaderPlatformInfo.h>
 
+#define UE_API TEXTUREGRAPHENGINE_API
+
 //////////////////////////////////////////////////////////////////////////
 /// Gaussian
 //////////////////////////////////////////////////////////////////////////
-class TEXTUREGRAPHENGINE_API FSH_GaussianBlur : public FSH_Base
+class FSH_GaussianBlur : public FSH_Base
 {
 public:
-	DECLARE_GLOBAL_SHADER(FSH_GaussianBlur);
+	DECLARE_EXPORTED_GLOBAL_SHADER(FSH_GaussianBlur, UE_API);
 	SHADER_USE_PARAMETER_STRUCT(FSH_GaussianBlur, FSH_Base);
 	BEGIN_SHADER_PARAMETER_STRUCT(FParameters, )
 		SHADER_PARAMETER_STRUCT(FTileInfo, TileInfo)
@@ -33,10 +35,10 @@ public:
 //////////////////////////////////////////////////////////////////////////
 /// Directional
 //////////////////////////////////////////////////////////////////////////
-class TEXTUREGRAPHENGINE_API FSH_DirectionalBlur : public FSH_Base
+class FSH_DirectionalBlur : public FSH_Base
 {
 public:
-	DECLARE_GLOBAL_SHADER(FSH_DirectionalBlur);
+	DECLARE_EXPORTED_GLOBAL_SHADER(FSH_DirectionalBlur, UE_API);
 	SHADER_USE_PARAMETER_STRUCT(FSH_DirectionalBlur, FSH_Base);
 	BEGIN_SHADER_PARAMETER_STRUCT(FParameters, )
 		SHADER_PARAMETER_STRUCT(FTileInfo, TileInfo)
@@ -59,10 +61,10 @@ public:
 //////////////////////////////////////////////////////////////////////////
 /// Radial
 //////////////////////////////////////////////////////////////////////////
-class TEXTUREGRAPHENGINE_API FSH_RadialBlur : public FSH_Base
+class FSH_RadialBlur : public FSH_Base
 {
 public:
-	DECLARE_GLOBAL_SHADER(FSH_RadialBlur);
+	DECLARE_EXPORTED_GLOBAL_SHADER(FSH_RadialBlur, UE_API);
 	SHADER_USE_PARAMETER_STRUCT(FSH_RadialBlur, FSH_Base);
 	BEGIN_SHADER_PARAMETER_STRUCT(FParameters, )
 		SHADER_PARAMETER_STRUCT(FTileInfo, TileInfo)
@@ -85,19 +87,21 @@ public:
 /**
  * 
  */
-class TEXTUREGRAPHENGINE_API T_Blur
+class T_Blur
 {
 public:
 
-	const static FVector2f RadialBlurCenter;
+	UE_API const static FVector2f RadialBlurCenter;
 	
-	T_Blur();
-	~T_Blur();
+	UE_API T_Blur();
+	UE_API ~T_Blur();
 
 	//////////////////////////////////////////////////////////////////////////
 	/// Static functions
 	//////////////////////////////////////////////////////////////////////////
-	static TiledBlobPtr				CreateGaussian(MixUpdateCyclePtr InCycle, BufferDescriptor DesiredDesc, TiledBlobPtr SourceTexture, int32 InRadius, int InTargetId);
-	static TiledBlobPtr				CreateDirectional(MixUpdateCyclePtr InCycle, BufferDescriptor DesiredDesc, TiledBlobPtr SourceTexture, float InAngle, float InStrength, int InTargetId);
-	static TiledBlobPtr				CreateRadial(MixUpdateCyclePtr InCycle, BufferDescriptor DesiredDesc, TiledBlobPtr SourceTexture, float InRadius, float InStrength, int InTargetId);
+	static UE_API TiledBlobPtr				CreateGaussian(MixUpdateCyclePtr InCycle, BufferDescriptor DesiredDesc, TiledBlobPtr SourceTexture, int32 InRadius, int InTargetId);
+	static UE_API TiledBlobPtr				CreateDirectional(MixUpdateCyclePtr InCycle, BufferDescriptor DesiredDesc, TiledBlobPtr SourceTexture, float InAngle, float InStrength, int InTargetId);
+	static UE_API TiledBlobPtr				CreateRadial(MixUpdateCyclePtr InCycle, BufferDescriptor DesiredDesc, TiledBlobPtr SourceTexture, float InRadius, float InStrength, int InTargetId);
 };
+
+#undef UE_API
diff --git a/Engine/Plugins/TextureGraph/Source/TextureGraphEngine/Transform/Expressions/T_Brightness.h b/Engine/Plugins/TextureGraph/Source/TextureGraphEngine/Transform/Expressions/T_Brightness.h
--- a/Engine/Plugins/TextureGraph/Source/TextureGraphEngine/Transform/Expressions/T_Brightness.h
+++ b/Engine/Plugins/TextureGraph/Source/TextureGraphEngine/Transform/Expressions/T_Brightness.h
@@ -7,13 +7,15 @@
 #include "Model/Mix/MixUpdateCycle.h"
 #include <DataDrivenShaderPlatformInfo.h>
 
+#define UE_API TEXTUREGRAPHENGINE_API
+
 //////////////////////////////////////////////////////////////////////////
 /// Shader
 //////////////////////////////////////////////////////////////////////////
-class TEXTUREGRAPHENGINE_API FSH_Brightness : public FSH_Base
+class FSH_Brightness : public FSH_Base
 {
 public:
-	DECLARE_GLOBAL_SHADER(FSH_Brightness);
+	DECLARE_EXPORTED_GLOBAL_SHADER(FSH_Brightness, UE_API);
 	SHADER_USE_PARAMETER_STRUCT(FSH_Brightness, FSH_Base);
 
 	BEGIN_SHADER_PARAMETER_STRUCT(FParameters, )
@@ -31,11 +33,13 @@ public:
 	static void						ModifyCompilationEnvironment(const FGlobalShaderPermutationParameters& params, FShaderCompilerEnvironment& env) {}
 };
 
-class TEXTUREGRAPHENGINE_API T_Brightness
+class T_Brightness
 {
 public:
 	//////////////////////////////////////////////////////////////////////////
 	/// Static functions
 	//////////////////////////////////////////////////////////////////////////
-	static TiledBlobPtr				Create(MixUpdateCyclePtr Cycle, BufferDescriptor DesiredOutputDesc, TiledBlobPtr Source, float Brightness, float Contrast, int32 TargetId);
+	static UE_API TiledBlobPtr				Create(MixUpdateCyclePtr Cycle, BufferDescriptor DesiredOutputDesc, TiledBlobPtr Source, float Brightness, float Contrast, int32 TargetId);
 };
+
+#undef UE_API
diff --git a/Engine/Plugins/TextureGraph/Source/TextureGraphEngine/Transform/Expressions/T_Channel.h b/Engine/Plugins/TextureGraph/Source/TextureGraphEngine/Transform/Expressions/T_Channel.h
--- a/Engine/Plugins/TextureGraph/Source/TextureGraphEngine/Transform/Expressions/T_Channel.h
+++ b/Engine/Plugins/TextureGraph/Source/TextureGraphEngine/Transform/Expressions/T_Channel.h
@@ -8,6 +8,8 @@
 #include <DataDrivenShaderPlatformInfo.h>
 #include "T_Channel.generated.h"
 
+#define UE_API TEXTUREGRAPHENGINE_API
+
 UENUM(BlueprintType)
 enum class EColorChannel : uint8
 {
@@ -20,7 +22,7 @@ enum class EColorChannel : uint8
 //////////////////////////////////////////////////////////////////////////
 /// Basic Channel Op
 //////////////////////////////////////////////////////////////////////////
-class TEXTUREGRAPHENGINE_API FSH_ChannelSplitter : public FSH_Base
+class FSH_ChannelSplitter : public FSH_Base
 {
 public:
 	SHADER_USE_PARAMETER_STRUCT(FSH_ChannelSplitter, FSH_Base);
@@ -37,50 +39,50 @@ public:
 //////////////////////////////////////////////////////////////////////////
 /// Red
 //////////////////////////////////////////////////////////////////////////
-class TEXTUREGRAPHENGINE_API FSH_ChannelSplitter_Red : public FSH_ChannelSplitter
+class FSH_ChannelSplitter_Red : public FSH_ChannelSplitter
 {
 public:
-	DECLARE_GLOBAL_SHADER(FSH_ChannelSplitter_Red);
+	DECLARE_EXPORTED_GLOBAL_SHADER(FSH_ChannelSplitter_Red, UE_API);
 	SHADER_USE_PARAMETER_STRUCT(FSH_ChannelSplitter_Red, FSH_ChannelSplitter);
 };
 
 //////////////////////////////////////////////////////////////////////////
 /// Green
 //////////////////////////////////////////////////////////////////////////
-class TEXTUREGRAPHENGINE_API FSH_ChannelSplitter_Green : public FSH_ChannelSplitter
+class FSH_ChannelSplitter_Green : public FSH_ChannelSplitter
 {
 public:
-	DECLARE_GLOBAL_SHADER(FSH_ChannelSplitter_Green);
+	DECLARE_EXPORTED_GLOBAL_SHADER(FSH_ChannelSplitter_Green, UE_API);
 	SHADER_USE_PARAMETER_STRUCT(FSH_ChannelSplitter_Green, FSH_ChannelSplitter);
 };
 
 //////////////////////////////////////////////////////////////////////////
 /// Red
 //////////////////////////////////////////////////////////////////////////
-class TEXTUREGRAPHENGINE_API FSH_ChannelSplitter_Blue : public FSH_ChannelSplitter
+class FSH_ChannelSplitter_Blue : public FSH_ChannelSplitter
 {
 public:
-	DECLARE_GLOBAL_SHADER(FSH_ChannelSplitter_Blue);
+	DECLARE_EXPORTED_GLOBAL_SHADER(FSH_ChannelSplitter_Blue, UE_API);
 	SHADER_USE_PARAMETER_STRUCT(FSH_ChannelSplitter_Blue, FSH_ChannelSplitter);
 };
 
 //////////////////////////////////////////////////////////////////////////
 /// Red
 //////////////////////////////////////////////////////////////////////////
-class TEXTUREGRAPHENGINE_API FSH_ChannelSplitter_Alpha : public FSH_ChannelSplitter
+class FSH_ChannelSplitter_Alpha : public FSH_ChannelSplitter
 {
 public:
-	DECLARE_GLOBAL_SHADER(FSH_ChannelSplitter_Alpha);
+	DECLARE_EXPORTED_GLOBAL_SHADER(FSH_ChannelSplitter_Alpha, UE_API);
 	SHADER_USE_PARAMETER_STRUCT(FSH_ChannelSplitter_Alpha, FSH_ChannelSplitter);
 };
 
 //////////////////////////////////////////////////////////////////////////
 /// Channel Combiner
 //////////////////////////////////////////////////////////////////////////
-class TEXTUREGRAPHENGINE_API FSH_ChannelCombiner : public FSH_Base
+class FSH_ChannelCombiner : public FSH_Base
 {
 public:
-	DECLARE_GLOBAL_SHADER(FSH_ChannelCombiner);
+	DECLARE_EXPORTED_GLOBAL_SHADER(FSH_ChannelCombiner, UE_API);
 	SHADER_USE_PARAMETER_STRUCT(FSH_ChannelCombiner, FSH_Base);
 
 	BEGIN_SHADER_PARAMETER_STRUCT(FParameters, )
@@ -103,10 +105,10 @@ class FVar_SwizzleDstChannelGreen : SHADER_PERMUTATION_INT("DST_CHANNEL_GREEN",
 class FVar_SwizzleDstChannelBlue : SHADER_PERMUTATION_INT("DST_CHANNEL_BLUE", 4);
 class FVar_SwizzleDstChannelAlpha : SHADER_PERMUTATION_INT("DST_CHANNEL_ALPHA", 4);
 
-class TEXTUREGRAPHENGINE_API FSH_ChannelSwizzle : public FSH_Base
+class FSH_ChannelSwizzle : public FSH_Base
 {
 public:
-	DECLARE_GLOBAL_SHADER(FSH_ChannelSwizzle);
+	DECLARE_EXPORTED_GLOBAL_SHADER(FSH_ChannelSwizzle, UE_API);
 	SHADER_USE_PARAMETER_STRUCT(FSH_ChannelSwizzle, FSH_Base);
 
 	BEGIN_SHADER_PARAMETER_STRUCT(FParameters, )
@@ -123,10 +125,12 @@ public:
 //////////////////////////////////////////////////////////////////////////
 /// Helper C++ class
 //////////////////////////////////////////////////////////////////////////
-class TEXTUREGRAPHENGINE_API T_Channel
+class T_Channel
 {
 public:
 	//////////////////////////////////////////////////////////////////////////
 	/// Static functions
 	//////////////////////////////////////////////////////////////////////////
 };
+
+#undef UE_API
diff --git a/Engine/Plugins/TextureGraph/Source/TextureGraphEngine/Transform/Expressions/T_Clamp.h b/Engine/Plugins/TextureGraph/Source/TextureGraphEngine/Transform/Expressions/T_Clamp.h
--- a/Engine/Plugins/TextureGraph/Source/TextureGraphEngine/Transform/Expressions/T_Clamp.h
+++ b/Engine/Plugins/TextureGraph/Source/TextureGraphEngine/Transform/Expressions/T_Clamp.h
@@ -7,13 +7,15 @@
 #include "Model/Mix/MixUpdateCycle.h"
 #include <DataDrivenShaderPlatformInfo.h>
 
+#define UE_API TEXTUREGRAPHENGINE_API
+
 //////////////////////////////////////////////////////////////////////////
 /// Clamp Shader
 //////////////////////////////////////////////////////////////////////////
-class TEXTUREGRAPHENGINE_API FSH_Clamp : public FSH_Base
+class FSH_Clamp : public FSH_Base
 {
 public:
-	DECLARE_GLOBAL_SHADER(FSH_Clamp);
+	DECLARE_EXPORTED_GLOBAL_SHADER(FSH_Clamp, UE_API);
 	SHADER_USE_PARAMETER_STRUCT(FSH_Clamp, FSH_Base);
 
 	BEGIN_SHADER_PARAMETER_STRUCT(FParameters, )
@@ -32,10 +34,10 @@ public:
 //////////////////////////////////////////////////////////////////////////
 /// Clamp Shader
 //////////////////////////////////////////////////////////////////////////
-class TEXTUREGRAPHENGINE_API FSH_SmoothStep : public FSH_Base
+class FSH_SmoothStep : public FSH_Base
 {
 public:
-	DECLARE_GLOBAL_SHADER(FSH_SmoothStep);
+	DECLARE_EXPORTED_GLOBAL_SHADER(FSH_SmoothStep, UE_API);
 	SHADER_USE_PARAMETER_STRUCT(FSH_SmoothStep, FSH_Base);
 
 	BEGIN_SHADER_PARAMETER_STRUCT(FParameters, )
@@ -54,10 +56,10 @@ public:
 //////////////////////////////////////////////////////////////////////////
 /// Min Shader
 //////////////////////////////////////////////////////////////////////////
-class TEXTUREGRAPHENGINE_API FSH_Min : public FSH_Base
+class FSH_Min : public FSH_Base
 {
 public:
-	DECLARE_GLOBAL_SHADER(FSH_Min);
+	DECLARE_EXPORTED_GLOBAL_SHADER(FSH_Min, UE_API);
 	SHADER_USE_PARAMETER_STRUCT(FSH_Min, FSH_Base);
 
 	BEGIN_SHADER_PARAMETER_STRUCT(FParameters, )
@@ -75,10 +77,10 @@ public:
 //////////////////////////////////////////////////////////////////////////
 /// Max Shader
 //////////////////////////////////////////////////////////////////////////
-class TEXTUREGRAPHENGINE_API FSH_Max : public FSH_Base
+class FSH_Max : public FSH_Base
 {
 public:
-	DECLARE_GLOBAL_SHADER(FSH_Max);
+	DECLARE_EXPORTED_GLOBAL_SHADER(FSH_Max, UE_API);
 	SHADER_USE_PARAMETER_STRUCT(FSH_Max, FSH_Base);
 
 	BEGIN_SHADER_PARAMETER_STRUCT(FParameters, )
@@ -93,14 +95,16 @@ public:
 	TEXTUREGRAPH_ENGINE_DEFAULT_COMPILATION_ENV;
 };
 
-class TEXTUREGRAPHENGINE_API T_Clamp
+class T_Clamp
 {
 public:
 	//////////////////////////////////////////////////////////////////////////
 	/// Static functions
 	//////////////////////////////////////////////////////////////////////////
-	static TiledBlobPtr				CreateClamp(MixUpdateCyclePtr Cycle, BufferDescriptor DesiredOutputDesc, TiledBlobPtr Source, TiledBlobPtr MinValue, TiledBlobPtr MaxValue, int32 TargetId);
-	static TiledBlobPtr				CreateSmoothStep(MixUpdateCyclePtr Cycle, BufferDescriptor DesiredOutputDesc, TiledBlobPtr Source, TiledBlobPtr MinValue, TiledBlobPtr MaxValue, int32 TargetId);
-	static TiledBlobPtr				CreateMin(MixUpdateCyclePtr Cycle, BufferDescriptor DesiredOutputDesc, TiledBlobPtr Input1, TiledBlobPtr Input2, int32 TargetId);
-	static TiledBlobPtr				CreateMax(MixUpdateCyclePtr Cycle, BufferDescriptor DesiredOutputDesc, TiledBlobPtr Input1, TiledBlobPtr Input2, int32 TargetId);
+	static UE_API TiledBlobPtr				CreateClamp(MixUpdateCyclePtr Cycle, BufferDescriptor DesiredOutputDesc, TiledBlobPtr Source, TiledBlobPtr MinValue, TiledBlobPtr MaxValue, int32 TargetId);
+	static UE_API TiledBlobPtr				CreateSmoothStep(MixUpdateCyclePtr Cycle, BufferDescriptor DesiredOutputDesc, TiledBlobPtr Source, TiledBlobPtr MinValue, TiledBlobPtr MaxValue, int32 TargetId);
+	static UE_API TiledBlobPtr				CreateMin(MixUpdateCyclePtr Cycle, BufferDescriptor DesiredOutputDesc, TiledBlobPtr Input1, TiledBlobPtr Input2, int32 TargetId);
+	static UE_API TiledBlobPtr				CreateMax(MixUpdateCyclePtr Cycle, BufferDescriptor DesiredOutputDesc, TiledBlobPtr Input1, TiledBlobPtr Input2, int32 TargetId);
 };
+
+#undef UE_API
diff --git a/Engine/Plugins/TextureGraph/Source/TextureGraphEngine/Transform/Expressions/T_Color.h b/Engine/Plugins/TextureGraph/Source/TextureGraphEngine/Transform/Expressions/T_Color.h
--- a/Engine/Plugins/TextureGraph/Source/TextureGraphEngine/Transform/Expressions/T_Color.h
+++ b/Engine/Plugins/TextureGraph/Source/TextureGraphEngine/Transform/Expressions/T_Color.h
@@ -7,14 +7,16 @@
 #include "Model/Mix/MixUpdateCycle.h"
 #include <DataDrivenShaderPlatformInfo.h>
 
+#define UE_API TEXTUREGRAPHENGINE_API
+
 //////////////////////////////////////////////////////////////////////////
 /// Simple grayscale shader
 //////////////////////////////////////////////////////////////////////////
-class TEXTUREGRAPHENGINE_API FSH_Grayscale : public FSH_Base
+class FSH_Grayscale : public FSH_Base
 {
 public:
 	SHADER_USE_PARAMETER_STRUCT(FSH_Grayscale, FSH_Base);
-	DECLARE_GLOBAL_SHADER(FSH_Grayscale);
+	DECLARE_EXPORTED_GLOBAL_SHADER(FSH_Grayscale, UE_API);
 
 	BEGIN_SHADER_PARAMETER_STRUCT(FParameters, )
 		SHADER_PARAMETER_STRUCT(FStandardSamplerStates, SamplerStates)
@@ -28,11 +30,11 @@ public:
 //////////////////////////////////////////////////////////////////////////
 /// Simple levels shader
 //////////////////////////////////////////////////////////////////////////
-class TEXTUREGRAPHENGINE_API FSH_Threshold : public FSH_Base
+class FSH_Threshold : public FSH_Base
 {
 public:
 	SHADER_USE_PARAMETER_STRUCT(FSH_Threshold, FSH_Base);
-	DECLARE_GLOBAL_SHADER(FSH_Threshold);
+	DECLARE_EXPORTED_GLOBAL_SHADER(FSH_Threshold, UE_API);
 
 	BEGIN_SHADER_PARAMETER_STRUCT(FParameters, )
 		SHADER_PARAMETER_STRUCT(FStandardSamplerStates, SamplerStates)
@@ -47,11 +49,11 @@ public:
 //////////////////////////////////////////////////////////////////////////
 /// Simple HSV shader
 //////////////////////////////////////////////////////////////////////////
-class TEXTUREGRAPHENGINE_API FSH_HSV : public FSH_Base
+class FSH_HSV : public FSH_Base
 {
 public:
 	SHADER_USE_PARAMETER_STRUCT(FSH_HSV, FSH_Base);
-	DECLARE_GLOBAL_SHADER(FSH_HSV);
+	DECLARE_EXPORTED_GLOBAL_SHADER(FSH_HSV, UE_API);
 
 	BEGIN_SHADER_PARAMETER_STRUCT(FParameters, )
 		SHADER_PARAMETER_STRUCT(FStandardSamplerStates, SamplerStates)
@@ -68,11 +70,11 @@ public:
 //////////////////////////////////////////////////////////////////////////
 /// Simple RGB2HSV shader
 //////////////////////////////////////////////////////////////////////////
-class TEXTUREGRAPHENGINE_API FSH_RGB2HSV : public FSH_Base
+class FSH_RGB2HSV : public FSH_Base
 {
 public:
 	SHADER_USE_PARAMETER_STRUCT(FSH_RGB2HSV, FSH_Base);
-	DECLARE_GLOBAL_SHADER(FSH_RGB2HSV);
+	DECLARE_EXPORTED_GLOBAL_SHADER(FSH_RGB2HSV, UE_API);
 
 	BEGIN_SHADER_PARAMETER_STRUCT(FParameters, )
 		SHADER_PARAMETER_STRUCT(FStandardSamplerStates, SamplerStates)
@@ -86,11 +88,11 @@ public:
 //////////////////////////////////////////////////////////////////////////
 /// Simple HSV2RGB shader
 //////////////////////////////////////////////////////////////////////////
-class TEXTUREGRAPHENGINE_API FSH_HSV2RGB : public FSH_Base
+class FSH_HSV2RGB : public FSH_Base
 {
 public:
 	SHADER_USE_PARAMETER_STRUCT(FSH_HSV2RGB, FSH_Base);
-	DECLARE_GLOBAL_SHADER(FSH_HSV2RGB);
+	DECLARE_EXPORTED_GLOBAL_SHADER(FSH_HSV2RGB, UE_API);
 
 	BEGIN_SHADER_PARAMETER_STRUCT(FParameters, )
 		SHADER_PARAMETER_STRUCT(FStandardSamplerStates, SamplerStates)
@@ -104,11 +106,11 @@ public:
 //////////////////////////////////////////////////////////////////////////
 /// Color correction shader
 //////////////////////////////////////////////////////////////////////////
-class TEXTUREGRAPHENGINE_API FSH_Premult : public FSH_Base
+class FSH_Premult : public FSH_Base
 {
 public:
 	SHADER_USE_PARAMETER_STRUCT(FSH_Premult, FSH_Base);
-	DECLARE_GLOBAL_SHADER(FSH_Premult);
+	DECLARE_EXPORTED_GLOBAL_SHADER(FSH_Premult, UE_API);
 
 	BEGIN_SHADER_PARAMETER_STRUCT(FParameters, )
 		SHADER_PARAMETER_STRUCT(FStandardSamplerStates, SamplerStates)
@@ -122,11 +124,11 @@ public:
 //////////////////////////////////////////////////////////////////////////
 /// Color correction shader
 //////////////////////////////////////////////////////////////////////////
-class TEXTUREGRAPHENGINE_API FSH_ColorCorrection : public FSH_Base
+class FSH_ColorCorrection : public FSH_Base
 {
 public:
 	SHADER_USE_PARAMETER_STRUCT(FSH_ColorCorrection, FSH_Base);
-	DECLARE_GLOBAL_SHADER(FSH_ColorCorrection);
+	DECLARE_EXPORTED_GLOBAL_SHADER(FSH_ColorCorrection, UE_API);
 
 	BEGIN_SHADER_PARAMETER_STRUCT(FParameters, )
 		SHADER_PARAMETER_STRUCT(FStandardSamplerStates, SamplerStates)
@@ -147,10 +149,12 @@ public:
 //////////////////////////////////////////////////////////////////////////
 /// Helper C++ class
 //////////////////////////////////////////////////////////////////////////
-class TEXTUREGRAPHENGINE_API T_Color
+class T_Color
 {
 public:
 	//////////////////////////////////////////////////////////////////////////
 	/// Static functions
 	//////////////////////////////////////////////////////////////////////////
 };
+
+#undef UE_API
diff --git a/Engine/Plugins/TextureGraph/Source/TextureGraphEngine/Transform/Expressions/T_Conversions.h b/Engine/Plugins/TextureGraph/Source/TextureGraphEngine/Transform/Expressions/T_Conversions.h
--- a/Engine/Plugins/TextureGraph/Source/TextureGraphEngine/Transform/Expressions/T_Conversions.h
+++ b/Engine/Plugins/TextureGraph/Source/TextureGraphEngine/Transform/Expressions/T_Conversions.h
@@ -8,14 +8,16 @@
 #include <DataDrivenShaderPlatformInfo.h>
 #include "Data/TiledBlob.h"
 #include "Data/RawBuffer.h"
+
+#define UE_API TEXTUREGRAPHENGINE_API
 //////////////////////////////////////////////////////////////////////////
 /// Normal from height map shader
 //////////////////////////////////////////////////////////////////////////
-class TEXTUREGRAPHENGINE_API FSH_NormalFromHeightMap : public FSH_Base
+class FSH_NormalFromHeightMap : public FSH_Base
 {
 public:
 	SHADER_USE_PARAMETER_STRUCT(FSH_NormalFromHeightMap, FSH_Base);
-	DECLARE_GLOBAL_SHADER(FSH_NormalFromHeightMap);
+	DECLARE_EXPORTED_GLOBAL_SHADER(FSH_NormalFromHeightMap, UE_API);
 
 	BEGIN_SHADER_PARAMETER_STRUCT(FParameters, )
 		SHADER_PARAMETER_STRUCT(FTileInfo, TileInfo)
@@ -35,15 +37,17 @@ public:
 //////////////////////////////////////////////////////////////////////////
 /// Helper C++ class
 //////////////////////////////////////////////////////////////////////////
-class TEXTUREGRAPHENGINE_API T_Conversions
+class T_Conversions
 {
 public:
-	T_Conversions();
-	~T_Conversions();
+	UE_API T_Conversions();
+	UE_API ~T_Conversions();
 	
 	//////////////////////////////////////////////////////////////////////////
 	/// Static functions
 	//////////////////////////////////////////////////////////////////////////
 	///
-	static TiledBlobPtr				CreateNormalFromHeightMap(MixUpdateCyclePtr InCycle, TiledBlobPtr HeightMapTexture, float Offset, float Strength, int InTargetId, BufferDescriptor DesiredOutputDesc);
+	static UE_API TiledBlobPtr				CreateNormalFromHeightMap(MixUpdateCyclePtr InCycle, TiledBlobPtr HeightMapTexture, float Offset, float Strength, int InTargetId, BufferDescriptor DesiredOutputDesc);
 };
+
+#undef UE_API
diff --git a/Engine/Plugins/TextureGraph/Source/TextureGraphEngine/Transform/Expressions/T_ExtractMaterialIds.h b/Engine/Plugins/TextureGraph/Source/TextureGraphEngine/Transform/Expressions/T_ExtractMaterialIds.h
--- a/Engine/Plugins/TextureGraph/Source/TextureGraphEngine/Transform/Expressions/T_ExtractMaterialIds.h
+++ b/Engine/Plugins/TextureGraph/Source/TextureGraphEngine/Transform/Expressions/T_ExtractMaterialIds.h
@@ -9,6 +9,8 @@
 #include "Model/Mix/MixUpdateCycle.h"
 #include "T_ExtractMaterialIds.generated.h"
 
+#define UE_API TEXTUREGRAPHENGINE_API
+
 typedef TArray<TArray<TArray<class HSVBucket*>>> Bucket3dArray;
 typedef TArray<TArray<TArray<int32>>> Color3dArray;
 
@@ -178,7 +180,7 @@ public:
 };
 
 //////////////////////////////////////////////////////////////////////////
-class TEXTUREGRAPHENGINE_API Job_ExtractMaterialIds : public Job
+class Job_ExtractMaterialIds : public Job
 {
 private:
 	TiledBlobRef SourceTexture;
@@ -187,7 +189,7 @@ private:
 	int32 ActiveColorsCount;
 
 public:
-	Job_ExtractMaterialIds(int32 targetId, TiledBlobRef InSourceTexture, FMaterialIDCollection& MaterialIDInfoCollection, TArray<FMaterialIDMaskInfo>& InMaterialIDMaskInfo, int32& ActiveColorsCount, UObject* InErrorOwner = nullptr, uint16 priority = (uint16)E_Priority::kHigh, uint64 id = 0);
+	UE_API Job_ExtractMaterialIds(int32 targetId, TiledBlobRef InSourceTexture, FMaterialIDCollection& MaterialIDInfoCollection, TArray<FMaterialIDMaskInfo>& InMaterialIDMaskInfo, int32& ActiveColorsCount, UObject* InErrorOwner = nullptr, uint16 priority = (uint16)E_Priority::kHigh, uint64 id = 0);
 
 protected:
 	virtual void GetDependencies(JobPtrVec& prior, JobPtrVec& after, JobRunInfo runInfo) override { RunInfo = runInfo; Prev.clear(); }
@@ -199,22 +201,24 @@ protected:
 	virtual cti::continuable<int32>	UnbindArgs_All(JobRunInfo runInfo) { return cti::make_ready_continuable(0); }
 	virtual bool CanHandleTiles() const { return false; }
 
-	virtual cti::continuable<int32>	PreExecAsync(ENamedThreads::Type execThread, ENamedThreads::Type returnThread) override;
-	virtual cti::continuable<int32>	ExecAsync(ENamedThreads::Type ExecThread, ENamedThreads::Type ReturnThread) override;
+	UE_API virtual cti::continuable<int32>	PreExecAsync(ENamedThreads::Type execThread, ENamedThreads::Type returnThread) override;
+	UE_API virtual cti::continuable<int32>	ExecAsync(ENamedThreads::Type ExecThread, ENamedThreads::Type ReturnThread) override;
 };
 
 /**
  * 
  */
-class TEXTUREGRAPHENGINE_API T_ExtractMaterialIds
+class T_ExtractMaterialIds
 {
 public:
-	T_ExtractMaterialIds();
-	~T_ExtractMaterialIds();
+	UE_API T_ExtractMaterialIds();
+	UE_API ~T_ExtractMaterialIds();
 
 	//////////////////////////////////////////////////////////////////////////
 	/// Static functions
 	//////////////////////////////////////////////////////////////////////////
 
-	static TiledBlobPtr Create(MixUpdateCyclePtr InCycle, TiledBlobPtr MaterialIDTexture, FMaterialIDCollection& MaterialIDCollection, TArray<FMaterialIDMaskInfo>& InMaterialIDMaskInfo, int32& ActiveColorsCount, int InTargetId);
+	static UE_API TiledBlobPtr Create(MixUpdateCyclePtr InCycle, TiledBlobPtr MaterialIDTexture, FMaterialIDCollection& MaterialIDCollection, TArray<FMaterialIDMaskInfo>& InMaterialIDMaskInfo, int32& ActiveColorsCount, int InTargetId);
 };
+
+#undef UE_API
diff --git a/Engine/Plugins/TextureGraph/Source/TextureGraphEngine/Transform/Expressions/T_Filter.h b/Engine/Plugins/TextureGraph/Source/TextureGraphEngine/Transform/Expressions/T_Filter.h
--- a/Engine/Plugins/TextureGraph/Source/TextureGraphEngine/Transform/Expressions/T_Filter.h
+++ b/Engine/Plugins/TextureGraph/Source/TextureGraphEngine/Transform/Expressions/T_Filter.h
@@ -10,13 +10,15 @@
 #include <DataDrivenShaderPlatformInfo.h>
 #include "T_Filter.generated.h"
 
+#define UE_API TEXTUREGRAPHENGINE_API
+
 //////////////////////////////////////////////////////////////////////////
 /// Edge detection
 //////////////////////////////////////////////////////////////////////////
-class TEXTUREGRAPHENGINE_API FSH_EdgeDetect : public FSH_Base
+class FSH_EdgeDetect : public FSH_Base
 {
 public:
-	DECLARE_GLOBAL_SHADER(FSH_EdgeDetect);
+	DECLARE_EXPORTED_GLOBAL_SHADER(FSH_EdgeDetect, UE_API);
 	SHADER_USE_PARAMETER_STRUCT(FSH_EdgeDetect, FSH_Base);
 	BEGIN_SHADER_PARAMETER_STRUCT(FParameters, )
 		SHADER_PARAMETER_STRUCT(FTileInfo, TileInfo)
@@ -44,10 +46,10 @@ namespace EWarp
 //////////////////////////////////////////////////////////////////////////
 /// Directional warp
 //////////////////////////////////////////////////////////////////////////
-class TEXTUREGRAPHENGINE_API FSH_DirectionalWarp : public FSH_Base
+class FSH_DirectionalWarp : public FSH_Base
 {
 public:
-	DECLARE_GLOBAL_SHADER(FSH_DirectionalWarp);
+	DECLARE_EXPORTED_GLOBAL_SHADER(FSH_DirectionalWarp, UE_API);
 	SHADER_USE_PARAMETER_STRUCT(FSH_DirectionalWarp, FSH_Base);
 	BEGIN_SHADER_PARAMETER_STRUCT(FParameters, )
 		SHADER_PARAMETER_STRUCT(FTileInfo, TileInfo)
@@ -66,10 +68,10 @@ public:
 //////////////////////////////////////////////////////////////////////////
 /// Directional warp
 //////////////////////////////////////////////////////////////////////////
-class TEXTUREGRAPHENGINE_API FSH_NormalWarp : public FSH_Base
+class FSH_NormalWarp : public FSH_Base
 {
 public:
-	DECLARE_GLOBAL_SHADER(FSH_NormalWarp);
+	DECLARE_EXPORTED_GLOBAL_SHADER(FSH_NormalWarp, UE_API);
 	SHADER_USE_PARAMETER_STRUCT(FSH_NormalWarp, FSH_Base);
 	BEGIN_SHADER_PARAMETER_STRUCT(FParameters, )
 		SHADER_PARAMETER_STRUCT(FTileInfo, TileInfo)
@@ -87,10 +89,10 @@ public:
 //////////////////////////////////////////////////////////////////////////
 /// Directional warp
 //////////////////////////////////////////////////////////////////////////
-class TEXTUREGRAPHENGINE_API FSH_SineWarp : public FSH_Base
+class FSH_SineWarp : public FSH_Base
 {
 public:
-	DECLARE_GLOBAL_SHADER(FSH_SineWarp);
+	DECLARE_EXPORTED_GLOBAL_SHADER(FSH_SineWarp, UE_API);
 	SHADER_USE_PARAMETER_STRUCT(FSH_SineWarp, FSH_Base);
 	BEGIN_SHADER_PARAMETER_STRUCT(FParameters, )
 		SHADER_PARAMETER_STRUCT(FTileInfo, TileInfo)
@@ -113,10 +115,10 @@ class FVar_ErodeDilate_IsSingleChannel : SHADER_PERMUTATION_BOOL("ED_SINGLECHANN
 //////////////////////////////////////////////////////////////////////////
 /// Erode/Dilate
 //////////////////////////////////////////////////////////////////////////
-class TEXTUREGRAPHENGINE_API FSH_ErodeDilate : public FSH_Base
+class FSH_ErodeDilate : public FSH_Base
 {
 public:
-	DECLARE_GLOBAL_SHADER(FSH_ErodeDilate);
+	DECLARE_EXPORTED_GLOBAL_SHADER(FSH_ErodeDilate, UE_API);
 	SHADER_USE_PARAMETER_STRUCT(FSH_ErodeDilate, FSH_Base);
 	BEGIN_SHADER_PARAMETER_STRUCT(FParameters, )
 		SHADER_PARAMETER_STRUCT(FTileInfo, TileInfo)
@@ -133,14 +135,16 @@ public:
 
 
 //////////////////////////////////////////////////////////////////////////
-class TEXTUREGRAPHENGINE_API T_Filter
+class T_Filter
 {
 public:
 	//////////////////////////////////////////////////////////////////////////
 	/// Static functions
 	//////////////////////////////////////////////////////////////////////////
-	static TiledBlobPtr				CreateEdgeDetect(MixUpdateCyclePtr InCycle, BufferDescriptor DesiredDesc, TiledBlobPtr SourceTexture, float Thickness, int32 InTargetId);
-	static TiledBlobPtr				CreateDirectionalWarp(MixUpdateCyclePtr InCycle, BufferDescriptor DesiredDesc, TiledBlobPtr SourceTexture, TiledBlobPtr Mask, float Intensity, float AngleRad, int32 InTargetId);
-	static TiledBlobPtr				CreateNormalWarp(MixUpdateCyclePtr InCycle, BufferDescriptor DesiredDesc, TiledBlobPtr SourceTexture, TiledBlobPtr Mask, float Intensity, int32 InTargetId);
-	static TiledBlobPtr				CreateSineWarp(MixUpdateCyclePtr InCycle, BufferDescriptor DesiredDesc, TiledBlobPtr SourceTexture, TiledBlobPtr Mask, float Intensity, float PhaseU, float PhaseV, int32 InTargetId);
+	static UE_API TiledBlobPtr				CreateEdgeDetect(MixUpdateCyclePtr InCycle, BufferDescriptor DesiredDesc, TiledBlobPtr SourceTexture, float Thickness, int32 InTargetId);
+	static UE_API TiledBlobPtr				CreateDirectionalWarp(MixUpdateCyclePtr InCycle, BufferDescriptor DesiredDesc, TiledBlobPtr SourceTexture, TiledBlobPtr Mask, float Intensity, float AngleRad, int32 InTargetId);
+	static UE_API TiledBlobPtr				CreateNormalWarp(MixUpdateCyclePtr InCycle, BufferDescriptor DesiredDesc, TiledBlobPtr SourceTexture, TiledBlobPtr Mask, float Intensity, int32 InTargetId);
+	static UE_API TiledBlobPtr				CreateSineWarp(MixUpdateCyclePtr InCycle, BufferDescriptor DesiredDesc, TiledBlobPtr SourceTexture, TiledBlobPtr Mask, float Intensity, float PhaseU, float PhaseV, int32 InTargetId);
 };
+
+#undef UE_API
diff --git a/Engine/Plugins/TextureGraph/Source/TextureGraphEngine/Transform/Expressions/T_FlatColorTexture.h b/Engine/Plugins/TextureGraph/Source/TextureGraphEngine/Transform/Expressions/T_FlatColorTexture.h
--- a/Engine/Plugins/TextureGraph/Source/TextureGraphEngine/Transform/Expressions/T_FlatColorTexture.h
+++ b/Engine/Plugins/TextureGraph/Source/TextureGraphEngine/Transform/Expressions/T_FlatColorTexture.h
@@ -6,6 +6,8 @@
 #include "FxMat/FxMaterial.h"
 #include "Model/Mix/MixUpdateCycle.h"
 
+#define UE_API TEXTUREGRAPHENGINE_API
+
 //////////////////////////////////////////////////////////////////////////
 /// Shader
 //////////////////////////////////////////////////////////////////////////
@@ -14,7 +16,7 @@ class FSH_FlatColorTexture : public FSH_Base
 {
 public:
 
-	DECLARE_GLOBAL_SHADER(FSH_FlatColorTexture);
+	DECLARE_EXPORTED_GLOBAL_SHADER(FSH_FlatColorTexture, UE_API);
 	SHADER_USE_PARAMETER_STRUCT(FSH_FlatColorTexture, FSH_Base);
 
 	BEGIN_SHADER_PARAMETER_STRUCT(FParameters, )
@@ -28,12 +30,14 @@ typedef FxMaterial_Normal<VSH_Simple, FSH_FlatColorTexture>	Fx_FlatColorTexture;
 /**
  * 
  */
-class TEXTUREGRAPHENGINE_API T_FlatColorTexture
+class T_FlatColorTexture
 {
 public:
 	//////////////////////////////////////////////////////////////////////////
 	/// Static functions
 	//////////////////////////////////////////////////////////////////////////
-	static BufferDescriptor			GetFlatColorDesc(FString name, BufferFormat InBufferFormat = BufferFormat::Byte);
-	static TiledBlobPtr				Create(MixUpdateCyclePtr InCycle, BufferDescriptor DesiredOutputDesc, FLinearColor Color, int InTargetId);
+	static UE_API BufferDescriptor			GetFlatColorDesc(FString name, BufferFormat InBufferFormat = BufferFormat::Byte);
+	static UE_API TiledBlobPtr				Create(MixUpdateCyclePtr InCycle, BufferDescriptor DesiredOutputDesc, FLinearColor Color, int InTargetId);
 };
+
+#undef UE_API
diff --git a/Engine/Plugins/TextureGraph/Source/TextureGraphEngine/Transform/Expressions/T_Invert.h b/Engine/Plugins/TextureGraph/Source/TextureGraphEngine/Transform/Expressions/T_Invert.h
--- a/Engine/Plugins/TextureGraph/Source/TextureGraphEngine/Transform/Expressions/T_Invert.h
+++ b/Engine/Plugins/TextureGraph/Source/TextureGraphEngine/Transform/Expressions/T_Invert.h
@@ -7,13 +7,15 @@
 #include "Model/Mix/MixUpdateCycle.h"
 #include <DataDrivenShaderPlatformInfo.h>
 
+#define UE_API TEXTUREGRAPHENGINE_API
+
 //////////////////////////////////////////////////////////////////////////
 /// Shader
 //////////////////////////////////////////////////////////////////////////
-class TEXTUREGRAPHENGINE_API FSH_Invert : public FSH_Base
+class FSH_Invert : public FSH_Base
 {
 public:
-	DECLARE_GLOBAL_SHADER(FSH_Invert);
+	DECLARE_EXPORTED_GLOBAL_SHADER(FSH_Invert, UE_API);
 	SHADER_USE_PARAMETER_STRUCT(FSH_Invert, FSH_Base);
 
 	BEGIN_SHADER_PARAMETER_STRUCT(FParameters, )
@@ -34,11 +36,13 @@ public:
 	static void						ModifyCompilationEnvironment(const FGlobalShaderPermutationParameters& params, FShaderCompilerEnvironment& env) {}
 };
 
-class TEXTUREGRAPHENGINE_API T_Invert
+class T_Invert
 {
 public:
 	//////////////////////////////////////////////////////////////////////////
 	/// Static functions
 	//////////////////////////////////////////////////////////////////////////
-	static TiledBlobPtr				Create(MixUpdateCyclePtr Cycle, BufferDescriptor DesiredOutputDesc, TiledBlobPtr Source, float MaxValue, int32 TargetId, bool bIncludeAlpha, bool bClamp);
+	static UE_API TiledBlobPtr				Create(MixUpdateCyclePtr Cycle, BufferDescriptor DesiredOutputDesc, TiledBlobPtr Source, float MaxValue, int32 TargetId, bool bIncludeAlpha, bool bClamp);
 };
+
+#undef UE_API
diff --git a/Engine/Plugins/TextureGraph/Source/TextureGraphEngine/Transform/Expressions/T_Levels.h b/Engine/Plugins/TextureGraph/Source/TextureGraphEngine/Transform/Expressions/T_Levels.h
--- a/Engine/Plugins/TextureGraph/Source/TextureGraphEngine/Transform/Expressions/T_Levels.h
+++ b/Engine/Plugins/TextureGraph/Source/TextureGraphEngine/Transform/Expressions/T_Levels.h
@@ -7,15 +7,17 @@
 #include "Model/Mix/MixUpdateCycle.h"
 #include <DataDrivenShaderPlatformInfo.h>
 
+#define UE_API TEXTUREGRAPHENGINE_API
+
 
 //////////////////////////////////////////////////////////////////////////
 /// Simple levels shader
 //////////////////////////////////////////////////////////////////////////
-class TEXTUREGRAPHENGINE_API FSH_Levels : public FSH_Base
+class FSH_Levels : public FSH_Base
 {
 public:
 	SHADER_USE_PARAMETER_STRUCT(FSH_Levels, FSH_Base);
-	DECLARE_GLOBAL_SHADER(FSH_Levels);
+	DECLARE_EXPORTED_GLOBAL_SHADER(FSH_Levels, UE_API);
 
 	BEGIN_SHADER_PARAMETER_STRUCT(FParameters, )
 		SHADER_PARAMETER_STRUCT(FStandardSamplerStates, SamplerStates)
@@ -44,13 +46,13 @@ public:
 };
 
 
-struct TEXTUREGRAPHENGINE_API FLevelsHistogramData
+struct FLevelsHistogramData
 {
 	TArray<FVector4f> HistogramData;
 };
 
 
-struct TEXTUREGRAPHENGINE_API FLevels
+struct FLevels
 {
 
 	// The Low value of the Levels adjustment, any pixel under that value is set to black. Default is 0.
@@ -72,32 +74,33 @@ struct TEXTUREGRAPHENGINE_API FLevels
 	// The white point of the output. Moving this will remap the white point to this value. Default is 1
 	float								OutHigh = 1;
 
-	bool SetLow(float InValue);
-	bool SetMid(float InValue);
-	bool SetHigh(float InValue);
+	UE_API bool SetLow(float InValue);
+	UE_API bool SetMid(float InValue);
+	UE_API bool SetHigh(float InValue);
 
 	// Eval Low-High range mapping on value
-	float EvalRange(float Val) const;
+	UE_API float EvalRange(float Val) const;
 	// Eval reverse Low-High range mapping on value
-	float EvalRangeInv(float Val) const;
+	UE_API float EvalRangeInv(float Val) const;
 
 	// Evaluate the MidExponent of the power curve applying the midpoint filter
-	float EvalMidExponent() const;
-	bool SetMidFromMidExponent(float InExponent);
+	UE_API float EvalMidExponent() const;
+	UE_API bool SetMidFromMidExponent(float InExponent);
 
-	void InitFromLowMidHigh(float LowValue, float MidValue, float HighValue, float OutLow, float OutHigh);
-	void InitFromAutoLevels(float MidPercentage);
-	void InitFromPositionContrast(float MidValue, float Contrast);
-	void InitFromRange(float InRange, float InPosition);
+	UE_API void InitFromLowMidHigh(float LowValue, float MidValue, float HighValue, float OutLow, float OutHigh);
+	UE_API void InitFromAutoLevels(float MidPercentage);
+	UE_API void InitFromPositionContrast(float MidValue, float Contrast);
+	UE_API void InitFromRange(float InRange, float InPosition);
 
 	FLevelsHistogramData HistogramData;
 };
 
 using FLevelsPtr = TSharedPtr< FLevels>;
 
-class TEXTUREGRAPHENGINE_API T_Levels
+class T_Levels
 {
 public:
-	static TiledBlobPtr				Create(MixUpdateCyclePtr Cycle, BufferDescriptor DesiredOutputDesc, TiledBlobPtr Source, const FLevelsPtr& Levels, int32 TargetId);
+	static UE_API TiledBlobPtr				Create(MixUpdateCyclePtr Cycle, BufferDescriptor DesiredOutputDesc, TiledBlobPtr Source, const FLevelsPtr& Levels, int32 TargetId);
 };
 
+#undef UE_API
diff --git a/Engine/Plugins/TextureGraph/Source/TextureGraphEngine/Transform/Expressions/T_MaterialIDMask.h b/Engine/Plugins/TextureGraph/Source/TextureGraphEngine/Transform/Expressions/T_MaterialIDMask.h
--- a/Engine/Plugins/TextureGraph/Source/TextureGraphEngine/Transform/Expressions/T_MaterialIDMask.h
+++ b/Engine/Plugins/TextureGraph/Source/TextureGraphEngine/Transform/Expressions/T_MaterialIDMask.h
@@ -10,6 +10,8 @@
 #include "Model/Mix/MixUpdateCycle.h"
 #include <DataDrivenShaderPlatformInfo.h>
 
+#define UE_API TEXTUREGRAPHENGINE_API
+
 //////////////////////////////////////////////////////////////////////////
 /// Shader
 //////////////////////////////////////////////////////////////////////////
@@ -18,7 +20,7 @@ class FSH_MaterialIDMask : public FSH_Base
 {
 public:
 
-	DECLARE_GLOBAL_SHADER(FSH_MaterialIDMask);
+	DECLARE_EXPORTED_GLOBAL_SHADER(FSH_MaterialIDMask, UE_API);
 	SHADER_USE_PARAMETER_STRUCT(FSH_MaterialIDMask, FSH_Base);
 
 	BEGIN_SHADER_PARAMETER_STRUCT(FParameters, )
@@ -43,15 +45,17 @@ public:
 /**
  * 
  */
-class TEXTUREGRAPHENGINE_API T_MaterialIDMask
+class T_MaterialIDMask
 {
 public:
-	T_MaterialIDMask();
-	~T_MaterialIDMask();
+	UE_API T_MaterialIDMask();
+	UE_API ~T_MaterialIDMask();
 
 	//////////////////////////////////////////////////////////////////////////
 	/// Static functions
 	//////////////////////////////////////////////////////////////////////////
 
-	static TiledBlobPtr				Create(MixUpdateCyclePtr InCycle, TiledBlobPtr InMaterialIDTexture, const TArray<FLinearColor>& InActiveColors, const int32& ActiveColorsCount, BufferDescriptor DesiredDesc, int InTargetId);
+	static UE_API TiledBlobPtr				Create(MixUpdateCyclePtr InCycle, TiledBlobPtr InMaterialIDTexture, const TArray<FLinearColor>& InActiveColors, const int32& ActiveColorsCount, BufferDescriptor DesiredDesc, int InTargetId);
 };
+
+#undef UE_API
diff --git a/Engine/Plugins/TextureGraph/Source/TextureGraphEngine/Transform/Expressions/T_Maths_OneInput.h b/Engine/Plugins/TextureGraph/Source/TextureGraphEngine/Transform/Expressions/T_Maths_OneInput.h
--- a/Engine/Plugins/TextureGraph/Source/TextureGraphEngine/Transform/Expressions/T_Maths_OneInput.h
+++ b/Engine/Plugins/TextureGraph/Source/TextureGraphEngine/Transform/Expressions/T_Maths_OneInput.h
@@ -9,10 +9,12 @@
 
 #include "T_Maths_OneInput.generated.h"
 
+#define UE_API TEXTUREGRAPHENGINE_API
+
 //////////////////////////////////////////////////////////////////////////
 /// Basic one input maths op
 //////////////////////////////////////////////////////////////////////////
-class TEXTUREGRAPHENGINE_API FSH_MathsOp_OneInput : public FSH_Base
+class FSH_MathsOp_OneInput : public FSH_Base
 {
 public:
 	SHADER_USE_PARAMETER_STRUCT(FSH_MathsOp_OneInput, FSH_Base);
@@ -67,25 +69,27 @@ enum class ETrigFunction : uint8
 //////////////////////////////////////////////////////////////////////////
 /// Helper C++ class
 //////////////////////////////////////////////////////////////////////////
-class TEXTUREGRAPHENGINE_API T_Maths_OneInput
+class T_Maths_OneInput
 {
 public:
 	//////////////////////////////////////////////////////////////////////////
 	/// Static functions
 	//////////////////////////////////////////////////////////////////////////
-	static TiledBlobPtr			CreateTrigonometry(MixUpdateCyclePtr Cycle, BufferDescriptor DesiredOutputDesc, int32 TargetId, TiledBlobPtr Input, ETrigFunction Function);
-	static TiledBlobPtr			CreateToRadians(MixUpdateCyclePtr Cycle, BufferDescriptor DesiredOutputDesc, int32 TargetId, TiledBlobPtr Input);
-	static TiledBlobPtr			CreateToDegrees(MixUpdateCyclePtr Cycle, BufferDescriptor DesiredOutputDesc, int32 TargetId, TiledBlobPtr Input);
-	static TiledBlobPtr			CreateAbs(MixUpdateCyclePtr Cycle, BufferDescriptor DesiredOutputDesc, int32 TargetId, TiledBlobPtr Input);
-	static TiledBlobPtr			CreateSqrt(MixUpdateCyclePtr Cycle, BufferDescriptor DesiredOutputDesc, int32 TargetId, TiledBlobPtr Input);
-	static TiledBlobPtr			CreateSquare(MixUpdateCyclePtr Cycle, BufferDescriptor DesiredOutputDesc, int32 TargetId, TiledBlobPtr Input);
-	static TiledBlobPtr			CreateCbrt(MixUpdateCyclePtr Cycle, BufferDescriptor DesiredOutputDesc, int32 TargetId, TiledBlobPtr Input);
-	static TiledBlobPtr			CreateCube(MixUpdateCyclePtr Cycle, BufferDescriptor DesiredOutputDesc, int32 TargetId, TiledBlobPtr Input);
-	static TiledBlobPtr			CreateExp(MixUpdateCyclePtr Cycle, BufferDescriptor DesiredOutputDesc, int32 TargetId, TiledBlobPtr Input);
-	static TiledBlobPtr			CreateLog2(MixUpdateCyclePtr Cycle, BufferDescriptor DesiredOutputDesc, int32 TargetId, TiledBlobPtr Input);
-	static TiledBlobPtr			CreateLog10(MixUpdateCyclePtr Cycle, BufferDescriptor DesiredOutputDesc, int32 TargetId, TiledBlobPtr Input);
-	static TiledBlobPtr			CreateLog(MixUpdateCyclePtr Cycle, BufferDescriptor DesiredOutputDesc, int32 TargetId, TiledBlobPtr Input);
-	static TiledBlobPtr			CreateFloor(MixUpdateCyclePtr Cycle, BufferDescriptor DesiredOutputDesc, int32 TargetId, TiledBlobPtr Input);
-	static TiledBlobPtr			CreateCeil(MixUpdateCyclePtr Cycle, BufferDescriptor DesiredOutputDesc, int32 TargetId, TiledBlobPtr Input);
-	static TiledBlobPtr			CreateRound(MixUpdateCyclePtr Cycle, BufferDescriptor DesiredOutputDesc, int32 TargetId, TiledBlobPtr Input);
+	static UE_API TiledBlobPtr			CreateTrigonometry(MixUpdateCyclePtr Cycle, BufferDescriptor DesiredOutputDesc, int32 TargetId, TiledBlobPtr Input, ETrigFunction Function);
+	static UE_API TiledBlobPtr			CreateToRadians(MixUpdateCyclePtr Cycle, BufferDescriptor DesiredOutputDesc, int32 TargetId, TiledBlobPtr Input);
+	static UE_API TiledBlobPtr			CreateToDegrees(MixUpdateCyclePtr Cycle, BufferDescriptor DesiredOutputDesc, int32 TargetId, TiledBlobPtr Input);
+	static UE_API TiledBlobPtr			CreateAbs(MixUpdateCyclePtr Cycle, BufferDescriptor DesiredOutputDesc, int32 TargetId, TiledBlobPtr Input);
+	static UE_API TiledBlobPtr			CreateSqrt(MixUpdateCyclePtr Cycle, BufferDescriptor DesiredOutputDesc, int32 TargetId, TiledBlobPtr Input);
+	static UE_API TiledBlobPtr			CreateSquare(MixUpdateCyclePtr Cycle, BufferDescriptor DesiredOutputDesc, int32 TargetId, TiledBlobPtr Input);
+	static UE_API TiledBlobPtr			CreateCbrt(MixUpdateCyclePtr Cycle, BufferDescriptor DesiredOutputDesc, int32 TargetId, TiledBlobPtr Input);
+	static UE_API TiledBlobPtr			CreateCube(MixUpdateCyclePtr Cycle, BufferDescriptor DesiredOutputDesc, int32 TargetId, TiledBlobPtr Input);
+	static UE_API TiledBlobPtr			CreateExp(MixUpdateCyclePtr Cycle, BufferDescriptor DesiredOutputDesc, int32 TargetId, TiledBlobPtr Input);
+	static UE_API TiledBlobPtr			CreateLog2(MixUpdateCyclePtr Cycle, BufferDescriptor DesiredOutputDesc, int32 TargetId, TiledBlobPtr Input);
+	static UE_API TiledBlobPtr			CreateLog10(MixUpdateCyclePtr Cycle, BufferDescriptor DesiredOutputDesc, int32 TargetId, TiledBlobPtr Input);
+	static UE_API TiledBlobPtr			CreateLog(MixUpdateCyclePtr Cycle, BufferDescriptor DesiredOutputDesc, int32 TargetId, TiledBlobPtr Input);
+	static UE_API TiledBlobPtr			CreateFloor(MixUpdateCyclePtr Cycle, BufferDescriptor DesiredOutputDesc, int32 TargetId, TiledBlobPtr Input);
+	static UE_API TiledBlobPtr			CreateCeil(MixUpdateCyclePtr Cycle, BufferDescriptor DesiredOutputDesc, int32 TargetId, TiledBlobPtr Input);
+	static UE_API TiledBlobPtr			CreateRound(MixUpdateCyclePtr Cycle, BufferDescriptor DesiredOutputDesc, int32 TargetId, TiledBlobPtr Input);
 };
+
+#undef UE_API
diff --git a/Engine/Plugins/TextureGraph/Source/TextureGraphEngine/Transform/Expressions/T_Maths_ThreeInputs.h b/Engine/Plugins/TextureGraph/Source/TextureGraphEngine/Transform/Expressions/T_Maths_ThreeInputs.h
--- a/Engine/Plugins/TextureGraph/Source/TextureGraphEngine/Transform/Expressions/T_Maths_ThreeInputs.h
+++ b/Engine/Plugins/TextureGraph/Source/TextureGraphEngine/Transform/Expressions/T_Maths_ThreeInputs.h
@@ -7,12 +7,14 @@
 #include "Model/Mix/MixUpdateCycle.h"
 #include <DataDrivenShaderPlatformInfo.h>
 
+#define UE_API TEXTUREGRAPHENGINE_API
+
 //#include "T_Maths_ThreeInputs.generated.h"
 
 //////////////////////////////////////////////////////////////////////////
 /// MAD
 //////////////////////////////////////////////////////////////////////////
-class TEXTUREGRAPHENGINE_API FSH_ThreeInputs : public FSH_Base
+class FSH_ThreeInputs : public FSH_Base
 {
 public:
 	SHADER_USE_PARAMETER_STRUCT(FSH_ThreeInputs, FSH_Base);
@@ -34,12 +36,14 @@ DECLARE_EMPTY_GLOBAL_SHADER_DERIVED_FROM(FSH_Lerp, FSH_ThreeInputs);
 //////////////////////////////////////////////////////////////////////////
 /// Helper C++ class
 //////////////////////////////////////////////////////////////////////////
-class TEXTUREGRAPHENGINE_API T_Maths_ThreeInputs
+class T_Maths_ThreeInputs
 {
 public:
 	//////////////////////////////////////////////////////////////////////////
 	/// Static functions
 	//////////////////////////////////////////////////////////////////////////
-	static TiledBlobPtr				CreateMad(MixUpdateCyclePtr Cycle, BufferDescriptor DesiredOutputDesc, int32 TargetId, TiledBlobPtr Operand1, TiledBlobPtr Operand2, TiledBlobPtr Operand3);
-	static TiledBlobPtr				CreateLerp(MixUpdateCyclePtr Cycle, BufferDescriptor DesiredOutputDesc, int32 TargetId, TiledBlobPtr Operand1, TiledBlobPtr Operand2, TiledBlobPtr Operand3);
+	static UE_API TiledBlobPtr				CreateMad(MixUpdateCyclePtr Cycle, BufferDescriptor DesiredOutputDesc, int32 TargetId, TiledBlobPtr Operand1, TiledBlobPtr Operand2, TiledBlobPtr Operand3);
+	static UE_API TiledBlobPtr				CreateLerp(MixUpdateCyclePtr Cycle, BufferDescriptor DesiredOutputDesc, int32 TargetId, TiledBlobPtr Operand1, TiledBlobPtr Operand2, TiledBlobPtr Operand3);
 };
+
+#undef UE_API
diff --git a/Engine/Plugins/TextureGraph/Source/TextureGraphEngine/Transform/Expressions/T_Maths_TwoInputs.h b/Engine/Plugins/TextureGraph/Source/TextureGraphEngine/Transform/Expressions/T_Maths_TwoInputs.h
--- a/Engine/Plugins/TextureGraph/Source/TextureGraphEngine/Transform/Expressions/T_Maths_TwoInputs.h
+++ b/Engine/Plugins/TextureGraph/Source/TextureGraphEngine/Transform/Expressions/T_Maths_TwoInputs.h
@@ -9,10 +9,12 @@
 
 #include "T_Maths_TwoInputs.generated.h"
 
+#define UE_API TEXTUREGRAPHENGINE_API
+
 //////////////////////////////////////////////////////////////////////////
 /// Basic math op
 //////////////////////////////////////////////////////////////////////////
-class TEXTUREGRAPHENGINE_API FSH_BasicMathOp : public FSH_Base
+class FSH_BasicMathOp : public FSH_Base
 {
 public:
 	SHADER_USE_PARAMETER_STRUCT(FSH_BasicMathOp, FSH_Base);
@@ -59,11 +61,11 @@ enum EIfThenElseType : int
 	Grayscale			UMETA(DisplayName = "Grayscale"),
 };
 
-class TEXTUREGRAPHENGINE_API FSH_IfThenElse : public FSH_Base
+class FSH_IfThenElse : public FSH_Base
 {
 public:
 	SHADER_USE_PARAMETER_STRUCT(FSH_IfThenElse, FSH_Base);
-	DECLARE_GLOBAL_SHADER(FSH_IfThenElse);
+	DECLARE_EXPORTED_GLOBAL_SHADER(FSH_IfThenElse, UE_API);
 
 	BEGIN_SHADER_PARAMETER_STRUCT(FParameters, )
 		SHADER_PARAMETER_STRUCT(FStandardSamplerStates, SamplerStates)
@@ -107,21 +109,23 @@ DECLARE_EMPTY_GLOBAL_SHADER_DERIVED_FROM(FSH_IfThenElse_NEQ_Grayscale, FSH_IfThe
 //////////////////////////////////////////////////////////////////////////
 /// Helper C++ class
 //////////////////////////////////////////////////////////////////////////
-class TEXTUREGRAPHENGINE_API T_Maths_TwoInputs
+class T_Maths_TwoInputs
 {
 public:
 	//////////////////////////////////////////////////////////////////////////
 	/// Static functions
 	//////////////////////////////////////////////////////////////////////////
-	static TiledBlobPtr				CreateMultiply(MixUpdateCyclePtr Cycle, BufferDescriptor DesiredOutputDesc, int32 TargetId, TiledBlobPtr Operand1, TiledBlobPtr Operand2);
-	static TiledBlobPtr				CreateDivide(MixUpdateCyclePtr Cycle, BufferDescriptor DesiredOutputDesc, int32 TargetId, TiledBlobPtr Operand1, TiledBlobPtr Operand2);
-	static TiledBlobPtr				CreateAdd(MixUpdateCyclePtr Cycle, BufferDescriptor DesiredOutputDesc, int32 TargetId, TiledBlobPtr Operand1, TiledBlobPtr Operand2);
-	static TiledBlobPtr				CreateSubtract(MixUpdateCyclePtr Cycle, BufferDescriptor DesiredOutputDesc, int32 TargetId, TiledBlobPtr Operand1, TiledBlobPtr Operand2);
-	static TiledBlobPtr				CreateDot(MixUpdateCyclePtr Cycle, BufferDescriptor DesiredOutputDesc, int32 TargetId, TiledBlobPtr Operand1, TiledBlobPtr Operand2);
-	static TiledBlobPtr				CreateCross(MixUpdateCyclePtr Cycle, BufferDescriptor DesiredOutputDesc, int32 TargetId, TiledBlobPtr Operand1, TiledBlobPtr Operand2);
-	static TiledBlobPtr				CreatePow(MixUpdateCyclePtr Cycle, BufferDescriptor DesiredOutputDesc, int32 TargetId, TiledBlobPtr Operand1, TiledBlobPtr Operand2);
-	static TiledBlobPtr				CreateStep(MixUpdateCyclePtr Cycle, BufferDescriptor DesiredOutputDesc, int32 TargetId, TiledBlobPtr Operand1, TiledBlobPtr Operand2);
+	static UE_API TiledBlobPtr				CreateMultiply(MixUpdateCyclePtr Cycle, BufferDescriptor DesiredOutputDesc, int32 TargetId, TiledBlobPtr Operand1, TiledBlobPtr Operand2);
+	static UE_API TiledBlobPtr				CreateDivide(MixUpdateCyclePtr Cycle, BufferDescriptor DesiredOutputDesc, int32 TargetId, TiledBlobPtr Operand1, TiledBlobPtr Operand2);
+	static UE_API TiledBlobPtr				CreateAdd(MixUpdateCyclePtr Cycle, BufferDescriptor DesiredOutputDesc, int32 TargetId, TiledBlobPtr Operand1, TiledBlobPtr Operand2);
+	static UE_API TiledBlobPtr				CreateSubtract(MixUpdateCyclePtr Cycle, BufferDescriptor DesiredOutputDesc, int32 TargetId, TiledBlobPtr Operand1, TiledBlobPtr Operand2);
+	static UE_API TiledBlobPtr				CreateDot(MixUpdateCyclePtr Cycle, BufferDescriptor DesiredOutputDesc, int32 TargetId, TiledBlobPtr Operand1, TiledBlobPtr Operand2);
+	static UE_API TiledBlobPtr				CreateCross(MixUpdateCyclePtr Cycle, BufferDescriptor DesiredOutputDesc, int32 TargetId, TiledBlobPtr Operand1, TiledBlobPtr Operand2);
+	static UE_API TiledBlobPtr				CreatePow(MixUpdateCyclePtr Cycle, BufferDescriptor DesiredOutputDesc, int32 TargetId, TiledBlobPtr Operand1, TiledBlobPtr Operand2);
+	static UE_API TiledBlobPtr				CreateStep(MixUpdateCyclePtr Cycle, BufferDescriptor DesiredOutputDesc, int32 TargetId, TiledBlobPtr Operand1, TiledBlobPtr Operand2);
 	
-	static TiledBlobPtr				CreateIfThenElse(MixUpdateCyclePtr Cycle, BufferDescriptor DesiredOutputDesc, int32 TargetId,
+	static UE_API TiledBlobPtr				CreateIfThenElse(MixUpdateCyclePtr Cycle, BufferDescriptor DesiredOutputDesc, int32 TargetId,
 		TiledBlobPtr LHS, TiledBlobPtr RHS, TiledBlobPtr Then, TiledBlobPtr Else, EIfThenElseOperator Operator, EIfThenElseType Type);
 };
+
+#undef UE_API
diff --git a/Engine/Plugins/TextureGraph/Source/TextureGraphEngine/Transform/Expressions/T_Transform.h b/Engine/Plugins/TextureGraph/Source/TextureGraphEngine/Transform/Expressions/T_Transform.h
--- a/Engine/Plugins/TextureGraph/Source/TextureGraphEngine/Transform/Expressions/T_Transform.h
+++ b/Engine/Plugins/TextureGraph/Source/TextureGraphEngine/Transform/Expressions/T_Transform.h
@@ -9,14 +9,16 @@
 #include "Model/Mix/MixUpdateCycle.h"
 #include <DataDrivenShaderPlatformInfo.h>
 
+#define UE_API TEXTUREGRAPHENGINE_API
+
 //////////////////////////////////////////////////////////////////////////
 /// Transform blit shader
 //////////////////////////////////////////////////////////////////////////
-class TEXTUREGRAPHENGINE_API FSH_TransformBlit : public FSH_Base
+class FSH_TransformBlit : public FSH_Base
 {
 public:
 	SHADER_USE_PARAMETER_STRUCT(FSH_TransformBlit, FSH_Base);
-	DECLARE_GLOBAL_SHADER(FSH_TransformBlit);
+	DECLARE_EXPORTED_GLOBAL_SHADER(FSH_TransformBlit, UE_API);
 
 	BEGIN_SHADER_PARAMETER_STRUCT(FParameters, )
 		SHADER_PARAMETER_STRUCT(FTileInfo, TileInfo)
@@ -55,7 +57,7 @@ public:
 };
 
 //////////////////////////////////////////////////////////////////////////
-class TEXTUREGRAPHENGINE_API T_Transform
+class T_Transform
 {
 public:
 	//////////////////////////////////////////////////////////////////////////
@@ -89,10 +91,12 @@ public:
 		float ShowDebugGrid = 0.0f;
 	};
 
-	static TiledBlobPtr		Create(MixUpdateCyclePtr Cycle, BufferDescriptor DesiredOutputDesc, TiledBlobPtr Source,
+	static UE_API TiledBlobPtr		Create(MixUpdateCyclePtr Cycle, BufferDescriptor DesiredOutputDesc, TiledBlobPtr Source,
 								const TransformParameter& TransformParam, 
 								const CellParameter& CellParam,
 								const ColorParameter& ColorParam,
 								int32 TargetId = 0);
 
 };
+
+#undef UE_API
diff --git a/Engine/Plugins/TextureGraph/Source/TextureGraphEngine/Transform/Layer/T_AdjustColor.h b/Engine/Plugins/TextureGraph/Source/TextureGraphEngine/Transform/Layer/T_AdjustColor.h
--- a/Engine/Plugins/TextureGraph/Source/TextureGraphEngine/Transform/Layer/T_AdjustColor.h
+++ b/Engine/Plugins/TextureGraph/Source/TextureGraphEngine/Transform/Layer/T_AdjustColor.h
@@ -4,6 +4,8 @@
 #include "FxMat/FxMaterial.h"
 #include <DataDrivenShaderPlatformInfo.h>
 
+#define UE_API TEXTUREGRAPHENGINE_API
+
 class ULayer_Textured;
 class MixUpdateCycle;
 typedef std::shared_ptr<MixUpdateCycle> MixUpdateCyclePtr;
@@ -19,7 +21,7 @@ class FSH_AdjustColor : public FSH_Base
 {
 public:
 
-	DECLARE_GLOBAL_SHADER(FSH_AdjustColor);
+	DECLARE_EXPORTED_GLOBAL_SHADER(FSH_AdjustColor, UE_API);
 	SHADER_USE_PARAMETER_STRUCT(FSH_AdjustColor, FSH_Base);
 
 	BEGIN_SHADER_PARAMETER_STRUCT(FParameters, )
@@ -49,13 +51,15 @@ public:
 
 template <> void SetupDefaultParameters(FSH_AdjustColor::FParameters& params);
 //////////////////////////////////////////////////////////////////////////
-class TEXTUREGRAPHENGINE_API T_AdjustColor 
+class T_AdjustColor 
 {
 public:
-									T_AdjustColor();
-	virtual							~T_AdjustColor();
+									UE_API T_AdjustColor();
+	UE_API virtual							~T_AdjustColor();
 
 	//////////////////////////////////////////////////////////////////////////
 	/// Static functions
 	//////////////////////////////////////////////////////////////////////////
 };
+
+#undef UE_API
diff --git a/Engine/Plugins/TextureGraph/Source/TextureGraphEngine/Transform/Layer/T_AdjustFrequency.h b/Engine/Plugins/TextureGraph/Source/TextureGraphEngine/Transform/Layer/T_AdjustFrequency.h
--- a/Engine/Plugins/TextureGraph/Source/TextureGraphEngine/Transform/Layer/T_AdjustFrequency.h
+++ b/Engine/Plugins/TextureGraph/Source/TextureGraphEngine/Transform/Layer/T_AdjustFrequency.h
@@ -7,11 +7,13 @@
 #include "Model/Mix/MixUpdateCycle.h"
 #include <DataDrivenShaderPlatformInfo.h>
 
+#define UE_API TEXTUREGRAPHENGINE_API
+
 
 class FSH_AdjustFrequency : public FSH_Base
 {
 public:
-	DECLARE_GLOBAL_SHADER(FSH_AdjustFrequency);
+	DECLARE_EXPORTED_GLOBAL_SHADER(FSH_AdjustFrequency, UE_API);
 	SHADER_USE_PARAMETER_STRUCT(FSH_AdjustFrequency, FSH_Base);
 
 	BEGIN_SHADER_PARAMETER_STRUCT(FParameters, )
@@ -48,7 +50,7 @@ namespace {
 class FSH_AdjustFrequencyNormals : public FSH_Base
 {
 public:
-	DECLARE_GLOBAL_SHADER(FSH_AdjustFrequencyNormals);
+	DECLARE_EXPORTED_GLOBAL_SHADER(FSH_AdjustFrequencyNormals, UE_API);
 	SHADER_USE_PARAMETER_STRUCT(FSH_AdjustFrequencyNormals, FSH_Base);
 
 	BEGIN_SHADER_PARAMETER_STRUCT(FParameters, )
@@ -85,13 +87,15 @@ class ULayer_Textured;
 //////////////////////////////////////////////////////////////////////////
 /// Height Mask Transformation Function
 //////////////////////////////////////////////////////////////////////////
-class TEXTUREGRAPHENGINE_API T_AdjustFrequency
+class T_AdjustFrequency
 {
 public:
-	T_AdjustFrequency();
-	virtual							~T_AdjustFrequency();
+	UE_API T_AdjustFrequency();
+	UE_API virtual							~T_AdjustFrequency();
 
 	//////////////////////////////////////////////////////////////////////////
 	/// Static functions
 	//////////////////////////////////////////////////////////////////////////
 };
+
+#undef UE_API
diff --git a/Engine/Plugins/TextureGraph/Source/TextureGraphEngine/Transform/Layer/T_BlendWithHeightMask.h b/Engine/Plugins/TextureGraph/Source/TextureGraphEngine/Transform/Layer/T_BlendWithHeightMask.h
--- a/Engine/Plugins/TextureGraph/Source/TextureGraphEngine/Transform/Layer/T_BlendWithHeightMask.h
+++ b/Engine/Plugins/TextureGraph/Source/TextureGraphEngine/Transform/Layer/T_BlendWithHeightMask.h
@@ -7,6 +7,8 @@
 #include "Model/Mix/MixUpdateCycle.h"
 #include <DataDrivenShaderPlatformInfo.h>
 
+#define UE_API TEXTUREGRAPHENGINE_API
+
 //////////////////////////////////////////////////////////////////////////
 class ULayer_Textured;
 namespace
@@ -15,7 +17,7 @@ namespace
 class FSH_NormalBlendWithHeightMask : public FSH_Base
 {
 public:
-	DECLARE_GLOBAL_SHADER(FSH_NormalBlendWithHeightMask);
+	DECLARE_EXPORTED_GLOBAL_SHADER(FSH_NormalBlendWithHeightMask, UE_API);
 	SHADER_USE_PARAMETER_STRUCT(FSH_NormalBlendWithHeightMask, FSH_Base);
 
 	BEGIN_SHADER_PARAMETER_STRUCT(FParameters, )
@@ -44,10 +46,12 @@ template <> void SetupDefaultParameters(FSH_NormalBlendWithHeightMask::FParamete
 typedef FxMaterial_Normal<VSH_Simple, FSH_NormalBlendWithHeightMask>	Fx_NormalBlendWithHeightMask;
 
 //////////////////////////////////////////////////////////////////////////
-class TEXTUREGRAPHENGINE_API T_BlendWithHeightMask 
+class T_BlendWithHeightMask 
 {
 public:
 	//////////////////////////////////////////////////////////////////////////
 	/// Static functions
 	//////////////////////////////////////////////////////////////////////////
 };
+
+#undef UE_API
diff --git a/Engine/Plugins/TextureGraph/Source/TextureGraphEngine/Transform/Layer/T_Displacement.h b/Engine/Plugins/TextureGraph/Source/TextureGraphEngine/Transform/Layer/T_Displacement.h
--- a/Engine/Plugins/TextureGraph/Source/TextureGraphEngine/Transform/Layer/T_Displacement.h
+++ b/Engine/Plugins/TextureGraph/Source/TextureGraphEngine/Transform/Layer/T_Displacement.h
@@ -7,12 +7,14 @@
 #include "Model/Mix/MixUpdateCycle.h"
 #include <DataDrivenShaderPlatformInfo.h>
 
+#define UE_API TEXTUREGRAPHENGINE_API
+
 class ULayer_Textured;
 
 class FSH_LayerDisplacement : public FSH_Base
 {
 public:
-	DECLARE_GLOBAL_SHADER(FSH_LayerDisplacement);
+	DECLARE_EXPORTED_GLOBAL_SHADER(FSH_LayerDisplacement, UE_API);
 	SHADER_USE_PARAMETER_STRUCT(FSH_LayerDisplacement, FSH_Base);
 
 	BEGIN_SHADER_PARAMETER_STRUCT(FParameters, )
@@ -48,13 +50,15 @@ template <> void SetupDefaultParameters(FSH_LayerDisplacement::FParameters& para
 typedef FxMaterial_Normal<VSH_Simple, FSH_LayerDisplacement> Fx_LayerDisplacement;
 
 //////////////////////////////////////////////////////////////////////////
-class TEXTUREGRAPHENGINE_API T_Displacement 
+class T_Displacement 
 {
 public:
-									T_Displacement();
-	virtual							~T_Displacement();
+									UE_API T_Displacement();
+	UE_API virtual							~T_Displacement();
 
 	//////////////////////////////////////////////////////////////////////////
 	/// Static functions
 	//////////////////////////////////////////////////////////////////////////
 };
+
+#undef UE_API
diff --git a/Engine/Plugins/TextureGraph/Source/TextureGraphEngine/Transform/Layer/T_Dynamic.h b/Engine/Plugins/TextureGraph/Source/TextureGraphEngine/Transform/Layer/T_Dynamic.h
--- a/Engine/Plugins/TextureGraph/Source/TextureGraphEngine/Transform/Layer/T_Dynamic.h
+++ b/Engine/Plugins/TextureGraph/Source/TextureGraphEngine/Transform/Layer/T_Dynamic.h
@@ -8,7 +8,7 @@
 class ULayer_Textured;
 class UStackComponent;
 //////////////////////////////////////////////////////////////////////////
-class TEXTUREGRAPHENGINE_API T_Dynamic 
+class T_Dynamic 
 {
 public:
 
diff --git a/Engine/Plugins/TextureGraph/Source/TextureGraphEngine/Transform/Layer/T_LayerMask.h b/Engine/Plugins/TextureGraph/Source/TextureGraphEngine/Transform/Layer/T_LayerMask.h
--- a/Engine/Plugins/TextureGraph/Source/TextureGraphEngine/Transform/Layer/T_LayerMask.h
+++ b/Engine/Plugins/TextureGraph/Source/TextureGraphEngine/Transform/Layer/T_LayerMask.h
@@ -8,6 +8,8 @@
 #include "ShaderPermutation.h"
 #include <DataDrivenShaderPlatformInfo.h>
 
+#define UE_API TEXTUREGRAPHENGINE_API
+
 //////////////////////////////////////////////////////////////////////////
 /// Shader
 //////////////////////////////////////////////////////////////////////////
@@ -30,7 +32,7 @@ namespace
 class FSH_LayerMask : public FSH_Base
 {
 public:
-	DECLARE_GLOBAL_SHADER(FSH_LayerMask);
+	DECLARE_EXPORTED_GLOBAL_SHADER(FSH_LayerMask, UE_API);
 	SHADER_USE_PARAMETER_STRUCT(FSH_LayerMask, FSH_Base);
 
 	BEGIN_SHADER_PARAMETER_STRUCT(FParameters, )
@@ -86,13 +88,15 @@ class ULayer_Textured;
 //////////////////////////////////////////////////////////////////////////
 /// Height Mask Transformation Function
 //////////////////////////////////////////////////////////////////////////
-class TEXTUREGRAPHENGINE_API T_LayerMask
+class T_LayerMask
 {
 public:
-									T_LayerMask();
-	virtual							~T_LayerMask();
+									UE_API T_LayerMask();
+	UE_API virtual							~T_LayerMask();
 
 	//////////////////////////////////////////////////////////////////////////
 	/// Static functions
 	//////////////////////////////////////////////////////////////////////////
 };
+
+#undef UE_API
diff --git a/Engine/Plugins/TextureGraph/Source/TextureGraphEngine/Transform/Layer/T_Normal.h b/Engine/Plugins/TextureGraph/Source/TextureGraphEngine/Transform/Layer/T_Normal.h
--- a/Engine/Plugins/TextureGraph/Source/TextureGraphEngine/Transform/Layer/T_Normal.h
+++ b/Engine/Plugins/TextureGraph/Source/TextureGraphEngine/Transform/Layer/T_Normal.h
@@ -7,12 +7,14 @@
 #include "Model/Mix/MixUpdateCycle.h"
 #include <DataDrivenShaderPlatformInfo.h>
 
+#define UE_API TEXTUREGRAPHENGINE_API
+
 class ULayer_Textured;
 
 class FSH_LayerNormal : public FSH_Base
 {
 public:
-	DECLARE_GLOBAL_SHADER(FSH_LayerNormal);
+	DECLARE_EXPORTED_GLOBAL_SHADER(FSH_LayerNormal, UE_API);
 	SHADER_USE_PARAMETER_STRUCT(FSH_LayerNormal, FSH_Base);
 
 	BEGIN_SHADER_PARAMETER_STRUCT(FParameters, )
@@ -48,13 +50,15 @@ template <> void SetupDefaultParameters(FSH_LayerNormal::FParameters& params);
 typedef FxMaterial_Normal<VSH_Simple, FSH_LayerNormal> Fx_LayerNormal;
 
 //////////////////////////////////////////////////////////////////////////
-class TEXTUREGRAPHENGINE_API T_Normal 
+class T_Normal 
 {
 public:
-									T_Normal();
-	virtual							~T_Normal();
+									UE_API T_Normal();
+	UE_API virtual							~T_Normal();
 
 	//////////////////////////////////////////////////////////////////////////
 	/// Static functions
 	//////////////////////////////////////////////////////////////////////////
 };
+
+#undef UE_API
diff --git a/Engine/Plugins/TextureGraph/Source/TextureGraphEngine/Transform/Layer/T_Thumbnail.h b/Engine/Plugins/TextureGraph/Source/TextureGraphEngine/Transform/Layer/T_Thumbnail.h
--- a/Engine/Plugins/TextureGraph/Source/TextureGraphEngine/Transform/Layer/T_Thumbnail.h
+++ b/Engine/Plugins/TextureGraph/Source/TextureGraphEngine/Transform/Layer/T_Thumbnail.h
@@ -9,12 +9,16 @@
 #include "FxMat/MaterialManager.h"
 #include "TextureGraphEngine.h"
 
+#define UE_API TEXTUREGRAPHENGINE_API
+
 //////////////////////////////////////////////////////////////////////////
-class TEXTUREGRAPHENGINE_API T_Thumbnail
+class T_Thumbnail
 {
 public:
 	//////////////////////////////////////////////////////////////////////////
 	/// Static functions
 	//////////////////////////////////////////////////////////////////////////
-	static TiledBlobRef				Bind(UMixInterface* Mix, UObject* Model, TiledBlobPtr InBlobToBind, int32 TargetId);
+	static UE_API TiledBlobRef				Bind(UMixInterface* Mix, UObject* Model, TiledBlobPtr InBlobToBind, int32 TargetId);
 };
+
+#undef UE_API
diff --git a/Engine/Plugins/TextureGraph/Source/TextureGraphEngine/Transform/Mask/T_Gradient.h b/Engine/Plugins/TextureGraph/Source/TextureGraphEngine/Transform/Mask/T_Gradient.h
--- a/Engine/Plugins/TextureGraph/Source/TextureGraphEngine/Transform/Mask/T_Gradient.h
+++ b/Engine/Plugins/TextureGraph/Source/TextureGraphEngine/Transform/Mask/T_Gradient.h
@@ -9,6 +9,8 @@
 #include <DataDrivenShaderPlatformInfo.h>
 #include "T_Gradient.generated.h"
 
+#define UE_API TEXTUREGRAPHENGINE_API
+
 UENUM(BlueprintType)
 enum class EGradientType : uint8
 {
@@ -51,10 +53,10 @@ namespace
 }
 
 ///////////////////////////////////////////////////////////////////////////////
-class TEXTUREGRAPHENGINE_API FSH_GradientLinear_1 : public FSH_Base
+class FSH_GradientLinear_1 : public FSH_Base
 {
 public:
-	DECLARE_GLOBAL_SHADER(FSH_GradientLinear_1);
+	DECLARE_EXPORTED_GLOBAL_SHADER(FSH_GradientLinear_1, UE_API);
 	SHADER_USE_PARAMETER_STRUCT(FSH_GradientLinear_1, FSH_Base);
 	BEGIN_SHADER_PARAMETER_STRUCT(FParameters, )
 		SHADER_PARAMETER_STRUCT(FTileInfo, TileInfo)
@@ -67,10 +69,10 @@ public:
 };
 
 ///////////////////////////////////////////////////////////////////////////////
-class TEXTUREGRAPHENGINE_API FSH_GradientLinear_2 : public FSH_Base
+class FSH_GradientLinear_2 : public FSH_Base
 {
 public:
-	DECLARE_GLOBAL_SHADER(FSH_GradientLinear_2);
+	DECLARE_EXPORTED_GLOBAL_SHADER(FSH_GradientLinear_2, UE_API);
 	SHADER_USE_PARAMETER_STRUCT(FSH_GradientLinear_2, FSH_Base);
 	BEGIN_SHADER_PARAMETER_STRUCT(FParameters, )
 		SHADER_PARAMETER_STRUCT(FTileInfo, TileInfo)
@@ -83,10 +85,10 @@ public:
 };
 
 ///////////////////////////////////////////////////////////////////////////////
-class TEXTUREGRAPHENGINE_API FSH_GradientRadial : public FSH_Base
+class FSH_GradientRadial : public FSH_Base
 {
 public:
-	DECLARE_GLOBAL_SHADER(FSH_GradientRadial);
+	DECLARE_EXPORTED_GLOBAL_SHADER(FSH_GradientRadial, UE_API);
 	SHADER_USE_PARAMETER_STRUCT(FSH_GradientRadial, FSH_Base);
 	BEGIN_SHADER_PARAMETER_STRUCT(FParameters, )
 		SHADER_PARAMETER_STRUCT(FTileInfo, TileInfo)
@@ -100,10 +102,10 @@ public:
 };
 
 ///////////////////////////////////////////////////////////////////////////////
-class TEXTUREGRAPHENGINE_API FSH_GradientAxial1: public FSH_Base
+class FSH_GradientAxial1: public FSH_Base
 {
 public:
-	DECLARE_GLOBAL_SHADER(FSH_GradientAxial1);
+	DECLARE_EXPORTED_GLOBAL_SHADER(FSH_GradientAxial1, UE_API);
 	SHADER_USE_PARAMETER_STRUCT(FSH_GradientAxial1, FSH_Base);
 	BEGIN_SHADER_PARAMETER_STRUCT(FParameters, )
 		SHADER_PARAMETER_STRUCT(FTileInfo, TileInfo)
@@ -118,10 +120,10 @@ public:
 };
 
 ///////////////////////////////////////////////////////////////////////////////
-class TEXTUREGRAPHENGINE_API FSH_GradientAxial2: public FSH_Base
+class FSH_GradientAxial2: public FSH_Base
 {
 public:
-	DECLARE_GLOBAL_SHADER(FSH_GradientAxial2);
+	DECLARE_EXPORTED_GLOBAL_SHADER(FSH_GradientAxial2, UE_API);
 	SHADER_USE_PARAMETER_STRUCT(FSH_GradientAxial2, FSH_Base);
 	BEGIN_SHADER_PARAMETER_STRUCT(FParameters, )
 		SHADER_PARAMETER_STRUCT(FTileInfo, TileInfo)
@@ -139,11 +141,11 @@ public:
 /**
  * 
  */
-class TEXTUREGRAPHENGINE_API T_Gradient
+class T_Gradient
 {
 public:
 	static constexpr int		DefaultSize = 1024;
-	static BufferDescriptor		InitOutputDesc(BufferDescriptor DesiredOutputDesc);
+	static UE_API BufferDescriptor		InitOutputDesc(BufferDescriptor DesiredOutputDesc);
 
 	struct Params {
 		EGradientType			Type = EGradientType::GT_Linear_1;
@@ -155,5 +157,7 @@ public:
 		FVector2f				Point2;
 	};
 
-	static TiledBlobPtr	Create(MixUpdateCyclePtr Cycle, BufferDescriptor DesiredOutputDesc, const Params& InParams, int32 TargetId = 0);
+	static UE_API TiledBlobPtr	Create(MixUpdateCyclePtr Cycle, BufferDescriptor DesiredOutputDesc, const Params& InParams, int32 TargetId = 0);
 };
+
+#undef UE_API
diff --git a/Engine/Plugins/TextureGraph/Source/TextureGraphEngine/Transform/Mask/T_ImageMask.h b/Engine/Plugins/TextureGraph/Source/TextureGraphEngine/Transform/Mask/T_ImageMask.h
--- a/Engine/Plugins/TextureGraph/Source/TextureGraphEngine/Transform/Mask/T_ImageMask.h
+++ b/Engine/Plugins/TextureGraph/Source/TextureGraphEngine/Transform/Mask/T_ImageMask.h
@@ -7,6 +7,8 @@
 #include "Model/Mix/MixUpdateCycle.h"
 #include <DataDrivenShaderPlatformInfo.h>
 
+#define UE_API TEXTUREGRAPHENGINE_API
+
 //////////////////////////////////////////////////////////////////////////
 /// Shader
 //////////////////////////////////////////////////////////////////////////
@@ -17,7 +19,7 @@ public:
 	class FVar_NORMALIZE : SHADER_PERMUTATION_BOOL("NORMALIZE");
 	class FVar_INVERT : SHADER_PERMUTATION_BOOL("INVERT");
 
-	DECLARE_GLOBAL_SHADER(FSH_ImageMask);
+	DECLARE_EXPORTED_GLOBAL_SHADER(FSH_ImageMask, UE_API);
 	SHADER_USE_PARAMETER_STRUCT(FSH_ImageMask, FSH_Base);
 
 	BEGIN_SHADER_PARAMETER_STRUCT(FParameters, )
@@ -47,16 +49,18 @@ typedef FxMaterial_Normal<VSH_Simple, FSH_ImageMask>	Fx_ImageMask;
 
 class UImageMask;
 
-class TEXTUREGRAPHENGINE_API T_ImageMask
+class T_ImageMask
 {
 private:
 
 
 public:
-									T_ImageMask();
-	virtual							~T_ImageMask();
+									UE_API T_ImageMask();
+	UE_API virtual							~T_ImageMask();
 
 	//////////////////////////////////////////////////////////////////////////
 	/// Static functions
 	//////////////////////////////////////////////////////////////////////////
 };
+
+#undef UE_API
diff --git a/Engine/Plugins/TextureGraph/Source/TextureGraphEngine/Transform/Mask/T_MaskBlend_Multiply.h b/Engine/Plugins/TextureGraph/Source/TextureGraphEngine/Transform/Mask/T_MaskBlend_Multiply.h
--- a/Engine/Plugins/TextureGraph/Source/TextureGraphEngine/Transform/Mask/T_MaskBlend_Multiply.h
+++ b/Engine/Plugins/TextureGraph/Source/TextureGraphEngine/Transform/Mask/T_MaskBlend_Multiply.h
@@ -7,13 +7,15 @@
 #include "Model/Mix/MixUpdateCycle.h"
 #include <DataDrivenShaderPlatformInfo.h>
 
+#define UE_API TEXTUREGRAPHENGINE_API
+
 //////////////////////////////////////////////////////////////////////////
 /// Shader
 //////////////////////////////////////////////////////////////////////////
 class FSH_MaskBlend_Multiply : public FSH_Base
 {
 public:
-	DECLARE_GLOBAL_SHADER(FSH_MaskBlend_Multiply);
+	DECLARE_EXPORTED_GLOBAL_SHADER(FSH_MaskBlend_Multiply, UE_API);
 	SHADER_USE_PARAMETER_STRUCT(FSH_MaskBlend_Multiply, FSH_Base);
 
 	BEGIN_SHADER_PARAMETER_STRUCT(FParameters, )
@@ -42,12 +44,14 @@ typedef FxMaterial_Normal<VSH_Simple, FSH_MaskBlend_Multiply>	Fx_BlendMode_Multi
 
 class UImageMask;
 
-class TEXTUREGRAPHENGINE_API T_MaskBlend_Multiply
+class T_MaskBlend_Multiply
 {
 private:
 
 
 public:
-									T_MaskBlend_Multiply();
-	virtual							~T_MaskBlend_Multiply();
+									UE_API T_MaskBlend_Multiply();
+	UE_API virtual							~T_MaskBlend_Multiply();
 };
+
+#undef UE_API
