Ran UnrealCodeFixup to fix dll storage

diff --git a/Engine/Plugins/TextureGraph/Source/TextureGraphEngine/Transform/Mask/T_MaskModifier.h b/Engine/Plugins/TextureGraph/Source/TextureGraphEngine/Transform/Mask/T_MaskModifier.h
--- a/Engine/Plugins/TextureGraph/Source/TextureGraphEngine/Transform/Mask/T_MaskModifier.h
+++ b/Engine/Plugins/TextureGraph/Source/TextureGraphEngine/Transform/Mask/T_MaskModifier.h
@@ -7,6 +7,8 @@
 #include "Model/Mix/MixUpdateCycle.h"
 #include <DataDrivenShaderPlatformInfo.h>
 
+#define UE_API TEXTUREGRAPHENGINE_API
+
 //////////////////////////////////////////////////////////////////////////
 /// Shader
 //////////////////////////////////////////////////////////////////////////
@@ -16,10 +18,10 @@ namespace
 	class FVar_INVERT : SHADER_PERMUTATION_BOOL("INVERT");
 }
 
-class TEXTUREGRAPHENGINE_API FSH_BrightnessMaskModifier : public FSH_Base
+class FSH_BrightnessMaskModifier : public FSH_Base
 {
 public:
-	DECLARE_GLOBAL_SHADER(FSH_BrightnessMaskModifier);
+	DECLARE_EXPORTED_GLOBAL_SHADER(FSH_BrightnessMaskModifier, UE_API);
 	SHADER_USE_PARAMETER_STRUCT(FSH_BrightnessMaskModifier, FSH_Base);
 	BEGIN_SHADER_PARAMETER_STRUCT(FParameters, )
 		SHADER_PARAMETER_TEXTURE(Texture2D, SourceTexture)
@@ -39,10 +41,10 @@ public:
 	}
 };
 
-class TEXTUREGRAPHENGINE_API FSH_ClampMaskModifier : public FSH_Base
+class FSH_ClampMaskModifier : public FSH_Base
 {
 public:
-	DECLARE_GLOBAL_SHADER(FSH_ClampMaskModifier);
+	DECLARE_EXPORTED_GLOBAL_SHADER(FSH_ClampMaskModifier, UE_API);
 	SHADER_USE_PARAMETER_STRUCT(FSH_ClampMaskModifier, FSH_Base);
 	BEGIN_SHADER_PARAMETER_STRUCT(FParameters, )
 		SHADER_PARAMETER_TEXTURE(Texture2D, SourceTexture)
@@ -61,10 +63,10 @@ public:
 	}
 };
 
-class TEXTUREGRAPHENGINE_API FSH_InvertMaskModifier : public FSH_Base
+class FSH_InvertMaskModifier : public FSH_Base
 {
 public:
-	DECLARE_GLOBAL_SHADER(FSH_InvertMaskModifier);
+	DECLARE_EXPORTED_GLOBAL_SHADER(FSH_InvertMaskModifier, UE_API);
 	SHADER_USE_PARAMETER_STRUCT(FSH_InvertMaskModifier, FSH_Base);
 	BEGIN_SHADER_PARAMETER_STRUCT(FParameters, )
 		SHADER_PARAMETER_TEXTURE(Texture2D, SourceTexture)
@@ -82,10 +84,10 @@ public:
 	}
 };
 
-class TEXTUREGRAPHENGINE_API FSH_NormalizeMaskModifier : public FSH_Base
+class FSH_NormalizeMaskModifier : public FSH_Base
 {
 public:
-	DECLARE_GLOBAL_SHADER(FSH_NormalizeMaskModifier);
+	DECLARE_EXPORTED_GLOBAL_SHADER(FSH_NormalizeMaskModifier, UE_API);
 	SHADER_USE_PARAMETER_STRUCT(FSH_NormalizeMaskModifier, FSH_Base);
 	BEGIN_SHADER_PARAMETER_STRUCT(FParameters, )
 		SHADER_PARAMETER_TEXTURE(Texture2D, SourceTexture)
@@ -103,10 +105,10 @@ public:
 	}
 };
 
-class TEXTUREGRAPHENGINE_API FSH_GradientRemapMaskModifier : public FSH_Base
+class FSH_GradientRemapMaskModifier : public FSH_Base
 {
 public:
-	DECLARE_GLOBAL_SHADER(FSH_GradientRemapMaskModifier);
+	DECLARE_EXPORTED_GLOBAL_SHADER(FSH_GradientRemapMaskModifier, UE_API);
 	SHADER_USE_PARAMETER_STRUCT(FSH_GradientRemapMaskModifier, FSH_Base);
 	BEGIN_SHADER_PARAMETER_STRUCT(FParameters, )
 		SHADER_PARAMETER_TEXTURE(Texture2D, SourceTexture)
@@ -128,10 +130,10 @@ public:
 	}
 };
 
-class TEXTUREGRAPHENGINE_API FSH_PosterizeMaskModifier : public FSH_Base
+class FSH_PosterizeMaskModifier : public FSH_Base
 {
 public:
-	DECLARE_GLOBAL_SHADER(FSH_PosterizeMaskModifier);
+	DECLARE_EXPORTED_GLOBAL_SHADER(FSH_PosterizeMaskModifier, UE_API);
 	SHADER_USE_PARAMETER_STRUCT(FSH_PosterizeMaskModifier, FSH_Base);
 	BEGIN_SHADER_PARAMETER_STRUCT(FParameters, )
 		SHADER_PARAMETER_TEXTURE(Texture2D, SourceTexture)
@@ -155,13 +157,15 @@ class UMaskModifier;
 /**
  * 
  */
-class TEXTUREGRAPHENGINE_API T_MaskModifier
+class T_MaskModifier
 {
 public:
-	T_MaskModifier();
-	~T_MaskModifier();
+	UE_API T_MaskModifier();
+	UE_API ~T_MaskModifier();
 
 	//////////////////////////////////////////////////////////////////////////
 	/// Static functions
 	//////////////////////////////////////////////////////////////////////////
 };
+
+#undef UE_API
diff --git a/Engine/Plugins/TextureGraph/Source/TextureGraphEngine/Transform/Mask/T_NoiseMask.h b/Engine/Plugins/TextureGraph/Source/TextureGraphEngine/Transform/Mask/T_NoiseMask.h
--- a/Engine/Plugins/TextureGraph/Source/TextureGraphEngine/Transform/Mask/T_NoiseMask.h
+++ b/Engine/Plugins/TextureGraph/Source/TextureGraphEngine/Transform/Mask/T_NoiseMask.h
@@ -8,6 +8,8 @@
 #include <DataDrivenShaderPlatformInfo.h>
 #include "T_NoiseMask.generated.h"
 
+#define UE_API TEXTUREGRAPHENGINE_API
+
 UENUM()
 enum class NoiseType
 {
@@ -28,10 +30,10 @@ namespace
 	class FVar_NoiseType : SHADER_PERMUTATION_INT("NOISE_TYPE", (int32)NoiseType::Count);
 }
 
-class TEXTUREGRAPHENGINE_API FSH_NoiseMask : public FSH_Base
+class FSH_NoiseMask : public FSH_Base
 {
 public:
-	DECLARE_GLOBAL_SHADER(FSH_NoiseMask);
+	DECLARE_EXPORTED_GLOBAL_SHADER(FSH_NoiseMask, UE_API);
 	SHADER_USE_PARAMETER_STRUCT(FSH_NoiseMask, FSH_Base);
 	BEGIN_SHADER_PARAMETER_STRUCT(FParameters, )
 		SHADER_PARAMETER_STRUCT(FTileInfo, TileInfo)
@@ -68,13 +70,15 @@ class UNoiseMask;
 /**
  * 
  */
-class TEXTUREGRAPHENGINE_API T_NoiseMask
+class T_NoiseMask
 {
 public:
-	T_NoiseMask();
-	~T_NoiseMask();
+	UE_API T_NoiseMask();
+	UE_API ~T_NoiseMask();
 
 	//////////////////////////////////////////////////////////////////////////
 	/// Static functions
 	//////////////////////////////////////////////////////////////////////////
 };
+
+#undef UE_API
diff --git a/Engine/Plugins/TextureGraph/Source/TextureGraphEngine/Transform/Mask/T_NormalMask.h b/Engine/Plugins/TextureGraph/Source/TextureGraphEngine/Transform/Mask/T_NormalMask.h
--- a/Engine/Plugins/TextureGraph/Source/TextureGraphEngine/Transform/Mask/T_NormalMask.h
+++ b/Engine/Plugins/TextureGraph/Source/TextureGraphEngine/Transform/Mask/T_NormalMask.h
@@ -7,6 +7,8 @@
 #include "Model/Mix/MixUpdateCycle.h"
 #include <DataDrivenShaderPlatformInfo.h>
 
+#define UE_API TEXTUREGRAPHENGINE_API
+
 namespace
 {
 	// Declare one Shader Permutation Var class per parameters
@@ -18,7 +20,7 @@ namespace
 class FSH_NormalMask : public FSH_Base
 {
 public:
-	DECLARE_GLOBAL_SHADER(FSH_NormalMask);
+	DECLARE_EXPORTED_GLOBAL_SHADER(FSH_NormalMask, UE_API);
 	SHADER_USE_PARAMETER_STRUCT(FSH_NormalMask, FSH_Base);
 	BEGIN_SHADER_PARAMETER_STRUCT(FParameters, )
 		SHADER_PARAMETER_TEXTURE(Texture2D, SourceTex)
@@ -57,13 +59,15 @@ class UNormalMask;
 /**
  * 
  */
-class TEXTUREGRAPHENGINE_API T_NormalMask
+class T_NormalMask
 {
 public:
-	T_NormalMask();
-	~T_NormalMask();
+	UE_API T_NormalMask();
+	UE_API ~T_NormalMask();
 
 	//////////////////////////////////////////////////////////////////////////
 	/// Static functions
 	//////////////////////////////////////////////////////////////////////////
 };
+
+#undef UE_API
diff --git a/Engine/Plugins/TextureGraph/Source/TextureGraphEngine/Transform/Mask/T_PatternMask.h b/Engine/Plugins/TextureGraph/Source/TextureGraphEngine/Transform/Mask/T_PatternMask.h
--- a/Engine/Plugins/TextureGraph/Source/TextureGraphEngine/Transform/Mask/T_PatternMask.h
+++ b/Engine/Plugins/TextureGraph/Source/TextureGraphEngine/Transform/Mask/T_PatternMask.h
@@ -9,6 +9,8 @@
 #include <DataDrivenShaderPlatformInfo.h>
 #include "T_PatternMask.generated.h"
 
+#define UE_API TEXTUREGRAPHENGINE_API
+
 UENUM()
 enum class PatternType
 {
@@ -29,10 +31,10 @@ namespace
 	class FVar_PatternType : SHADER_PERMUTATION_INT("PATTERN_TYPE", (int32) PatternType::Count);
 }
 
-class TEXTUREGRAPHENGINE_API FSH_PatternMask : public FSH_Base
+class FSH_PatternMask : public FSH_Base
 {
 public:
-	DECLARE_GLOBAL_SHADER(FSH_PatternMask);
+	DECLARE_EXPORTED_GLOBAL_SHADER(FSH_PatternMask, UE_API);
 	SHADER_USE_PARAMETER_STRUCT(FSH_PatternMask, FSH_Base);
 	BEGIN_SHADER_PARAMETER_STRUCT(FParameters, )
 		SHADER_PARAMETER_STRUCT(FTileInfo, TileInfo)
@@ -83,13 +85,15 @@ class UPatternMask;
 /**
  * 
  */
-class TEXTUREGRAPHENGINE_API T_PatternMask
+class T_PatternMask
 {
 public:
-	T_PatternMask();
-	~T_PatternMask();
+	UE_API T_PatternMask();
+	UE_API ~T_PatternMask();
 
 	//////////////////////////////////////////////////////////////////////////
 	/// Static functions
 	//////////////////////////////////////////////////////////////////////////
 };
+
+#undef UE_API
diff --git a/Engine/Plugins/TextureGraph/Source/TextureGraphEngine/Transform/Mask/T_PositionGradientMask.h b/Engine/Plugins/TextureGraph/Source/TextureGraphEngine/Transform/Mask/T_PositionGradientMask.h
--- a/Engine/Plugins/TextureGraph/Source/TextureGraphEngine/Transform/Mask/T_PositionGradientMask.h
+++ b/Engine/Plugins/TextureGraph/Source/TextureGraphEngine/Transform/Mask/T_PositionGradientMask.h
@@ -7,6 +7,8 @@
 #include "Model/Mix/MixUpdateCycle.h"
 #include <DataDrivenShaderPlatformInfo.h>
 
+#define UE_API TEXTUREGRAPHENGINE_API
+
 namespace PositionGradientMask
 {
 	// Declare one Shader Permutation Var class per parameters
@@ -18,7 +20,7 @@ namespace PositionGradientMask
 class FSH_PositionGradientMask : public FSH_Base
 {
 public:
-	DECLARE_GLOBAL_SHADER(FSH_PositionGradientMask);
+	DECLARE_EXPORTED_GLOBAL_SHADER(FSH_PositionGradientMask, UE_API);
 	SHADER_USE_PARAMETER_STRUCT(FSH_PositionGradientMask, FSH_Base);
 	BEGIN_SHADER_PARAMETER_STRUCT(FParameters, )
 		SHADER_PARAMETER_STRUCT(FTileInfo, TileInfo)
@@ -58,7 +60,7 @@ typedef FxMaterial_Normal<VSH_Simple, FSH_PositionGradientMask>	Fx_PositionGradi
 class FSH_CombineDisplacementAndWorldPos : public FSH_Base
 {
 public:
-	DECLARE_GLOBAL_SHADER(FSH_CombineDisplacementAndWorldPos);
+	DECLARE_EXPORTED_GLOBAL_SHADER(FSH_CombineDisplacementAndWorldPos, UE_API);
 	SHADER_USE_PARAMETER_STRUCT(FSH_CombineDisplacementAndWorldPos, FSH_Base);
 	BEGIN_SHADER_PARAMETER_STRUCT(FParameters, )		
 		SHADER_PARAMETER_TEXTURE(Texture2D, SourceDisplacement)
@@ -87,13 +89,15 @@ class UPositionGradientMask;
 /**
  * 
  */
-class TEXTUREGRAPHENGINE_API T_PositionGradientMask
+class T_PositionGradientMask
 {
 public:
-	T_PositionGradientMask();
-	~T_PositionGradientMask();
+	UE_API T_PositionGradientMask();
+	UE_API ~T_PositionGradientMask();
 
 	//////////////////////////////////////////////////////////////////////////
 	/// Static functions
 	//////////////////////////////////////////////////////////////////////////=
 };
+
+#undef UE_API
diff --git a/Engine/Plugins/TextureGraph/Source/TextureGraphEngine/Transform/Mask/T_ShapeMask.h b/Engine/Plugins/TextureGraph/Source/TextureGraphEngine/Transform/Mask/T_ShapeMask.h
--- a/Engine/Plugins/TextureGraph/Source/TextureGraphEngine/Transform/Mask/T_ShapeMask.h
+++ b/Engine/Plugins/TextureGraph/Source/TextureGraphEngine/Transform/Mask/T_ShapeMask.h
@@ -9,6 +9,8 @@
 #include <DataDrivenShaderPlatformInfo.h>
 #include "T_ShapeMask.generated.h"
 
+#define UE_API TEXTUREGRAPHENGINE_API
+
 UENUM(BlueprintType)
 enum class EShapeMaskType : uint8
 {
@@ -59,10 +61,10 @@ namespace
 	class FVar_ShapeType : SHADER_PERMUTATION_INT("SHAPE_TYPE", (int32)EShapeMaskType::Count);
 }
 
-class TEXTUREGRAPHENGINE_API FSH_ShapeMask : public FSH_Base
+class FSH_ShapeMask : public FSH_Base
 {
 public:
-	DECLARE_GLOBAL_SHADER(FSH_ShapeMask);
+	DECLARE_EXPORTED_GLOBAL_SHADER(FSH_ShapeMask, UE_API);
 	SHADER_USE_PARAMETER_STRUCT(FSH_ShapeMask, FSH_Base);
 	BEGIN_SHADER_PARAMETER_STRUCT(FParameters, )
 		SHADER_PARAMETER_STRUCT(FTileInfo, TileInfo)
@@ -98,7 +100,7 @@ typedef FxMaterial_Normal<VSH_Simple, FSH_ShapeMask>	Fx_ShapeMask;
 /**
  * 
  */
-class TEXTUREGRAPHENGINE_API T_ShapeMask
+class T_ShapeMask
 {
 public:
 	static constexpr int				DefaultSize = 1024;
@@ -113,6 +115,8 @@ public:
 		float BlendSDF = 0.0f;
 	};
 
-	static TiledBlobPtr	Create(MixUpdateCyclePtr Cycle, BufferDescriptor DesiredOutputDesc,
+	static UE_API TiledBlobPtr	Create(MixUpdateCyclePtr Cycle, BufferDescriptor DesiredOutputDesc,
 		EShapeMaskType ShapeMaskType, const FParams& Params, int32 TargetId = 0);
 };
+
+#undef UE_API
diff --git a/Engine/Plugins/TextureGraph/Source/TextureGraphEngine/Transform/Mesh/T_MeshDilateUVs.h b/Engine/Plugins/TextureGraph/Source/TextureGraphEngine/Transform/Mesh/T_MeshDilateUVs.h
--- a/Engine/Plugins/TextureGraph/Source/TextureGraphEngine/Transform/Mesh/T_MeshDilateUVs.h
+++ b/Engine/Plugins/TextureGraph/Source/TextureGraphEngine/Transform/Mesh/T_MeshDilateUVs.h
@@ -7,13 +7,15 @@
 #include "Model/Mix/MixUpdateCycle.h"
 #include <DataDrivenShaderPlatformInfo.h>
 
+#define UE_API TEXTUREGRAPHENGINE_API
+
 class ULayer_Textured;
 
 //////////////////////////////////////////////////////////////////////////
 class FSH_DilateMeshTexture : public FSH_Base
 {
 public:
-	DECLARE_GLOBAL_SHADER(FSH_DilateMeshTexture);
+	DECLARE_EXPORTED_GLOBAL_SHADER(FSH_DilateMeshTexture, UE_API);
 	SHADER_USE_PARAMETER_STRUCT(FSH_DilateMeshTexture, FSH_Base);
 
 	BEGIN_SHADER_PARAMETER_STRUCT(FParameters, )
@@ -40,14 +42,16 @@ template <> void SetupDefaultParameters(FSH_DilateMeshTexture::FParameters& para
 typedef FxMaterial_Normal<VSH_Simple, FSH_DilateMeshTexture> Fx_DilateMeshTexture;
 
 //////////////////////////////////////////////////////////////////////////
-class TEXTUREGRAPHENGINE_API T_MeshDilateUVs
+class T_MeshDilateUVs
 {
 public:
-									T_MeshDilateUVs();
-	virtual							~T_MeshDilateUVs();
+									UE_API T_MeshDilateUVs();
+	UE_API virtual							~T_MeshDilateUVs();
 
 	//////////////////////////////////////////////////////////////////////////
 	/// Static functions
 	//////////////////////////////////////////////////////////////////////////
-	static TiledBlobPtr				Create(MixUpdateCyclePtr cycle, int32 targetId, TiledBlobPtr sourceTexture);
+	static UE_API TiledBlobPtr				Create(MixUpdateCyclePtr cycle, int32 targetId, TiledBlobPtr sourceTexture);
 };
+
+#undef UE_API
diff --git a/Engine/Plugins/TextureGraph/Source/TextureGraphEngine/Transform/Mesh/T_MeshEncodedAsTexture.h b/Engine/Plugins/TextureGraph/Source/TextureGraphEngine/Transform/Mesh/T_MeshEncodedAsTexture.h
--- a/Engine/Plugins/TextureGraph/Source/TextureGraphEngine/Transform/Mesh/T_MeshEncodedAsTexture.h
+++ b/Engine/Plugins/TextureGraph/Source/TextureGraphEngine/Transform/Mesh/T_MeshEncodedAsTexture.h
@@ -8,13 +8,15 @@
 #include "UObject/NoExportTypes.h"
 #include <DataDrivenShaderPlatformInfo.h>
 
+#define UE_API TEXTUREGRAPHENGINE_API
+
 class ULayer_Textured;
 
 ////////////////////////////////////////////////////////////////////////////
 class VSH_MeshTexture : public VSH_Base
 {
 public:
-	DECLARE_GLOBAL_SHADER(VSH_MeshTexture);
+	DECLARE_EXPORTED_GLOBAL_SHADER(VSH_MeshTexture, UE_API);
 	SHADER_USE_PARAMETER_STRUCT(VSH_MeshTexture, VSH_Base);
 
 	BEGIN_SHADER_PARAMETER_STRUCT(FParameters, )
@@ -31,7 +33,7 @@ public:
 class VSH_MeshTexture_WorldPos : public VSH_Base
 {
 public:
-	DECLARE_GLOBAL_SHADER(VSH_MeshTexture_WorldPos);
+	DECLARE_EXPORTED_GLOBAL_SHADER(VSH_MeshTexture_WorldPos, UE_API);
 	SHADER_USE_PARAMETER_STRUCT(VSH_MeshTexture_WorldPos, VSH_Base);
 
 	BEGIN_SHADER_PARAMETER_STRUCT(FParameters, )
@@ -50,7 +52,7 @@ public:
 class FSH_MeshTexture_WorldPos : public FSH_Base
 {
 public:
-	DECLARE_GLOBAL_SHADER(FSH_MeshTexture_WorldPos);
+	DECLARE_EXPORTED_GLOBAL_SHADER(FSH_MeshTexture_WorldPos, UE_API);
 	SHADER_USE_PARAMETER_STRUCT(FSH_MeshTexture_WorldPos, FSH_Base);
 
 	BEGIN_SHADER_PARAMETER_STRUCT(FParameters, )
@@ -71,7 +73,7 @@ public:
 class FSH_MeshTexture_WorldNormals : public FSH_Base
 {
 public:
-	DECLARE_GLOBAL_SHADER(FSH_MeshTexture_WorldNormals);
+	DECLARE_EXPORTED_GLOBAL_SHADER(FSH_MeshTexture_WorldNormals, UE_API);
 	SHADER_USE_PARAMETER_STRUCT(FSH_MeshTexture_WorldNormals, FSH_Base);
 
 	BEGIN_SHADER_PARAMETER_STRUCT(FParameters, )
@@ -92,7 +94,7 @@ public:
 class FSH_MeshTexture_WorldTangents : public FSH_Base
 {
 public:
-	DECLARE_GLOBAL_SHADER(FSH_MeshTexture_WorldTangents);
+	DECLARE_EXPORTED_GLOBAL_SHADER(FSH_MeshTexture_WorldTangents, UE_API);
 	SHADER_USE_PARAMETER_STRUCT(FSH_MeshTexture_WorldTangents, FSH_Base);
 
 	BEGIN_SHADER_PARAMETER_STRUCT(FParameters, )
@@ -113,7 +115,7 @@ public:
 class FSH_MeshTexture_WorldUVMask : public FSH_Base
 {
 public:
-	DECLARE_GLOBAL_SHADER(FSH_MeshTexture_WorldUVMask);
+	DECLARE_EXPORTED_GLOBAL_SHADER(FSH_MeshTexture_WorldUVMask, UE_API);
 	SHADER_USE_PARAMETER_STRUCT(FSH_MeshTexture_WorldUVMask, FSH_Base);
 
 	BEGIN_SHADER_PARAMETER_STRUCT(FParameters, )
@@ -131,18 +133,20 @@ public:
 };
 
 //////////////////////////////////////////////////////////////////////////
-class TEXTUREGRAPHENGINE_API T_MeshEncodedAsTexture 
+class T_MeshEncodedAsTexture 
 {
 public:
-	static int32					s_minMeshmapRes;			// We need to create mesh maps in 4k. Otherwise we get artifacts when using in curvature and painting. 
-									T_MeshEncodedAsTexture();
-	virtual							~T_MeshEncodedAsTexture();
+	static UE_API int32					s_minMeshmapRes;			// We need to create mesh maps in 4k. Otherwise we get artifacts when using in curvature and painting. 
+									UE_API T_MeshEncodedAsTexture();
+	UE_API virtual							~T_MeshEncodedAsTexture();
 
 	//////////////////////////////////////////////////////////////////////////
 	/// Static functions
 	//////////////////////////////////////////////////////////////////////////
-	static TiledBlobPtr				Create_WorldPos(MixUpdateCyclePtr cycle, int32 targetId);
-	static TiledBlobPtr				Create_WorldNormals(MixUpdateCyclePtr cycle, int32 targetId);
-	static TiledBlobPtr				Create_WorldTangents(MixUpdateCyclePtr cycle, int32 targetId);
-	static TiledBlobPtr				Create_WorldUVMask(MixUpdateCyclePtr cycle, int32 targetId);
+	static UE_API TiledBlobPtr				Create_WorldPos(MixUpdateCyclePtr cycle, int32 targetId);
+	static UE_API TiledBlobPtr				Create_WorldNormals(MixUpdateCyclePtr cycle, int32 targetId);
+	static UE_API TiledBlobPtr				Create_WorldTangents(MixUpdateCyclePtr cycle, int32 targetId);
+	static UE_API TiledBlobPtr				Create_WorldUVMask(MixUpdateCyclePtr cycle, int32 targetId);
 };
+
+#undef UE_API
diff --git a/Engine/Plugins/TextureGraph/Source/TextureGraphEngine/Transform/Mix/T_InvalidateTiles.h b/Engine/Plugins/TextureGraph/Source/TextureGraphEngine/Transform/Mix/T_InvalidateTiles.h
--- a/Engine/Plugins/TextureGraph/Source/TextureGraphEngine/Transform/Mix/T_InvalidateTiles.h
+++ b/Engine/Plugins/TextureGraph/Source/TextureGraphEngine/Transform/Mix/T_InvalidateTiles.h
@@ -6,20 +6,24 @@
 #include "Job/Job.h"
 #include "Model/Mix/MixUpdateCycle.h"
 
-class TEXTUREGRAPHENGINE_API T_InvalidateTiles : public BlobTransform
+#define UE_API TEXTUREGRAPHENGINE_API
+
+class T_InvalidateTiles : public BlobTransform
 {
 public:
-									T_InvalidateTiles();
-	virtual							~T_InvalidateTiles() override;
+									UE_API T_InvalidateTiles();
+	UE_API virtual							~T_InvalidateTiles() override;
 
-	virtual Device*					TargetDevice(size_t index) const override;
-	virtual AsyncTransformResultPtr	Exec(const TransformArgs& args) override;
+	UE_API virtual Device*					TargetDevice(size_t index) const override;
+	UE_API virtual AsyncTransformResultPtr	Exec(const TransformArgs& args) override;
 
 	virtual bool					GeneratesData() const override { return false; }
 
 	////////////////////////////////////////////////////////////////////////// 
 	/// Static functions
 	//////////////////////////////////////////////////////////////////////////
-	static JobUPtr					CreateJob(MixUpdateCyclePtr cycle, int32 targetId);
-	static void						Create(MixUpdateCyclePtr cycle, int32 targetId);
+	static UE_API JobUPtr					CreateJob(MixUpdateCyclePtr cycle, int32 targetId);
+	static UE_API void						Create(MixUpdateCyclePtr cycle, int32 targetId);
 };
+
+#undef UE_API
diff --git a/Engine/Plugins/TextureGraph/Source/TextureGraphEngine/Transform/Mix/T_UpdateTargets.h b/Engine/Plugins/TextureGraph/Source/TextureGraphEngine/Transform/Mix/T_UpdateTargets.h
--- a/Engine/Plugins/TextureGraph/Source/TextureGraphEngine/Transform/Mix/T_UpdateTargets.h
+++ b/Engine/Plugins/TextureGraph/Source/TextureGraphEngine/Transform/Mix/T_UpdateTargets.h
@@ -6,16 +6,18 @@
 #include "Job/Job.h"
 #include "Model/Mix/MixUpdateCycle.h"
 
-class TEXTUREGRAPHENGINE_API T_UpdateTargets : public BlobTransform
+#define UE_API TEXTUREGRAPHENGINE_API
+
+class T_UpdateTargets : public BlobTransform
 {
 private:
 	bool							_shouldUpdate = false;				/// This will ensure default texture is used when there is no child in layer set
-	static int						_gpuTesselationFactor;
+	static UE_API int						_gpuTesselationFactor;
 public:
-									T_UpdateTargets();
-	virtual							~T_UpdateTargets() override;
-	virtual Device*					TargetDevice(size_t index) const override;
-	virtual AsyncTransformResultPtr	Exec(const TransformArgs& args) override;
+									UE_API T_UpdateTargets();
+	UE_API virtual							~T_UpdateTargets() override;
+	UE_API virtual Device*					TargetDevice(size_t index) const override;
+	UE_API virtual AsyncTransformResultPtr	Exec(const TransformArgs& args) override;
 
 	virtual bool					GeneratesData() const override { return false; }
 	virtual void					SetShouldUpdate(bool shouldUpdate) { _shouldUpdate = shouldUpdate; }
@@ -23,6 +25,8 @@ public:
 	/// Static functions
 	//////////////////////////////////////////////////////////////////////////
 	static void						SetGPUTesselationFactor(int tesselationFactor) { _gpuTesselationFactor = tesselationFactor; }
-	static JobUPtr					CreateJob(MixUpdateCyclePtr cycle, int32 targetId, bool shouldUpdate);
-	static void						Create(MixUpdateCyclePtr cycle, int32 targetId, bool shouldUpdate);
+	static UE_API JobUPtr					CreateJob(MixUpdateCyclePtr cycle, int32 targetId, bool shouldUpdate);
+	static UE_API void						Create(MixUpdateCyclePtr cycle, int32 targetId, bool shouldUpdate);
 };
+
+#undef UE_API
diff --git a/Engine/Plugins/TextureGraph/Source/TextureGraphEngine/Transform/Utility/T_CombineTiledBlob.h b/Engine/Plugins/TextureGraph/Source/TextureGraphEngine/Transform/Utility/T_CombineTiledBlob.h
--- a/Engine/Plugins/TextureGraph/Source/TextureGraphEngine/Transform/Utility/T_CombineTiledBlob.h
+++ b/Engine/Plugins/TextureGraph/Source/TextureGraphEngine/Transform/Utility/T_CombineTiledBlob.h
@@ -9,6 +9,8 @@
 #include <DataDrivenShaderPlatformInfo.h>
 #include "Device/FX/Device_FX.h"
 
+#define UE_API TEXTUREGRAPHENGINE_API
+
 //UENUM(BlueprintType)
 //enum class ResizeSettings : uint8
 //{
@@ -17,7 +19,7 @@
 //	Radial = 2			UMETA(DisplayName = "Radial")
 //};
 
-struct TEXTUREGRAPHENGINE_API CombineSettings
+struct CombineSettings
 {
 	bool bFixed = false;
 	bool bMaintainAspectRatio = false;
@@ -27,33 +29,33 @@ struct TEXTUREGRAPHENGINE_API CombineSettings
 /**
  * A dedicated BlobTransform creating a Combined version of a TiledBlob
  */
-class TEXTUREGRAPHENGINE_API CombineTiledBlob_Transform : public BlobTransform
+class CombineTiledBlob_Transform : public BlobTransform
 {
 private:
 	TiledBlobPtr					Source;				/// The source blob that we are combining
 	CombineSettings					Settings;			/// Combine settings
 	DrawTilesSettings				DrawSettings;		/// FX device specific draw settings
 
-	void							InitDrawSettings(BlobPtr Target, T_Tiles<DeviceBufferRef>& Tiles);
-	void							InitDrawSettingsWithoutAspectRatio(BlobPtr Target, T_Tiles<DeviceBufferRef>& Tiles);
-	void							InitDrawSettingsWithAspectRatio(BlobPtr Target, T_Tiles<DeviceBufferRef>& Tiles);
+	UE_API void							InitDrawSettings(BlobPtr Target, T_Tiles<DeviceBufferRef>& Tiles);
+	UE_API void							InitDrawSettingsWithoutAspectRatio(BlobPtr Target, T_Tiles<DeviceBufferRef>& Tiles);
+	UE_API void							InitDrawSettingsWithAspectRatio(BlobPtr Target, T_Tiles<DeviceBufferRef>& Tiles);
 
 public:
-	CombineTiledBlob_Transform(FString InName, TiledBlobPtr InSource, const CombineSettings* InSettings = nullptr);
-	virtual AsyncBufferResultPtr	Bind(BlobPtr Value, const ResourceBindInfo& BindInfo) override;
-	virtual AsyncBufferResultPtr	Unbind(BlobPtr BlobObj, const ResourceBindInfo& BindInfo) override;
-	virtual Device*					TargetDevice(size_t DevIndex) const override;
-	virtual std::shared_ptr<BlobTransform> DuplicateInstance(FString NewName) override;
-	virtual bool					GeneratesData() const override;
-	virtual bool					CanHandleTiles() const override;
+	UE_API CombineTiledBlob_Transform(FString InName, TiledBlobPtr InSource, const CombineSettings* InSettings = nullptr);
+	UE_API virtual AsyncBufferResultPtr	Bind(BlobPtr Value, const ResourceBindInfo& BindInfo) override;
+	UE_API virtual AsyncBufferResultPtr	Unbind(BlobPtr BlobObj, const ResourceBindInfo& BindInfo) override;
+	UE_API virtual Device*					TargetDevice(size_t DevIndex) const override;
+	UE_API virtual std::shared_ptr<BlobTransform> DuplicateInstance(FString NewName) override;
+	UE_API virtual bool					GeneratesData() const override;
+	UE_API virtual bool					CanHandleTiles() const override;
 
-	virtual AsyncTransformResultPtr	Exec(const TransformArgs& Args) override;
+	UE_API virtual AsyncTransformResultPtr	Exec(const TransformArgs& Args) override;
 };
 
 /**
  * Tiles To Combined Transform
  */
-class TEXTUREGRAPHENGINE_API T_CombineTiledBlob
+class T_CombineTiledBlob
 {
 public:
 
@@ -61,6 +63,8 @@ public:
 	/// Static functions
 	//////////////////////////////////////////////////////////////////////////
 	
-	static TiledBlobPtr				Create(MixUpdateCyclePtr Cycle, BufferDescriptor DesiredOutputDesc, int32 TargetId, 
+	static UE_API TiledBlobPtr				Create(MixUpdateCyclePtr Cycle, BufferDescriptor DesiredOutputDesc, int32 TargetId, 
 									TiledBlobPtr SourceTex, JobUPtr JobToUse = nullptr, const CombineSettings* InSettings = nullptr);
 };
+
+#undef UE_API
diff --git a/Engine/Plugins/TextureGraph/Source/TextureGraphEngine/Transform/Utility/T_FinaliseBlob.h b/Engine/Plugins/TextureGraph/Source/TextureGraphEngine/Transform/Utility/T_FinaliseBlob.h
--- a/Engine/Plugins/TextureGraph/Source/TextureGraphEngine/Transform/Utility/T_FinaliseBlob.h
+++ b/Engine/Plugins/TextureGraph/Source/TextureGraphEngine/Transform/Utility/T_FinaliseBlob.h
@@ -6,11 +6,15 @@
 #include "Job/Job.h"
 #include "Model/Mix/MixUpdateCycle.h"
 
-class TEXTUREGRAPHENGINE_API T_FinaliseBlob 
+#define UE_API TEXTUREGRAPHENGINE_API
+
+class T_FinaliseBlob 
 {
 public:
 	////////////////////////////////////////////////////////////////////////// 
 	/// Static functions
 	//////////////////////////////////////////////////////////////////////////
-	static JobPtr					Create(MixUpdateCyclePtr Cycle, JobPtr JobObj);
+	static UE_API JobPtr					Create(MixUpdateCyclePtr Cycle, JobPtr JobObj);
 };
+
+#undef UE_API
diff --git a/Engine/Plugins/TextureGraph/Source/TextureGraphEngine/Transform/Utility/T_LoadStaticResource.h b/Engine/Plugins/TextureGraph/Source/TextureGraphEngine/Transform/Utility/T_LoadStaticResource.h
--- a/Engine/Plugins/TextureGraph/Source/TextureGraphEngine/Transform/Utility/T_LoadStaticResource.h
+++ b/Engine/Plugins/TextureGraph/Source/TextureGraphEngine/Transform/Utility/T_LoadStaticResource.h
@@ -5,21 +5,23 @@
 #include "Job/Job.h"
 #include "Model/Mix/MixUpdateCycle.h"
 
+#define UE_API TEXTUREGRAPHENGINE_API
+
 class UStaticImageResource;
 
-class TEXTUREGRAPHENGINE_API Job_LoadStaticImageResource : public Job, public FGCObject
+class Job_LoadStaticImageResource : public Job, public FGCObject
 {
 	TObjectPtr<UStaticImageResource> Source;			/// The source for the static image resource
 
 public:
-									Job_LoadStaticImageResource(UMixInterface* InMix, UStaticImageResource* InSource, int32 TargetId, UObject* InErrorOwner = nullptr, uint16 Priority = (uint16)E_Priority::kHigh);
+									UE_API Job_LoadStaticImageResource(UMixInterface* InMix, UStaticImageResource* InSource, int32 TargetId, UObject* InErrorOwner = nullptr, uint16 Priority = (uint16)E_Priority::kHigh);
 	virtual							~Job_LoadStaticImageResource() override = default;
 
 	//////////////////////////////////////////////////////////////////////////
 	/// FGCObject
 	//////////////////////////////////////////////////////////////////////////
-	virtual void					AddReferencedObjects(FReferenceCollector& Collector) override;
-	virtual FString					GetReferencerName() const override;
+	UE_API virtual void					AddReferencedObjects(FReferenceCollector& Collector) override;
+	UE_API virtual FString					GetReferencerName() const override;
 
 protected:
 	virtual void					GetDependencies(JobPtrVec& prior, JobPtrVec& after, JobRunInfo runInfo) override { RunInfo = runInfo; Prev.clear(); }
@@ -31,11 +33,13 @@ protected:
 	virtual cti::continuable<int32>	UnbindArgs_All(JobRunInfo runInfo) override { return cti::make_ready_continuable(0); }
 	virtual bool					CanHandleTiles() const override { return false; }
 
-	virtual cti::continuable<int32>	PreExecAsync(ENamedThreads::Type execThread, ENamedThreads::Type returnThread) override;
-	virtual cti::continuable<int32>	ExecAsync(ENamedThreads::Type execThread, ENamedThreads::Type returnThread) override;
+	UE_API virtual cti::continuable<int32>	PreExecAsync(ENamedThreads::Type execThread, ENamedThreads::Type returnThread) override;
+	UE_API virtual cti::continuable<int32>	ExecAsync(ENamedThreads::Type execThread, ENamedThreads::Type returnThread) override;
 	virtual bool					IsAsync() const override { return true; }
 };
 
 //////////////////////////////////////////////////////////////////////////
 /// File based
 ////////////////////////////////////////////////////////////////////////// 
+
+#undef UE_API
diff --git a/Engine/Plugins/TextureGraph/Source/TextureGraphEngine/Transform/Utility/T_MinMax.h b/Engine/Plugins/TextureGraph/Source/TextureGraphEngine/Transform/Utility/T_MinMax.h
--- a/Engine/Plugins/TextureGraph/Source/TextureGraphEngine/Transform/Utility/T_MinMax.h
+++ b/Engine/Plugins/TextureGraph/Source/TextureGraphEngine/Transform/Utility/T_MinMax.h
@@ -8,6 +8,8 @@
 #include "Model/Mix/MixUpdateCycle.h"
 #include <DataDrivenShaderPlatformInfo.h>
 
+#define UE_API TEXTUREGRAPHENGINE_API
+
 
 namespace MinMax
 {
@@ -22,7 +24,7 @@ namespace MinMax
 class FSH_MinMaxDownsample : public FSH_Base
 {
 public:
-	DECLARE_GLOBAL_SHADER(FSH_MinMaxDownsample);
+	DECLARE_EXPORTED_GLOBAL_SHADER(FSH_MinMaxDownsample, UE_API);
 	SHADER_USE_PARAMETER_STRUCT(FSH_MinMaxDownsample, FSH_Base);
 	BEGIN_SHADER_PARAMETER_STRUCT(FParameters, )
 		SHADER_PARAMETER_STRUCT(FStandardSamplerStates, SamplerStates)
@@ -54,7 +56,7 @@ typedef FxMaterial_Normal<VSH_Simple, FSH_MinMaxDownsample>	Fx_MinMax;
 class CSH_MinMaxDownsample : public CmpSH_Base<1, 1, 1>
 {
 public:
-	DECLARE_GLOBAL_SHADER(CSH_MinMaxDownsample);
+	DECLARE_EXPORTED_GLOBAL_SHADER(CSH_MinMaxDownsample, UE_API);
 	SHADER_USE_PARAMETER_STRUCT(CSH_MinMaxDownsample, CmpSH_Base);
 	BEGIN_SHADER_PARAMETER_STRUCT(FParameters, )
 		SHADER_PARAMETER(FIntVector4, TilingDimensions)
@@ -81,14 +83,16 @@ public:
 /**
  * 
  */
-class TEXTUREGRAPHENGINE_API T_MinMax
+class T_MinMax
 {
 public:
-	T_MinMax();
-	~T_MinMax();
+	UE_API T_MinMax();
+	UE_API ~T_MinMax();
 
 	//////////////////////////////////////////////////////////////////////////
 /// Static functions
 //////////////////////////////////////////////////////////////////////////=
-	static TiledBlobPtr				Create(MixUpdateCyclePtr Cycle, TiledBlobPtr SourceTex, int32 TargetId, bool bUseMeshUVMask = false);
+	static UE_API TiledBlobPtr				Create(MixUpdateCyclePtr Cycle, TiledBlobPtr SourceTex, int32 TargetId, bool bUseMeshUVMask = false);
 };
+
+#undef UE_API
diff --git a/Engine/Plugins/TextureGraph/Source/TextureGraphEngine/Transform/Utility/T_MipMap.h b/Engine/Plugins/TextureGraph/Source/TextureGraphEngine/Transform/Utility/T_MipMap.h
--- a/Engine/Plugins/TextureGraph/Source/TextureGraphEngine/Transform/Utility/T_MipMap.h
+++ b/Engine/Plugins/TextureGraph/Source/TextureGraphEngine/Transform/Utility/T_MipMap.h
@@ -8,6 +8,8 @@
 #include "Model/Mix/MixUpdateCycle.h"
 #include <DataDrivenShaderPlatformInfo.h>
 
+#define UE_API TEXTUREGRAPHENGINE_API
+
 
 //////////////////////////////////////////////////////////////////////////
 /// Shader
@@ -16,7 +18,7 @@
 class CSH_MipMapDownsample : public CmpSH_Base<1, 1, 1>
 {
 public:
-	DECLARE_GLOBAL_SHADER(CSH_MipMapDownsample);
+	DECLARE_EXPORTED_GLOBAL_SHADER(CSH_MipMapDownsample, UE_API);
 	SHADER_USE_PARAMETER_STRUCT(CSH_MipMapDownsample, CmpSH_Base);
 	BEGIN_SHADER_PARAMETER_STRUCT(FParameters, )
 		SHADER_PARAMETER_TEXTURE(Texture2D, SourceTiles)
@@ -38,11 +40,11 @@ public:
 /**
  * MipMap Transform
  */
-class TEXTUREGRAPHENGINE_API T_MipMap
+class T_MipMap
 {
 public:
-	T_MipMap();
-	~T_MipMap();
+	UE_API T_MipMap();
+	UE_API ~T_MipMap();
 
 	//////////////////////////////////////////////////////////////////////////
 	/// Static functions
@@ -54,5 +56,7 @@ public:
 	// thus if the required level is higher than the actual deepest tile can produce ( aka log2(size of the tile) )
 	// then the transform is returning the deepest level possibly generated with the same tile grid as sourceTex.
 	// @TODO SG: Fix that behavior!
-	static TiledBlobPtr				Create(MixUpdateCyclePtr Cycle, TiledBlobPtr SourceTex, int32 TargetId, int32 GenMipLevel = -1);
+	static UE_API TiledBlobPtr				Create(MixUpdateCyclePtr Cycle, TiledBlobPtr SourceTex, int32 TargetId, int32 GenMipLevel = -1);
 };
+
+#undef UE_API
diff --git a/Engine/Plugins/TextureGraph/Source/TextureGraphEngine/Transform/Utility/T_PrepareResources.h b/Engine/Plugins/TextureGraph/Source/TextureGraphEngine/Transform/Utility/T_PrepareResources.h
--- a/Engine/Plugins/TextureGraph/Source/TextureGraphEngine/Transform/Utility/T_PrepareResources.h
+++ b/Engine/Plugins/TextureGraph/Source/TextureGraphEngine/Transform/Utility/T_PrepareResources.h
@@ -6,11 +6,15 @@
 #include "Job/Job.h"
 #include "Model/Mix/MixUpdateCycle.h"
 
-class TEXTUREGRAPHENGINE_API T_PrepareResources 
+#define UE_API TEXTUREGRAPHENGINE_API
+
+class T_PrepareResources 
 {
 public:
 	////////////////////////////////////////////////////////////////////////// 
 	/// Static functions
 	//////////////////////////////////////////////////////////////////////////
-	static JobPtr					Create(MixUpdateCyclePtr Cycle, JobPtr JobObj);
+	static UE_API JobPtr					Create(MixUpdateCyclePtr Cycle, JobPtr JobObj);
 };
+
+#undef UE_API
diff --git a/Engine/Plugins/TextureGraph/Source/TextureGraphEngine/Transform/Utility/T_RenderCapture.h b/Engine/Plugins/TextureGraph/Source/TextureGraphEngine/Transform/Utility/T_RenderCapture.h
--- a/Engine/Plugins/TextureGraph/Source/TextureGraphEngine/Transform/Utility/T_RenderCapture.h
+++ b/Engine/Plugins/TextureGraph/Source/TextureGraphEngine/Transform/Utility/T_RenderCapture.h
@@ -6,37 +6,41 @@
 #include "Job/Job.h"
 #include "Model/Mix/MixUpdateCycle.h"
 
-class TEXTUREGRAPHENGINE_API T_BeginRenderCapture : public BlobTransform
+#define UE_API TEXTUREGRAPHENGINE_API
+
+class T_BeginRenderCapture : public BlobTransform
 {
 public:
-	T_BeginRenderCapture();
-	virtual							~T_BeginRenderCapture() override;
+	UE_API T_BeginRenderCapture();
+	UE_API virtual							~T_BeginRenderCapture() override;
 
-	virtual Device* TargetDevice(size_t index) const override;
-	virtual AsyncTransformResultPtr	Exec(const TransformArgs& args) override;
+	UE_API virtual Device* TargetDevice(size_t index) const override;
+	UE_API virtual AsyncTransformResultPtr	Exec(const TransformArgs& args) override;
 
 	virtual bool					GeneratesData() const override { return false; }
 
 	////////////////////////////////////////////////////////////////////////// 
 	/// Static functions
 	//////////////////////////////////////////////////////////////////////////
-	static JobUPtr					CreateJob(MixUpdateCyclePtr cycle, int32 targetId);
+	static UE_API JobUPtr					CreateJob(MixUpdateCyclePtr cycle, int32 targetId);
 };
 
 
-class TEXTUREGRAPHENGINE_API T_EndRenderCapture : public BlobTransform
+class T_EndRenderCapture : public BlobTransform
 {
 public:
-	T_EndRenderCapture();
-	virtual							~T_EndRenderCapture() override;
+	UE_API T_EndRenderCapture();
+	UE_API virtual							~T_EndRenderCapture() override;
 
-	virtual Device* TargetDevice(size_t index) const override;
-	virtual AsyncTransformResultPtr	Exec(const TransformArgs& args) override;
+	UE_API virtual Device* TargetDevice(size_t index) const override;
+	UE_API virtual AsyncTransformResultPtr	Exec(const TransformArgs& args) override;
 
 	virtual bool					GeneratesData() const override { return false; }
 
 	////////////////////////////////////////////////////////////////////////// 
 	/// Static functions
 	//////////////////////////////////////////////////////////////////////////
-	static JobUPtr					CreateJob(MixUpdateCyclePtr cycle, int32 targetId);
-};
\ No newline at end of file
+	static UE_API JobUPtr					CreateJob(MixUpdateCyclePtr cycle, int32 targetId);
+};
+
+#undef UE_API
diff --git a/Engine/Plugins/TextureGraph/Source/TextureGraphEngine/Transform/Utility/T_SplitToTiles.h b/Engine/Plugins/TextureGraph/Source/TextureGraphEngine/Transform/Utility/T_SplitToTiles.h
--- a/Engine/Plugins/TextureGraph/Source/TextureGraphEngine/Transform/Utility/T_SplitToTiles.h
+++ b/Engine/Plugins/TextureGraph/Source/TextureGraphEngine/Transform/Utility/T_SplitToTiles.h
@@ -8,6 +8,8 @@
 #include "Model/Mix/MixUpdateCycle.h"
 #include <DataDrivenShaderPlatformInfo.h>
 
+#define UE_API TEXTUREGRAPHENGINE_API
+
 //////////////////////////////////////////////////////////////////////////
 /// Shader
 //////////////////////////////////////////////////////////////////////////
@@ -15,7 +17,7 @@
 class CSH_SplitToTiles : public CmpSH_Base<1, 1, 1>
 {
 public:
-	DECLARE_GLOBAL_SHADER(CSH_SplitToTiles);
+	DECLARE_EXPORTED_GLOBAL_SHADER(CSH_SplitToTiles, UE_API);
 	SHADER_USE_PARAMETER_STRUCT(CSH_SplitToTiles, CmpSH_Base);
 	BEGIN_SHADER_PARAMETER_STRUCT(FParameters, )
 		SHADER_PARAMETER_STRUCT(FTileInfo, TileInfo)
@@ -70,15 +72,17 @@ public:
 /**
  * SplitToTiles Transform
  */
-class TEXTUREGRAPHENGINE_API T_SplitToTiles
+class T_SplitToTiles
 {
 public:
-	T_SplitToTiles();
-	~T_SplitToTiles();
+	UE_API T_SplitToTiles();
+	UE_API ~T_SplitToTiles();
 
 	//////////////////////////////////////////////////////////////////////////
 	/// Static functions
 	//////////////////////////////////////////////////////////////////////////
 	
-	static TiledBlobPtr				Create(MixUpdateCyclePtr Cycle, int32 TargetId, TiledBlobPtr SourceTex);
+	static UE_API TiledBlobPtr				Create(MixUpdateCyclePtr Cycle, int32 TargetId, TiledBlobPtr SourceTex);
 };
+
+#undef UE_API
diff --git a/Engine/Plugins/TextureGraph/Source/TextureGraphEngine/Transform/Utility/T_TextureHistogram.h b/Engine/Plugins/TextureGraph/Source/TextureGraphEngine/Transform/Utility/T_TextureHistogram.h
--- a/Engine/Plugins/TextureGraph/Source/TextureGraphEngine/Transform/Utility/T_TextureHistogram.h
+++ b/Engine/Plugins/TextureGraph/Source/TextureGraphEngine/Transform/Utility/T_TextureHistogram.h
@@ -8,20 +8,24 @@
 #include "Model/Mix/MixUpdateCycle.h"
 #include <DataDrivenShaderPlatformInfo.h>
 
+#define UE_API TEXTUREGRAPHENGINE_API
+
 
 /**
  * T_TextureHistogram Transform
  */
-class TEXTUREGRAPHENGINE_API T_TextureHistogram
+class T_TextureHistogram
 {
 public:
-	T_TextureHistogram();
-	~T_TextureHistogram();
+	UE_API T_TextureHistogram();
+	UE_API ~T_TextureHistogram();
 
-	static TiledBlobPtr	Create(MixUpdateCyclePtr Cycle, TiledBlobPtr SourceTex, int32 TargetId);
+	static UE_API TiledBlobPtr	Create(MixUpdateCyclePtr Cycle, TiledBlobPtr SourceTex, int32 TargetId);
 
-	static TiledBlobPtr	CreateOnService(UMixInterface* InMix, TiledBlobPtr SourceTex, int32 TargetId);
+	static UE_API TiledBlobPtr	CreateOnService(UMixInterface* InMix, TiledBlobPtr SourceTex, int32 TargetId);
 
 private:
-	static TiledBlobPtr			CreateJobAndResult(JobUPtr& OutJob, MixUpdateCyclePtr Cycle, TiledBlobPtr Histogram, int32 TargetId);
+	static UE_API TiledBlobPtr			CreateJobAndResult(JobUPtr& OutJob, MixUpdateCyclePtr Cycle, TiledBlobPtr Histogram, int32 TargetId);
 };
+
+#undef UE_API
diff --git a/Engine/Plugins/TextureGraph/Source/TextureGraphInsight/Public/Model/TextureGraphInsightObserver.h b/Engine/Plugins/TextureGraph/Source/TextureGraphInsight/Public/Model/TextureGraphInsightObserver.h
--- a/Engine/Plugins/TextureGraph/Source/TextureGraphInsight/Public/Model/TextureGraphInsightObserver.h
+++ b/Engine/Plugins/TextureGraph/Source/TextureGraphInsight/Public/Model/TextureGraphInsightObserver.h
@@ -7,54 +7,56 @@
 #include "Data/Blobber.h"
 #include "Job/Scheduler.h"
 
+#define UE_API TEXTUREGRAPHINSIGHT_API
+
 using HashArray = DeviceObserverSource::HashArray;
 
 DECLARE_LOG_CATEGORY_EXTERN(LogTextureGraphInsightObserver, Log, All);
 
 /// Concrete DeviceObserver interface
-class TEXTUREGRAPHINSIGHT_API TextureGraphInsightDeviceObserver : public DeviceObserverSource
+class TextureGraphInsightDeviceObserver : public DeviceObserverSource
 {
 protected:
 	DeviceType DevType;
 
 	/// Protected interface of emitters called by the device to notify the observers
-	virtual void DeviceBuffersUpdated(HashNDescArray&& AddedBuffers, HashArray&& RemovedBuffers) override;
+	UE_API virtual void DeviceBuffersUpdated(HashNDescArray&& AddedBuffers, HashArray&& RemovedBuffers) override;
 public:
-	TextureGraphInsightDeviceObserver();
-	explicit TextureGraphInsightDeviceObserver(DeviceType);
-	virtual ~TextureGraphInsightDeviceObserver() override;
+	UE_API TextureGraphInsightDeviceObserver();
+	UE_API explicit TextureGraphInsightDeviceObserver(DeviceType);
+	UE_API virtual ~TextureGraphInsightDeviceObserver() override;
 };
 
 /// Concrete BlobberObserver interface
-class TEXTUREGRAPHINSIGHT_API TextureGraphInsightBlobberObserver : public BlobberObserverSource
+class TextureGraphInsightBlobberObserver : public BlobberObserverSource
 {
 private:
 protected:
 	/// Protected interface of emitters called by the device to notify the observers
-	virtual void BlobberUpdated(HashArray&& AddedHashes, HashArray&& RemappedHashes) override;
+	UE_API virtual void BlobberUpdated(HashArray&& AddedHashes, HashArray&& RemappedHashes) override;
 
 public:
-	TextureGraphInsightBlobberObserver();
-	virtual ~TextureGraphInsightBlobberObserver() override;
+	UE_API TextureGraphInsightBlobberObserver();
+	UE_API virtual ~TextureGraphInsightBlobberObserver() override;
 };
 
 /// Concrete SchedulerObserver interface
-class TEXTUREGRAPHINSIGHT_API TextureGraphInsightSchedulerObserver : public SchedulerObserverSource
+class TextureGraphInsightSchedulerObserver : public SchedulerObserverSource
 {
 private:
 protected:
 	/// Protected interface of emitters called by the scheduler to notify the observers
-	virtual void Start() override;
-	virtual void UpdateIdle() override;
-	virtual void Stop() override;
-	virtual void BatchAdded(JobBatchPtr Batch) override;
-	virtual void BatchDone(JobBatchPtr Batch) override;
-	virtual void BatchJobsDone(JobBatchPtr Batch) override;
+	UE_API virtual void Start() override;
+	UE_API virtual void UpdateIdle() override;
+	UE_API virtual void Stop() override;
+	UE_API virtual void BatchAdded(JobBatchPtr Batch) override;
+	UE_API virtual void BatchDone(JobBatchPtr Batch) override;
+	UE_API virtual void BatchJobsDone(JobBatchPtr Batch) override;
 
 public:
 
-	TextureGraphInsightSchedulerObserver();
-	virtual ~TextureGraphInsightSchedulerObserver() override;
+	UE_API TextureGraphInsightSchedulerObserver();
+	UE_API virtual ~TextureGraphInsightSchedulerObserver() override;
 };
 
 /// Concrete EngineObserver interface
@@ -62,18 +64,20 @@ public:
 ///	  1/ watching the engine life cycle
 ///   2/ owning the other system observers, and installing them appropirately when an Engine is active
 ///	  3/ notifying Insight
-class TEXTUREGRAPHINSIGHT_API TextureGraphInsightEngineObserver : public EngineObserverSource
+class TextureGraphInsightEngineObserver : public EngineObserverSource
 {
 protected:
 	/// Protected interface of emitters called by the engine to notify the observers
-	virtual void Created() override;
-	virtual void Destroyed() override;
+	UE_API virtual void Created() override;
+	UE_API virtual void Destroyed() override;
 
 public:
-	TextureGraphInsightEngineObserver();
-	virtual ~TextureGraphInsightEngineObserver() override;
+	UE_API TextureGraphInsightEngineObserver();
+	UE_API virtual ~TextureGraphInsightEngineObserver() override;
 
 	std::shared_ptr<TextureGraphInsightDeviceObserver>			_deviceObservers[(uint32)DeviceType::Count];
 	std::shared_ptr<TextureGraphInsightBlobberObserver>		BlobberObserver;
 	std::shared_ptr<TextureGraphInsightSchedulerObserver>		SchedulerObserver;
 };
+
+#undef UE_API
diff --git a/Engine/Plugins/TextureGraph/Source/TextureGraphInsight/Public/Model/TextureGraphInsightRecord.h b/Engine/Plugins/TextureGraph/Source/TextureGraphInsight/Public/Model/TextureGraphInsightRecord.h
--- a/Engine/Plugins/TextureGraph/Source/TextureGraphInsight/Public/Model/TextureGraphInsightRecord.h
+++ b/Engine/Plugins/TextureGraph/Source/TextureGraphInsight/Public/Model/TextureGraphInsightRecord.h
@@ -7,10 +7,12 @@
 #include "Data/Blob.h"
 #include "Device/Device.h"
 
+#define UE_API TEXTUREGRAPHINSIGHT_API
+
 using FStringArray = std::vector<FString>;
 
 template <typename DimType>
-class TEXTUREGRAPHINSIGHT_API T_TiledGrid
+class T_TiledGrid
 {
 public:
 	using Size = DimType;
@@ -48,7 +50,7 @@ public:
 using TiledGrid = T_TiledGrid<uint8_t>;
 
 template <typename Type>
-class TEXTUREGRAPHINSIGHT_API T_Tiled
+class T_Tiled
 {
 public:
 	using Size = typename TiledGrid::Size;
@@ -167,7 +169,7 @@ using RecordIDTiles = T_Tiled<RecordID>;
 /// 
 /// DeviceBuffer Record
 ///
-class TEXTUREGRAPHINSIGHT_API DeviceBufferRecord
+class DeviceBufferRecord
 {
 public:
 	// adding a default constructor to enforce this to be a non POD class to successfully initialize
@@ -175,7 +177,7 @@ public:
 	//RecordID						_selfID;
 	RecordID						SourceID;
 
-	static const DeviceBufferRecord null;
+	static UE_API const DeviceBufferRecord null;
 	DeviceType						DevType = DeviceType::Null;
 	uint64							ID = 0;
 	HashType						HashValue = 0;
@@ -195,14 +197,14 @@ using RecordFromHashMap = std::unordered_map<HashType, RecordID>;
 
 
 /// Blob Record 
-class TEXTUREGRAPHINSIGHT_API BlobRecord
+class BlobRecord
 {
 public:
 	// adding a default constructor to enforce this to be a non POD class to successfully initialize
 	BlobRecord() {}
 	RecordID						SelfID;
 
-	static const BlobRecord			null;
+	static UE_API const BlobRecord			null;
 	uint32							ReplayCount = 0;
 
 	FString							Name;
@@ -237,12 +239,12 @@ using RecordFromBlobIDMap = std::unordered_map<BlobID, RecordID>;
 /// Individual Job Record is the leaf of the data model
 /// It gather meta information to present what the job has done
 /// the inputs and output blob
-class TEXTUREGRAPHINSIGHT_API JobRecord
+class JobRecord
 {
 public:
 	// adding a default constructor to enforce this to be a non POD class to successfully initialize
 	JobRecord() {}
-	static const JobRecord	GNull;
+	static UE_API const JobRecord	GNull;
 
 	RecordID				SelfID;
 
@@ -299,7 +301,7 @@ using JobRecordArray = std::vector< JobRecord >;
 
 /// Batch Record
 /// Hold the Job Records
-class TEXTUREGRAPHINSIGHT_API BatchRecord
+class BatchRecord
 {
 public:
 	// adding a default constructor to enforce this to be a non POD class to successfully initialize
@@ -307,7 +309,7 @@ public:
 
 	RecordID						SelfID;
 
-	static const BatchRecord null;
+	static UE_API const BatchRecord null;
 
 	RecordID						MixID;
 
@@ -343,12 +345,12 @@ using RecordFromBatchIDMap = std::unordered_map<BatchID, RecordID>;
 
 /// Action Record
 /// Hold the Action Records
-class TEXTUREGRAPHINSIGHT_API ActionRecord
+class ActionRecord
 {
 public:
 	// adding a default constructor to enforce this to be a non POD class to successfully initialize
 	ActionRecord() {}
-	static const ActionRecord null;
+	static UE_API const ActionRecord null;
 
 	FString							Name;
 	FString							Meta;
@@ -362,14 +364,14 @@ using RecordFromActionIDMap = std::unordered_map<ActionID, RecordID>;
 
 /// Mix Record
 /// Hold the Mix Records
-class TEXTUREGRAPHINSIGHT_API MixRecord
+class MixRecord
 {
 public:
 	// adding a default constructor to enforce this to be a non POD class to successfully initialize
 	MixRecord() {}
 	RecordID						SelfID;
 
-	static const MixRecord null;
+	static UE_API const MixRecord null;
 
 	RecordID						ParentMixID; // If mix is an instance, it has a parentMix
 	bool							IsInstanceMix() const { return ParentMixID.IsValid(); }
@@ -389,7 +391,7 @@ using RecordFromMixIDMap = std::unordered_map<MixID, RecordID>;
 
 
 /// Session Record container defining the namespace for all the recordIDs
-class TEXTUREGRAPHINSIGHT_API SessionRecord
+class SessionRecord
 {
 private:
 	BatchRecordArray				Batches;
@@ -413,41 +415,41 @@ public:
 	double							StartTime_ms() const { return (Batches.empty() ? 0 : Batches.front().BeginTimeMS); }
 	
 	// Batch records
-	RecordID						FindBatchRecord(BatchID InBatchID) const;
+	UE_API RecordID						FindBatchRecord(BatchID InBatchID) const;
 	using BatchRecordConstructor	= std::function<BatchRecord (RecordID RecId)>;
 	using BatchRecordEditor			= std::function<bool (BatchRecord& r)>;
-	RecordID						NewBatchRecord(BatchID InBatchID, BatchRecordConstructor Constructor);
-	bool							EditBatchRecord(RecordID RecId, BatchRecordEditor Editor);
+	UE_API RecordID						NewBatchRecord(BatchID InBatchID, BatchRecordConstructor Constructor);
+	UE_API bool							EditBatchRecord(RecordID RecId, BatchRecordEditor Editor);
 
 	// Device Buffer records
-	std::tuple<RecordID, bool>		NewDeviceBufferRecord(DeviceType DevType, DeviceBufferID BufferID, const DeviceBufferRecord& Record);
-	RecordID						EraseDeviceBufferRecord(DeviceType DevType, DeviceBufferID BufferID); /// Assuming the db is in the current map, then move its key in the erased map and tag it erased
+	UE_API std::tuple<RecordID, bool>		NewDeviceBufferRecord(DeviceType DevType, DeviceBufferID BufferID, const DeviceBufferRecord& Record);
+	UE_API RecordID						EraseDeviceBufferRecord(DeviceType DevType, DeviceBufferID BufferID); /// Assuming the db is in the current map, then move its key in the erased map and tag it erased
 
 	/// Find the records from the array of bufferIDs & output the leaked buffers
-	RecordIDArray    				FindActiveDeviceBufferRecords(DeviceType DevType, const DeviceBufferIDArray& BufferIDs, RecordIDArray& LeakedBuffers);
+	UE_API RecordIDArray    				FindActiveDeviceBufferRecords(DeviceType DevType, const DeviceBufferIDArray& BufferIDs, RecordIDArray& LeakedBuffers);
 	
-	RecordID						FindDeviceBufferRecord(DeviceBufferID BufferID) const;
-	RecordIDArray					FindDeviceBufferRecords(const DeviceBufferIDArray& BufferIDs) const;
-	RecordIDArray					FetchActiveDeviceBufferIDs(DeviceType DevType) const;
+	UE_API RecordID						FindDeviceBufferRecord(DeviceBufferID BufferID) const;
+	UE_API RecordIDArray					FindDeviceBufferRecords(const DeviceBufferIDArray& BufferIDs) const;
+	UE_API RecordIDArray					FetchActiveDeviceBufferIDs(DeviceType DevType) const;
 
-	RecordID						FindDeviceBufferRecordFromHash(HashType InHash) const;
+	UE_API RecordID						FindDeviceBufferRecordFromHash(HashType InHash) const;
 
 	// Blob records
-	RecordID						FindBlobRecord(BlobID InBlobID) const;
+	UE_API RecordID						FindBlobRecord(BlobID InBlobID) const;
 	using BlobRecordConstructor		= std::function<BlobRecord(RecordID RecId)>;
 	using BlobRecordEditor			= std::function<bool(BlobRecord& r)>;
-	RecordID						NewBlobRecord(BlobID InBlobID, BlobRecordConstructor Constructor);
-	bool							EditBlobRecord(RecordID RecId, BlobRecordEditor Editor);
+	UE_API RecordID						NewBlobRecord(BlobID InBlobID, BlobRecordConstructor Constructor);
+	UE_API bool							EditBlobRecord(RecordID RecId, BlobRecordEditor Editor);
 
-	RecordID						NewActionRecord(ActionID InActionID, const ActionRecord& Record);
-	RecordID						FindActionRecord(ActionID InActionID) const;
+	UE_API RecordID						NewActionRecord(ActionID InActionID, const ActionRecord& Record);
+	UE_API RecordID						FindActionRecord(ActionID InActionID) const;
 
 	// Mix records
-	RecordID						FindMixRecord(MixID MixID) const;
+	UE_API RecordID						FindMixRecord(MixID MixID) const;
 	using MixRecordConstructor		= std::function<MixRecord(RecordID RecId)>;
 	using MixRecordEditor			= std::function<bool(MixRecord& r)>;
-	RecordID						NewMixRecord(BatchID InBatchID, MixRecordConstructor constructor);
-	bool							EditMixRecord(RecordID RecId, MixRecordEditor editor);
+	UE_API RecordID						NewMixRecord(BatchID InBatchID, MixRecordConstructor constructor);
+	UE_API bool							EditMixRecord(RecordID RecId, MixRecordEditor editor);
 
 	FORCEINLINE const DeviceBufferRecord&	GetBuffer(RecordID RecId) const { return ((RecId.Buffer() < DeviceBuffers[RecId.Buffer_DeviceType()].size()) ? DeviceBuffers[RecId.Buffer_DeviceType()][RecId.Buffer()] : DeviceBufferRecord::null); }
 	FORCEINLINE const BlobRecord&			GetBlob(RecordID RecId) const { return ((RecId.Blob() < Blobs.size()) ? Blobs[RecId.Blob()] : BlobRecord::null); }
@@ -460,3 +462,4 @@ public:
 
 using SessionRecordPtr = std::shared_ptr<SessionRecord>;
 
+#undef UE_API
diff --git a/Engine/Plugins/TextureGraph/Source/TextureGraphInsight/Public/Model/TextureGraphInsightSession.h b/Engine/Plugins/TextureGraph/Source/TextureGraphInsight/Public/Model/TextureGraphInsightSession.h
--- a/Engine/Plugins/TextureGraph/Source/TextureGraphInsight/Public/Model/TextureGraphInsightSession.h
+++ b/Engine/Plugins/TextureGraph/Source/TextureGraphInsight/Public/Model/TextureGraphInsightSession.h
@@ -4,13 +4,15 @@
 #include "TextureGraphInsightRecord.h"
 #include "TextureGraphInsightObserver.h"
 
+#define UE_API TEXTUREGRAPHINSIGHT_API
+
 class UMixInterface;
 using MixPtr = UMixInterface*;
 using MixPtrW = UMixInterface*;
 //using MixPtr = std::shared_ptr<UMixInterface>;
 //using MixPtrW = std::weak_ptr<UMixInterface>;
 
-class TEXTUREGRAPHINSIGHT_API TextureGraphInsightSession
+class TextureGraphInsightSession
 {
 public:
 	/// Keep all the pointers to "real" TextureGraph Engine objects in the live runtime cache
@@ -44,48 +46,48 @@ private:
 protected:
 	friend class TextureGraphInsightEngineObserver;
 	/// Protected interface of emitters called by the engine observer
-	void							EngineCreated();
-	void							EngineDestroyed();
-	void							EmitOnEngineReset(bool);
+	UE_API void							EngineCreated();
+	UE_API void							EngineDestroyed();
+	UE_API void							EmitOnEngineReset(bool);
 
 	/// manage DeviceBuffer update
 	friend class TextureGraphInsightDeviceObserver;
-	uint32							ParseDeviceBuffersCache(DeviceType type, RecordIDArray& addedBufferIds);
-	void							DeviceBuffersUpdated(DeviceType DevType, TextureGraphInsightDeviceObserver::HashNDescArray&& AddedBuffers, TextureGraphInsightDeviceObserver::HashArray&& RemovedBuffers);
+	UE_API uint32							ParseDeviceBuffersCache(DeviceType type, RecordIDArray& addedBufferIds);
+	UE_API void							DeviceBuffersUpdated(DeviceType DevType, TextureGraphInsightDeviceObserver::HashNDescArray&& AddedBuffers, TextureGraphInsightDeviceObserver::HashArray&& RemovedBuffers);
 
 	/// Manage Blobber update
 	friend class TextureGraphInsightBlobberObserver;
-	void							BlobberHashesUpdated(TextureGraphInsightBlobberObserver::HashArray&& AddedHashes, TextureGraphInsightBlobberObserver::HashArray&& RemappedHashes);
+	UE_API void							BlobberHashesUpdated(TextureGraphInsightBlobberObserver::HashArray&& AddedHashes, TextureGraphInsightBlobberObserver::HashArray&& RemappedHashes);
 
 	friend class TextureGraphInsightSchedulerObserver;
 	/// Protected interface of emitters called by the scheduler observer
-	void							UpdateIdle();
-	void							BatchAdded(JobBatchPtr Batch);
-	void							BatchDone(JobBatchPtr Batch);
-	void							BatchJobsDone(JobBatchPtr Batch);
+	UE_API void							UpdateIdle();
+	UE_API void							BatchAdded(JobBatchPtr Batch);
+	UE_API void							BatchDone(JobBatchPtr Batch);
+	UE_API void							BatchJobsDone(JobBatchPtr Batch);
 
 	/// manage creation and update of the Batch job records
-	BatchRecord						MakeBatchRecord(RecordID RecId, JobBatchPtr Batch, bool isDone = true);
-	void							UpdateBatchRecord(JobBatchPtr Batch, BatchRecord& record);
-	JobRecord						MakeJobRecord(JobPtrW job, uint32 idx, int32 phaseIdx = 0);
-	void							UpdateJobRecord(JobPtrW job, JobRecord& record, double batchBeginTime_ms, double batchEndTime_ms);
+	UE_API BatchRecord						MakeBatchRecord(RecordID RecId, JobBatchPtr Batch, bool isDone = true);
+	UE_API void							UpdateBatchRecord(JobBatchPtr Batch, BatchRecord& record);
+	UE_API JobRecord						MakeJobRecord(JobPtrW job, uint32 idx, int32 phaseIdx = 0);
+	UE_API void							UpdateJobRecord(JobPtrW job, JobRecord& record, double batchBeginTime_ms, double batchEndTime_ms);
 
 
 	/// manage creation and update of the BlobObj records
-	RecordID						FindOrCreateBlobRecordFromHash(HashType hash);
-	RecordID						RemapBlobHash(HashType oldHash, HashType newHash);
-	RecordID						FindOrCreateBlobRecord(BlobPtr BlobObj, RecordID sourceID = RecordID());
-	BlobRecord						MakeBlobRecord(RecordID RecId, BlobPtr BlobObj, RecordID sourceID = RecordID());
-	void							EditTiledBlobRecord(BlobRecord& rd, const RecordIDTiles& subTiles);
-	void							UpdateBlobRecord(RecordID RecId, BlobPtr BlobObj);
+	UE_API RecordID						FindOrCreateBlobRecordFromHash(HashType hash);
+	UE_API RecordID						RemapBlobHash(HashType oldHash, HashType newHash);
+	UE_API RecordID						FindOrCreateBlobRecord(BlobPtr BlobObj, RecordID sourceID = RecordID());
+	UE_API BlobRecord						MakeBlobRecord(RecordID RecId, BlobPtr BlobObj, RecordID sourceID = RecordID());
+	UE_API void							EditTiledBlobRecord(BlobRecord& rd, const RecordIDTiles& subTiles);
+	UE_API void							UpdateBlobRecord(RecordID RecId, BlobPtr BlobObj);
 	RecordIDArray					NewBlobsToNotify;
 	RecordIDArray					MappedBlobsToNotify;
-	void							EmitOnNewBlob();
-	void							EmitOnMappedBlob();
+	UE_API void							EmitOnNewBlob();
+	UE_API void							EmitOnMappedBlob();
 
 	/// manage creation and update of the MixObj records
-	RecordID						AssociateBatchToMix(RecordID Batch, MixPtr MixObj);
-	RecordID						FindOrCreateMixRecord(MixPtr MixObj);
+	UE_API RecordID						AssociateBatchToMix(RecordID Batch, MixPtr MixObj);
+	UE_API RecordID						FindOrCreateMixRecord(MixPtr MixObj);
 
 	/// manage creation and update of the buffer records
 	RecordIDArray					NewBuffersToNotify;
@@ -102,14 +104,14 @@ protected:
 	int32							EngineID = -1;
 public:
 
-	TextureGraphInsightSession();
-    ~TextureGraphInsightSession();
+	UE_API TextureGraphInsightSession();
+    UE_API ~TextureGraphInsightSession();
 
 	FORCEINLINE SessionRecord&		GetRecord() { std::lock_guard<std::mutex> Lock(RecordMutex); return (*Record); }
 	FORCEINLINE LiveCache&			GetCache() { std::lock_guard<std::mutex> Lock(RecordMutex); return (*Cache); }
 
 	/// Send the specified record to be inspected
-	void							SendToInspector(RecordID RecId) const;
+	UE_API void							SendToInspector(RecordID RecId) const;
 
 	/// Notifier subscription
 	DECLARE_MULTICAST_DELEGATE_OneParam(FOnEngineReset, int32);
@@ -142,8 +144,8 @@ public:
 
 	// Trigger Replay of Batch and jobs
 	// this call is async and returns if the scheduling of the replay is successful or not
-	bool							ReplayBatch(RecordID batchRecordID, bool captureRenderDoc);
-	bool							ReplayJob(RecordID jobRecordID, bool captureRenderDoc);
+	UE_API bool							ReplayBatch(RecordID batchRecordID, bool captureRenderDoc);
+	UE_API bool							ReplayJob(RecordID jobRecordID, bool captureRenderDoc);
 
 protected:
 	FOnEngineReset		OnEngineResetEvent;
@@ -168,3 +170,4 @@ protected:
 	FOnRecordID			OnBufferInspectedEvent;
 };
 
+#undef UE_API
diff --git a/Engine/Plugins/TextureGraph/Source/TextureGraphInsight/Public/TextureGraphInsight.h b/Engine/Plugins/TextureGraph/Source/TextureGraphInsight/Public/TextureGraphInsight.h
--- a/Engine/Plugins/TextureGraph/Source/TextureGraphInsight/Public/TextureGraphInsight.h
+++ b/Engine/Plugins/TextureGraph/Source/TextureGraphInsight/Public/TextureGraphInsight.h
@@ -5,6 +5,8 @@
 #include <memory>
 #include "Modules/ModuleManager.h"
 
+#define UE_API TEXTUREGRAPHINSIGHT_API
+
 DECLARE_LOG_CATEGORY_EXTERN(LogTextureGraphInsight, Log, All);
 
 class TextureGraphInsightSchedulerObserver;
@@ -19,19 +21,19 @@ public:
 
 struct RecordID;
 
-class TEXTUREGRAPHINSIGHT_API TextureGraphInsight
+class TextureGraphInsight
 {
 private:
-	static TextureGraphInsight* GInstance;					/// Static instance. Use Create to create a new instance. Destroy the current instance first
-	TextureGraphInsight();
-	~TextureGraphInsight();
+	static UE_API TextureGraphInsight* GInstance;					/// Static instance. Use Create to create a new instance. Destroy the current instance first
+	UE_API TextureGraphInsight();
+	UE_API ~TextureGraphInsight();
 
 	TextureGraphInsightSessionPtr					Session;
 public:
-	static bool								Create();  /// Create the Insight Instance ONLY if no other currently created AND if a TextureGraphEngine is already created
+	static UE_API bool								Create();  /// Create the Insight Instance ONLY if no other currently created AND if a TextureGraphEngine is already created
 
 	/// Destroy the current instance
-	static bool								Destroy(); /// Destroy the instance of Insight if it exists.
+	static UE_API bool								Destroy(); /// Destroy the instance of Insight if it exists.
 
 	//////////////////////////////////////////////////////////////////////////
 	/// Inline functions
@@ -44,3 +46,5 @@ public:
 // Macro to control in one place how we display hash value
 //#define HashToFString(h) FString::Printf(TEXT("%20llu"), (h))
 #define HashToFString(h) FString::Printf(TEXT("%llX"), (h))
+
+#undef UE_API
diff --git a/Engine/Plugins/TextureGraph/Source/TextureGraphInsight/Public/View/STextureGraphInsightActionView.h b/Engine/Plugins/TextureGraph/Source/TextureGraphInsight/Public/View/STextureGraphInsightActionView.h
--- a/Engine/Plugins/TextureGraph/Source/TextureGraphInsight/Public/View/STextureGraphInsightActionView.h
+++ b/Engine/Plugins/TextureGraph/Source/TextureGraphInsight/Public/View/STextureGraphInsightActionView.h
@@ -6,15 +6,17 @@
 #include "Model/TextureGraphInsightSession.h"
 #include <Widgets/Views/STreeView.h>
 
+#define UE_API TEXTUREGRAPHINSIGHT_API
+
 class STextureGraphInsightActionViewRow; /// Declare the concrete type of widget used for the raws of the view. Defined in the cpp file
 
-class TEXTUREGRAPHINSIGHT_API STextureGraphInsightActionView : public SCompoundWidget
+class STextureGraphInsightActionView : public SCompoundWidget
 {
 public:
 	SLATE_BEGIN_ARGS(STextureGraphInsightActionView) {}
 	SLATE_END_ARGS()
 
-	void Construct(const FArguments& Args);
+	UE_API void Construct(const FArguments& Args);
 
 	// Item Types
 	class FItemData;
@@ -31,10 +33,10 @@ public:
 	using SItemTableView = STreeView< FItem >;
 
 	// Standard delegates for the view
-	TSharedRef<ITableRow>	OnGenerateRowForView(FItem item, const TSharedRef<STableViewBase>& OwnerTable);
+	UE_API TSharedRef<ITableRow>	OnGenerateRowForView(FItem item, const TSharedRef<STableViewBase>& OwnerTable);
 	FORCEINLINE void		OnGetChildrenForView(FItem item, FItemArray& children) { children = item->_children; };
-	void					OnClickItemForView(FItem item);
-	void					OnDoubleClickItemForView(FItem item);
+	UE_API void					OnClickItemForView(FItem item);
+	UE_API void					OnDoubleClickItemForView(FItem item);
 
 	/// The list of root items
 	FItemArray _rootItems;
@@ -42,6 +44,8 @@ public:
 	// The TreeView widget
 	TSharedPtr<SItemTableView> _tableView;
 
-	void OnActionNew(RecordID rid);
-	void OnEngineReset(int);
+	UE_API void OnActionNew(RecordID rid);
+	UE_API void OnEngineReset(int);
 };
+
+#undef UE_API
diff --git a/Engine/Plugins/TextureGraph/Source/TextureGraphInsight/Public/View/STextureGraphInsightBlobView.h b/Engine/Plugins/TextureGraph/Source/TextureGraphInsight/Public/View/STextureGraphInsightBlobView.h
--- a/Engine/Plugins/TextureGraph/Source/TextureGraphInsight/Public/View/STextureGraphInsightBlobView.h
+++ b/Engine/Plugins/TextureGraph/Source/TextureGraphInsight/Public/View/STextureGraphInsightBlobView.h
@@ -7,11 +7,13 @@
 #include <Widgets/Layout/SBorder.h>
 #include <UObject/GCObject.h>
 
+#define UE_API TEXTUREGRAPHINSIGHT_API
+
 class STextureGraphInsightDeviceBufferView;
 class UMaterialInstanceDynamic;
 struct FSlateBrush;
 struct FSlateColorBrush;
-class TEXTUREGRAPHINSIGHT_API STextureGraphInsightBlobView : public SBorder, public FGCObject /// Need FGCObject to control garbage collection of objects
+class STextureGraphInsightBlobView : public SBorder, public FGCObject /// Need FGCObject to control garbage collection of objects
 {
 public:
 	SLATE_BEGIN_ARGS(STextureGraphInsightBlobView) :
@@ -22,7 +24,7 @@ public:
 		SLATE_ARGUMENT(BoolTiles, tilesMask)
 	SLATE_END_ARGS()
 
-	void Construct(const FArguments& Args);
+	UE_API void Construct(const FArguments& Args);
 
 	RecordID _recordID;
 	RecordID _inspectedID;
@@ -34,11 +36,11 @@ public:
 	// Border brush
 	TSharedPtr<FSlateColorBrush> _borderBrush;
 
-	FReply OnInspect();
+	UE_API FReply OnInspect();
 
 	// Capture mouse clicks here
-	virtual FReply OnMouseButtonDown(const FGeometry& MyGeometry, const FPointerEvent& MouseEvent) override;
-	virtual FReply OnMouseButtonDoubleClick(const FGeometry& InMyGeometry, const FPointerEvent& InMouseEvent) override;
+	UE_API virtual FReply OnMouseButtonDown(const FGeometry& MyGeometry, const FPointerEvent& MouseEvent) override;
+	UE_API virtual FReply OnMouseButtonDoubleClick(const FGeometry& InMyGeometry, const FPointerEvent& InMouseEvent) override;
 
 
 	virtual void AddReferencedObjects(FReferenceCollector& Collector) override
@@ -52,3 +54,5 @@ public:
 	}
 
 };
+
+#undef UE_API
diff --git a/Engine/Plugins/TextureGraph/Source/TextureGraphInsight/Public/View/STextureGraphInsightDeviceBufferView.h b/Engine/Plugins/TextureGraph/Source/TextureGraphInsight/Public/View/STextureGraphInsightDeviceBufferView.h
--- a/Engine/Plugins/TextureGraph/Source/TextureGraphInsight/Public/View/STextureGraphInsightDeviceBufferView.h
+++ b/Engine/Plugins/TextureGraph/Source/TextureGraphInsight/Public/View/STextureGraphInsightDeviceBufferView.h
@@ -7,9 +7,11 @@
 #include <Widgets/Layout/SBorder.h>
 #include <UObject/GCObject.h>
 
+#define UE_API TEXTUREGRAPHINSIGHT_API
+
 class UMaterialInstanceDynamic;
 
-class TEXTUREGRAPHINSIGHT_API STextureGraphInsightDeviceBufferView : public SBorder, public FGCObject /// Need FGCObject to control garbage collection of objects
+class STextureGraphInsightDeviceBufferView : public SBorder, public FGCObject /// Need FGCObject to control garbage collection of objects
 {
 public:
 	SLATE_BEGIN_ARGS(STextureGraphInsightDeviceBufferView) :
@@ -21,7 +23,7 @@ public:
 		SLATE_ARGUMENT(bool, withDescription)
 	SLATE_END_ARGS()
 
-	void Construct(const FArguments& Args);
+	UE_API void Construct(const FArguments& Args);
 
 	RecordID _recordID;
 	RecordID _blobID;
@@ -33,8 +35,8 @@ public:
 	TSharedPtr<FSlateBrush> _brush;
 	TObjectPtr<UMaterialInstanceDynamic> _brushMaterial = nullptr;
 
-	void CreateWidgetBrush(UTexture* blobTexture);
-	void MakeWidgetMaterial();
+	UE_API void CreateWidgetBrush(UTexture* blobTexture);
+	UE_API void MakeWidgetMaterial();
 
 	virtual void AddReferencedObjects(FReferenceCollector& Collector) override
 	{
@@ -46,3 +48,5 @@ public:
 		return TEXT("STextureGraphInsightDeviceBufferView");
 	}
 };
+
+#undef UE_API
diff --git a/Engine/Plugins/TextureGraph/Source/TextureGraphInsight/Public/View/STextureGraphInsightDeviceView.h b/Engine/Plugins/TextureGraph/Source/TextureGraphInsight/Public/View/STextureGraphInsightDeviceView.h
--- a/Engine/Plugins/TextureGraph/Source/TextureGraphInsight/Public/View/STextureGraphInsightDeviceView.h
+++ b/Engine/Plugins/TextureGraph/Source/TextureGraphInsight/Public/View/STextureGraphInsightDeviceView.h
@@ -7,12 +7,14 @@
 #include <Widgets/SCompoundWidget.h>
 #include <Widgets/Views/STreeView.h>
 
+#define UE_API TEXTUREGRAPHINSIGHT_API
+
 class STextureGraphInsightDeviceListViewRow; /// Declare the concrete type of widget used for the raws of the view. Defined in the cpp file
 class STableViewBase;
 class ITableRow;
 class STextBlock;
 
-class TEXTUREGRAPHINSIGHT_API STextureGraphInsightDeviceListView : public SCompoundWidget
+class STextureGraphInsightDeviceListView : public SCompoundWidget
 {
 public:
 	SLATE_BEGIN_ARGS(STextureGraphInsightDeviceListView)
@@ -22,7 +24,7 @@ public:
 	SLATE_ARGUMENT(bool, inspectDevices)
 	SLATE_END_ARGS()
 
-	void Construct(const FArguments& Args);
+	UE_API void Construct(const FArguments& Args);
 
 	// TreeView Item Types
 	class FItemData;
@@ -43,11 +45,11 @@ public:
 	using SItemTableView = STreeView< FItem >;
 
 	// Standard delegates for the view
-	TSharedRef<ITableRow>	OnGenerateRowForView(FItem item, const TSharedRef<STableViewBase>& OwnerTable);
+	UE_API TSharedRef<ITableRow>	OnGenerateRowForView(FItem item, const TSharedRef<STableViewBase>& OwnerTable);
 	FORCEINLINE void		OnGetChildrenForView(FItem item, FItemArray& children) { children = item->_children; if (item->_removedRoot) children.Add(item->_removedRoot); };
-	void OnClickItemForView(FItem item);
-	void OnDoubleClickItemForView(FItem item);
-	void OnInspect(FItem item);
+	UE_API void OnClickItemForView(FItem item);
+	UE_API void OnDoubleClickItemForView(FItem item);
+	UE_API void OnInspect(FItem item);
 	
 	RecordID	_rootDevice; /// Can use an empty RecordID parameter with type DeviceBuffer and the expected Device type to only care about that device
 
@@ -63,24 +65,24 @@ public:
 	// The TreeView widget
 	TSharedPtr<SItemTableView> _tableView;
 
-	void OnDeviceBufferAdded(const RecordIDArray& rids);
-	void OnDeviceBufferRemoved(const RecordIDArray& rids);
-	void OnEngineReset(int);
+	UE_API void OnDeviceBufferAdded(const RecordIDArray& rids);
+	UE_API void OnDeviceBufferRemoved(const RecordIDArray& rids);
+	UE_API void OnEngineReset(int);
 
-	void Reset();
+	UE_API void Reset();
 
-	void RefreshRootItems();
+	UE_API void RefreshRootItems();
 };
 
 inline bool operator ==(const STextureGraphInsightDeviceListView::FItem& i,  const RecordID& rid) { return i->_recordID.Buffer() == rid.Buffer(); }
 
-class TEXTUREGRAPHINSIGHT_API STextureGraphInsightDeviceView : public SCompoundWidget
+class STextureGraphInsightDeviceView : public SCompoundWidget
 {
 	SLATE_BEGIN_ARGS(STextureGraphInsightDeviceView) {}
 	SLATE_ARGUMENT(DeviceType, deviceType)
 	SLATE_END_ARGS()
 
-	void Construct(const FArguments& Args);
+	UE_API void Construct(const FArguments& Args);
 
 	DeviceType _deviceType;
 	
@@ -94,20 +96,21 @@ class TEXTUREGRAPHINSIGHT_API STextureGraphInsightDeviceView : public SCompoundW
 
 	TSharedPtr<STextBlock> _threadCount;
 	
-	void Refresh();
+	UE_API void Refresh();
 };
 
-class TEXTUREGRAPHINSIGHT_API STextureGraphInsightDeviceManagerView : public SCompoundWidget
+class STextureGraphInsightDeviceManagerView : public SCompoundWidget
 {
 	SLATE_BEGIN_ARGS(STextureGraphInsightDeviceManagerView) {}
 	SLATE_END_ARGS()
 
-	void Construct(const FArguments& Args);
+	UE_API void Construct(const FArguments& Args);
 
 	std::vector < TSharedPtr< STextureGraphInsightDeviceView >> _deviceViews;
 
-	void Reset(); // Rebuild the full ui of this widget, called by update
-	void OnUpdate();
-	void OnEngineReset(int32);
+	UE_API void Reset(); // Rebuild the full ui of this widget, called by update
+	UE_API void OnUpdate();
+	UE_API void OnEngineReset(int32);
 };
 
+#undef UE_API
diff --git a/Engine/Plugins/TextureGraph/Source/TextureGraphInsight/Public/View/STextureGraphInsightInspectorView.h b/Engine/Plugins/TextureGraph/Source/TextureGraphInsight/Public/View/STextureGraphInsightInspectorView.h
--- a/Engine/Plugins/TextureGraph/Source/TextureGraphInsight/Public/View/STextureGraphInsightInspectorView.h
+++ b/Engine/Plugins/TextureGraph/Source/TextureGraphInsight/Public/View/STextureGraphInsightInspectorView.h
@@ -5,18 +5,20 @@
 #include "Model/TextureGraphInsightSession.h"
 #include <Widgets/Views/STileView.h>
 
+#define UE_API TEXTUREGRAPHINSIGHT_API
+
 class STextureGraphInsightDeviceBufferView;
 class STextureGraphInsightBlobView;
 class STextureGraphInsightRecordTrailView;
 
-class TEXTUREGRAPHINSIGHT_API STextureGraphInsightBatchInspectorView : public SCompoundWidget
+class STextureGraphInsightBatchInspectorView : public SCompoundWidget
 {
 public:
 	SLATE_BEGIN_ARGS(STextureGraphInsightBatchInspectorView) {}
 	SLATE_ARGUMENT(RecordID, recordID)
 		SLATE_END_ARGS()
 
-		void Construct(const FArguments& Args);
+		UE_API void Construct(const FArguments& Args);
 
 	// ListView Item Types presenting various results of a batch
 	class FItemData;
@@ -34,7 +36,7 @@ public:
 	using SItemTileView = STileView< FItem >;
 
 	// Standard delegates for the tree view
-	TSharedRef<ITableRow> OnGenerateTileForView(FItem item, const TSharedRef<STableViewBase>& OwnerTable);
+	UE_API TSharedRef<ITableRow> OnGenerateTileForView(FItem item, const TSharedRef<STableViewBase>& OwnerTable);
 
 	/// The list of root items
 	FItemArray _rootItems;
@@ -43,17 +45,17 @@ public:
 
 
 	// Inspect the specified batch (Called by Constructor with recordID param)
-	void InspectBatch(RecordID job);
+	UE_API void InspectBatch(RecordID job);
 };
 
-class TEXTUREGRAPHINSIGHT_API STextureGraphInsightJobInspectorView : public SCompoundWidget
+class STextureGraphInsightJobInspectorView : public SCompoundWidget
 {
 public:
 	SLATE_BEGIN_ARGS(STextureGraphInsightJobInspectorView) {}
 	SLATE_ARGUMENT(RecordID, recordID)
 	SLATE_END_ARGS()
 
-	void Construct(const FArguments& Args);
+	UE_API void Construct(const FArguments& Args);
 
 	// ListView Item Types presenting various input of a job
 	class FItemData;
@@ -73,7 +75,7 @@ public:
 	using SItemTileView = STileView< FItem >;
 
 	// Standard delegates for the tree view
-	TSharedRef<ITableRow> OnGenerateTileForView(FItem item, const TSharedRef<STableViewBase>& OwnerTable);
+	UE_API TSharedRef<ITableRow> OnGenerateTileForView(FItem item, const TSharedRef<STableViewBase>& OwnerTable);
 
 	// The output blob view generated by the job
 	TSharedPtr<STextureGraphInsightBlobView> _outBlobView;
@@ -85,47 +87,47 @@ public:
 
 
 	// Inspect the specified Job (Called by Constructor with recordID param)
-	void InspectJob(RecordID job);
+	UE_API void InspectJob(RecordID job);
 };
 
 
-class TEXTUREGRAPHINSIGHT_API STextureGraphInsightBlobInspectorView : public SCompoundWidget
+class STextureGraphInsightBlobInspectorView : public SCompoundWidget
 {
 public:
 	SLATE_BEGIN_ARGS(STextureGraphInsightBlobInspectorView) {}
 	SLATE_ARGUMENT(RecordID, recordID)
 	SLATE_END_ARGS()
 
-	void Construct(const FArguments& Args);
+	UE_API void Construct(const FArguments& Args);
 
 	// The main blob view generated by the job
 	TSharedPtr<STextureGraphInsightBlobView> _blobView;
 
 	// Inspect the specified Blob (Called by Constructor with recordID param)
-	void InspectBlob(RecordID blob);
+	UE_API void InspectBlob(RecordID blob);
 };
 
-class TEXTUREGRAPHINSIGHT_API STextureGraphInsightDeviceBufferInspectorView : public SCompoundWidget
+class STextureGraphInsightDeviceBufferInspectorView : public SCompoundWidget
 {
 public:
 	SLATE_BEGIN_ARGS(STextureGraphInsightDeviceBufferInspectorView) {}
 	SLATE_ARGUMENT(RecordID, recordID)
 	SLATE_END_ARGS()
 
-	void Construct(const FArguments& Args);
+	UE_API void Construct(const FArguments& Args);
 
 
 	// Inspect the specified buffer (Called by Constructor with recordID param)
-	void InspectBuffer(RecordID buffer);
+	UE_API void InspectBuffer(RecordID buffer);
 };
 
-class TEXTUREGRAPHINSIGHT_API STextureGraphInsightInspectorView : public SCompoundWidget
+class STextureGraphInsightInspectorView : public SCompoundWidget
 {
 public:
 	SLATE_BEGIN_ARGS(STextureGraphInsightInspectorView) {}
 	SLATE_END_ARGS()
 
-	void Construct(const FArguments& Args);
+	UE_API void Construct(const FArguments& Args);
 
 	TSharedPtr<STextureGraphInsightRecordTrailView> _recordTrail;
 
@@ -133,6 +135,8 @@ public:
 
 
 
-	void OnInspected(RecordID rid);
-	void OnEngineReset(int32);
-};
\ No newline at end of file
+	UE_API void OnInspected(RecordID rid);
+	UE_API void OnEngineReset(int32);
+};
+
+#undef UE_API
diff --git a/Engine/Plugins/TextureGraph/Source/TextureGraphInsight/Public/View/STextureGraphInsightMixView.h b/Engine/Plugins/TextureGraph/Source/TextureGraphInsight/Public/View/STextureGraphInsightMixView.h
--- a/Engine/Plugins/TextureGraph/Source/TextureGraphInsight/Public/View/STextureGraphInsightMixView.h
+++ b/Engine/Plugins/TextureGraph/Source/TextureGraphInsight/Public/View/STextureGraphInsightMixView.h
@@ -6,17 +6,19 @@
 #include "Model/TextureGraphInsightSession.h"
 #include "Widgets/Views/STreeView.h"
 #include "Widgets/SCompoundWidget.h"
+
+#define UE_API TEXTUREGRAPHINSIGHT_API
 class STextureGraphInsightMixListViewRow; /// Declare the concrete type of widget used for the raws of the view. Defined in the cpp file
 class STableViewBase;
 class ITableRow;
 
-class TEXTUREGRAPHINSIGHT_API STextureGraphInsightMixListView : public SCompoundWidget
+class STextureGraphInsightMixListView : public SCompoundWidget
 {
 public:
 	SLATE_BEGIN_ARGS(STextureGraphInsightMixListView) {}
 	SLATE_END_ARGS()
 
-		void Construct(const FArguments& Args);
+		UE_API void Construct(const FArguments& Args);
 
 	// TreeView Item Types
 	class FItemData;
@@ -35,10 +37,10 @@ public:
 
 
 	// Standard delegates for the view
-	TSharedRef<ITableRow>	OnGenerateRowForView(FItem item, const TSharedRef<STableViewBase>& OwnerTable);
+	UE_API TSharedRef<ITableRow>	OnGenerateRowForView(FItem item, const TSharedRef<STableViewBase>& OwnerTable);
 	FORCEINLINE void		OnGetChildrenForView(FItem item, FItemArray& children) { children = item->_children; };
-	void					OnClickItemForView(FItem item);
-	void					OnDoubleClickItemForView(FItem item);
+	UE_API void					OnClickItemForView(FItem item);
+	UE_API void					OnDoubleClickItemForView(FItem item);
 
 	/// The list of root items
 	FItemArray _rootItems;
@@ -46,11 +48,13 @@ public:
 	// The TreeView widget
 	TSharedPtr<SItemTableView> _tableView;
 
-	void RefreshRootItems();
+	UE_API void RefreshRootItems();
+
+	UE_API void OnMixNew(RecordID rid);
+	UE_API void OnMixUpdate(RecordID rid);
+	UE_API void OnEngineReset(int id);
 
-	void OnMixNew(RecordID rid);
-	void OnMixUpdate(RecordID rid);
-	void OnEngineReset(int id);
+	UE_API void UpdateItemData(FItemData& item);
+};
 
-	void UpdateItemData(FItemData& item);
-};
\ No newline at end of file
+#undef UE_API
diff --git a/Engine/Plugins/TextureGraph/Source/TextureGraphInsight/Public/View/STextureGraphInsightRecordTrailView.h b/Engine/Plugins/TextureGraph/Source/TextureGraphInsight/Public/View/STextureGraphInsightRecordTrailView.h
--- a/Engine/Plugins/TextureGraph/Source/TextureGraphInsight/Public/View/STextureGraphInsightRecordTrailView.h
+++ b/Engine/Plugins/TextureGraph/Source/TextureGraphInsight/Public/View/STextureGraphInsightRecordTrailView.h
@@ -6,8 +6,10 @@
 
 #include "Model/TextureGraphInsightSession.h"
 
+#define UE_API TEXTUREGRAPHINSIGHT_API
 
-class TEXTUREGRAPHINSIGHT_API STextureGraphInsightRecordTrailView : public SCompoundWidget
+
+class STextureGraphInsightRecordTrailView : public SCompoundWidget
 {
 	TSharedPtr< SBreadcrumbTrail<RecordID> > _trail;
 public:
@@ -15,13 +17,14 @@ public:
 		SLATE_ARGUMENT(RecordID, recordID)
 	SLATE_END_ARGS()
 
-	void Construct(const FArguments& Args);
+	UE_API void Construct(const FArguments& Args);
 
-	void Refresh(RecordID rid);
+	UE_API void Refresh(RecordID rid);
 
-	static RecordID FindRootRecord(RecordID rid);
+	static UE_API RecordID FindRootRecord(RecordID rid);
 
-	void OnCrumbClicked(const RecordID& blob);
+	UE_API void OnCrumbClicked(const RecordID& blob);
 };
 
 
+#undef UE_API
diff --git a/Engine/Plugins/TextureGraph/Source/TextureGraphInsight/Public/View/STextureGraphInsightResourceView.h b/Engine/Plugins/TextureGraph/Source/TextureGraphInsight/Public/View/STextureGraphInsightResourceView.h
--- a/Engine/Plugins/TextureGraph/Source/TextureGraphInsight/Public/View/STextureGraphInsightResourceView.h
+++ b/Engine/Plugins/TextureGraph/Source/TextureGraphInsight/Public/View/STextureGraphInsightResourceView.h
@@ -6,17 +6,19 @@
 #include "Model/TextureGraphInsightSession.h"
 #include "Widgets/SCompoundWidget.h"
 #include "Widgets/Views/STreeView.h"
+
+#define UE_API TEXTUREGRAPHINSIGHT_API
 class STextureGraphInsightResourceViewRow; /// Declare the concrete type of widget used for the raws of the view. Defined in the cpp file
 class ITableRow;
 class STableViewBase;
 
-class TEXTUREGRAPHINSIGHT_API STextureGraphInsightResourceView : public SCompoundWidget
+class STextureGraphInsightResourceView : public SCompoundWidget
 {
 public:
 	SLATE_BEGIN_ARGS(STextureGraphInsightResourceView) {}
 	SLATE_END_ARGS()
 
-	void Construct(const FArguments& Args);
+	UE_API void Construct(const FArguments& Args);
 
 	// TreeView Item Types
 	class FItemData;
@@ -35,10 +37,10 @@ public:
 
 
 	// Standard delegates for the view
-	TSharedRef<ITableRow>	OnGenerateRowForView(FItem item, const TSharedRef<STableViewBase>& OwnerTable);
+	UE_API TSharedRef<ITableRow>	OnGenerateRowForView(FItem item, const TSharedRef<STableViewBase>& OwnerTable);
 	FORCEINLINE void		OnGetChildrenForView(FItem item, FItemArray& children) { children = item->_children; };
-	void					OnClickItemForView(FItem item);
-	void					OnDoubleClickItemForView(FItem item);
+	UE_API void					OnClickItemForView(FItem item);
+	UE_API void					OnDoubleClickItemForView(FItem item);
 
 	/// The list of root items
 	FItemArray _rootItems;
@@ -46,9 +48,11 @@ public:
 	// The TreeView widget
 	TSharedPtr<SItemTableView> _tableView;
 
-	void RefreshRootItems();
+	UE_API void RefreshRootItems();
 
-	void OnBlobNew(const RecordIDArray& rids);
-	void OnBlobMapped(const RecordIDArray& rids);
-	void OnEngineReset(int id);
+	UE_API void OnBlobNew(const RecordIDArray& rids);
+	UE_API void OnBlobMapped(const RecordIDArray& rids);
+	UE_API void OnEngineReset(int id);
 };
+
+#undef UE_API
diff --git a/Engine/Plugins/TextureGraph/Source/TextureGraphInsight/Public/View/STextureGraphInsightSessionView.h b/Engine/Plugins/TextureGraph/Source/TextureGraphInsight/Public/View/STextureGraphInsightSessionView.h
--- a/Engine/Plugins/TextureGraph/Source/TextureGraphInsight/Public/View/STextureGraphInsightSessionView.h
+++ b/Engine/Plugins/TextureGraph/Source/TextureGraphInsight/Public/View/STextureGraphInsightSessionView.h
@@ -7,11 +7,13 @@
 #include "Widgets/SCompoundWidget.h"
 #include "Widgets/Views/STreeView.h"
 
+#define UE_API TEXTUREGRAPHINSIGHT_API
+
 class STextureGraphInsightBatchJobViewRow; /// Declare the concrete type of widget used for the raws of the tree view. Defined in the cpp file
 class ITableRow;
 class STableViewBase;
 
-class TEXTUREGRAPHINSIGHT_API STextureGraphInsightBatchJobView : public SCompoundWidget
+class STextureGraphInsightBatchJobView : public SCompoundWidget
 {
 public:
 	SLATE_BEGIN_ARGS(STextureGraphInsightBatchJobView)
@@ -20,7 +22,7 @@ public:
 	SLATE_ARGUMENT(bool, inspectOnSimpleClick)
 	SLATE_END_ARGS()
 
-	void Construct(const FArguments& Args);
+	UE_API void Construct(const FArguments& Args);
 
 	// TreeView Item Types
 	class FItemData;
@@ -38,13 +40,13 @@ public:
 	using SItemTreeView = STreeView< FItem >;
 
 	// Standard delegates for the tree view
-	TSharedRef<ITableRow>	OnGenerateRowForTree(FItem item, const TSharedRef<STableViewBase>& OwnerTable);
+	UE_API TSharedRef<ITableRow>	OnGenerateRowForTree(FItem item, const TSharedRef<STableViewBase>& OwnerTable);
 	FORCEINLINE void		OnGetChildrenForView(FItem item, FItemArray& children) { children = item->_children; };
-	void					OnClickItemForTree(FItem item);
-	void					OnDoubleClickItemForTree(FItem item);
-	TSharedPtr<SWidget>		OnContextMenuOpeningForTree();
-	TSharedPtr<SWidget>		OnContextMenuBatch(FItem batchItem);
-	TSharedPtr<SWidget>		OnContextMenuJob(FItem jobItem);
+	UE_API void					OnClickItemForTree(FItem item);
+	UE_API void					OnDoubleClickItemForTree(FItem item);
+	UE_API TSharedPtr<SWidget>		OnContextMenuOpeningForTree();
+	UE_API TSharedPtr<SWidget>		OnContextMenuBatch(FItem batchItem);
+	UE_API TSharedPtr<SWidget>		OnContextMenuJob(FItem jobItem);
 
 	// Inspect on double click always
 	// eventually trigger inspect on simple click
@@ -57,22 +59,24 @@ public:
 	TSharedPtr<SItemTreeView> _treeView;
 
 
-	void OnReplayBatch(RecordID batchId, bool captureRenderDoc);
-	void OnReplayJob(RecordID batchId, bool captureRenderDoc);
+	UE_API void OnReplayBatch(RecordID batchId, bool captureRenderDoc);
+	UE_API void OnReplayJob(RecordID batchId, bool captureRenderDoc);
 
-	void OnBatchNew(RecordID batch);
-	void OnBatchUpdate(RecordID batch);
-	void OnEngineReset(int32);
+	UE_API void OnBatchNew(RecordID batch);
+	UE_API void OnBatchUpdate(RecordID batch);
+	UE_API void OnEngineReset(int32);
 };
 
 
-class TEXTUREGRAPHINSIGHT_API STextureGraphInsightSessionView : public SCompoundWidget
+class STextureGraphInsightSessionView : public SCompoundWidget
 {
 public:
 	SLATE_BEGIN_ARGS(STextureGraphInsightSessionView) {}
 	SLATE_END_ARGS()
 
-	void Construct(const FArguments& Args);
+	UE_API void Construct(const FArguments& Args);
 
 	TSharedPtr<STextureGraphInsightBatchJobView> _batchJobView;
 };
+
+#undef UE_API
