Fixed out of memory issues when running TG_AsyncTasks multiple times

diff --git a/Engine/Plugins/TextureGraph/Source/TextureGraph/Private/Blueprint/TG_AsyncRenderTask.cpp b/Engine/Plugins/TextureGraph/Source/TextureGraph/Private/Blueprint/TG_AsyncRenderTask.cpp
--- a/Engine/Plugins/TextureGraph/Source/TextureGraph/Private/Blueprint/TG_AsyncRenderTask.cpp
+++ b/Engine/Plugins/TextureGraph/Source/TextureGraph/Private/Blueprint/TG_AsyncRenderTask.cpp
@@ -23,7 +23,7 @@ UTG_AsyncRenderTaskBase::UTG_AsyncRenderTaskBase(const FObjectInitializer& Objec
 
 JobBatchPtr UTG_AsyncRenderTaskBase::PrepareActivate(JobBatchPtr Batch, bool bIsAsync)
 {
-	UE_LOG(LogTextureGraph, Log, TEXT("UTG_AsyncRenderTaskBase::PrepareActivate"));
+	UE_LOG(LogTextureGraph, Verbose, TEXT("UTG_AsyncRenderTaskBase::PrepareActivate"));
 
 	if (TextureGraphPtr == nullptr)
 	{
@@ -56,7 +56,7 @@ JobBatchPtr UTG_AsyncRenderTaskBase::PrepareActivate(JobBatchPtr Batch, bool bIs
 
 const TArray<UTextureRenderTarget2D*>& UTG_AsyncRenderTaskBase::ActivateBlocking(JobBatchPtr Batch)
 {
-	UE_LOG(LogTextureGraph, Log, TEXT("UTG_AsyncRenderTaskBase::ActivateBlocking"));
+	UE_LOG(LogTextureGraph, Verbose, TEXT("UTG_AsyncRenderTaskBase::ActivateBlocking"));
 
 	Batch = PrepareActivate(Batch, false);
 
@@ -127,7 +127,7 @@ void UTG_AsyncRenderTaskBase::GatherAllRenderTargets()
 
 	bRenderComplete = true;
 
-	UE_LOG(LogTextureGraph, Log, TEXT("UTG_AsyncRenderTaskBase:: OnDone : bShouldDestroyOnRenderComplete %i"), bShouldDestroyOnRenderComplete ? 1 : 0);
+	UE_LOG(LogTextureGraph, Verbose, TEXT("UTG_AsyncRenderTaskBase:: OnDone : bShouldDestroyOnRenderComplete %i"), bShouldDestroyOnRenderComplete ? 1 : 0);
 	if (bShouldDestroyOnRenderComplete)
 	{
 		SetReadyToDestroy();
@@ -161,20 +161,21 @@ void UTG_AsyncRenderTaskBase::OnWorldCleanup(UWorld* World, bool bSessionEnded,
 		SetReadyToDestroy();
 	}
 
-	UE_LOG(LogTextureGraph, Log, TEXT("UTG_AsyncRenderTaskBase:: OnWorldCleanup"));
+	UE_LOG(LogTextureGraph, Verbose, TEXT("UTG_AsyncRenderTaskBase:: OnWorldCleanup"));
 }
 
 void UTG_AsyncRenderTaskBase::SetReadyToDestroy()
 {
-	UE_LOG(LogTextureGraph, Log, TEXT("UTG_AsyncRenderTaskBase:: SetReadyToDestroy"));
+	UE_LOG(LogTextureGraph, Verbose, TEXT("UTG_AsyncRenderTaskBase:: SetReadyToDestroy"));
 	TextureGraphPtr->FlushInvalidations();
+	FlushRenderingCommands();
 	ClearFlags(RF_Standalone);
 	UTG_AsyncTask::SetReadyToDestroy();
 }
 
 void UTG_AsyncRenderTaskBase::FinishDestroy()
 {
-	UE_LOG(LogTextureGraph, Log, TEXT("UTG_AsyncRenderTaskBase:: FinishDestroy"));
+	UE_LOG(LogTextureGraph, Verbose, TEXT("UTG_AsyncRenderTaskBase:: FinishDestroy"));
 	if (TextureGraphPtr != nullptr)
 	{
 		TextureGraphPtr->GetSettings()->FreeTargets();
@@ -214,7 +215,7 @@ void UTG_AsyncRenderTask::Activate()
 {
 	// Start the async task on a new thread
 	Super::Activate();
-	UE_LOG(LogTextureGraph, Log, TEXT("UTG_AsyncRenderTask::Activate"));
+	UE_LOG(LogTextureGraph, Verbose, TEXT("UTG_AsyncRenderTask::Activate"));
 
 	JobBatchPtr Batch = PrepareActivate(nullptr, true);
 
diff --git a/Engine/Plugins/TextureGraph/Source/TextureGraph/Private/TG_Node.cpp b/Engine/Plugins/TextureGraph/Source/TextureGraph/Private/TG_Node.cpp
--- a/Engine/Plugins/TextureGraph/Source/TextureGraph/Private/TG_Node.cpp
+++ b/Engine/Plugins/TextureGraph/Source/TextureGraph/Private/TG_Node.cpp
@@ -145,7 +145,7 @@ void UTG_Node::NotifyGraphOfNodeChange(bool bIsTweaking)
 
 void UTG_Node::OnExpressionChangedWithoutVar(const FPropertyChangedEvent& PropertyChangedEvent)
 {
-	UE_LOG(LogTemp, Log, TEXT("UTG_Node::PropertyChangedEvent ChangeType: %d"), PropertyChangedEvent.ChangeType);
+	UE_LOG(LogTextureGraph, Verbose, TEXT("UTG_Node::PropertyChangedEvent ChangeType: %d"), PropertyChangedEvent.ChangeType);
 
 	if (PropertyChangedEvent.Property)
 	{
diff --git a/Engine/Plugins/TextureGraph/Source/TextureGraph/Public/Blueprint/TG_AsyncRenderTask.h b/Engine/Plugins/TextureGraph/Source/TextureGraph/Public/Blueprint/TG_AsyncRenderTask.h
--- a/Engine/Plugins/TextureGraph/Source/TextureGraph/Public/Blueprint/TG_AsyncRenderTask.h
+++ b/Engine/Plugins/TextureGraph/Source/TextureGraph/Public/Blueprint/TG_AsyncRenderTask.h
@@ -41,7 +41,7 @@ protected:
 	TArray<BlobPtr> OutputBlobs;
 	UTextureGraphBase* OriginalTextureGraphPtr;
 	UTextureGraphBase* TextureGraphPtr;
-	bool bShouldDestroyOnRenderComplete = false;
+	bool bShouldDestroyOnRenderComplete = true;
 	bool bRenderComplete = false;
 };
 
diff --git a/Engine/Plugins/TextureGraph/Source/TextureGraphEngine/Job/JobBatch.cpp b/Engine/Plugins/TextureGraph/Source/TextureGraphEngine/Job/JobBatch.cpp
--- a/Engine/Plugins/TextureGraph/Source/TextureGraphEngine/Job/JobBatch.cpp
+++ b/Engine/Plugins/TextureGraph/Source/TextureGraphEngine/Job/JobBatch.cpp
@@ -325,14 +325,14 @@ void JobBatch::DebugDumpUnfinishedJobs()
 	if (NumJobsRunning <= 0)
 		return;
 
-	UE_LOG(LogBatch, Warning, TEXT("=== Offending Jobs Batch: %llu ==="), BatchId);
+	UE_LOG(LogBatch, Verbose, TEXT("=== Offending Jobs Batch: %llu ==="), BatchId);
 
 	for (size_t ji = 0; ji < JobsFinished.size(); ji++)
 	{
 		if (!JobsFinished[ji])
 		{
 			JobPtr JobObj = AllJobs[ji];
-			UE_LOG(LogBatch, Warning, TEXT("    - Job [IsDone: %s, IsCulled: %s, HasStarted: %s] : %s"), 
+			UE_LOG(LogBatch, Verbose, TEXT("    - Job [IsDone: %s, IsCulled: %s, HasStarted: %s] : %s"), 
 				(JobObj->IsDone() ? TEXT("Yes") : TEXT(" No")),
 				(JobObj->IsCulled() ? TEXT("Yes") : TEXT(" No")),
 				(JobObj->GetStats().BeginNativeTime > 0 ? TEXT("Yes") : TEXT(" No")), 
diff --git a/Engine/Plugins/TextureGraph/Source/TextureGraphEngine/Job/Scheduler.cpp b/Engine/Plugins/TextureGraph/Source/TextureGraphEngine/Job/Scheduler.cpp
--- a/Engine/Plugins/TextureGraph/Source/TextureGraphEngine/Job/Scheduler.cpp
+++ b/Engine/Plugins/TextureGraph/Source/TextureGraphEngine/Job/Scheduler.cpp
@@ -83,7 +83,7 @@ void Scheduler::Update(float dt)
 
 			if (delta > timeoutLimit)
 			{
-				UE_LOG(LogBatch, Warning, TEXT("Current Batch: %llu been running for %f ms [Max threshold: %f, Jobs: %u/%u]!"), 
+				UE_LOG(LogBatch, Verbose, TEXT("Current Batch: %llu been running for %f ms [Max threshold: %f, Jobs: %u/%u]!"), 
 					CurrentBatch->GetBatchId(), (float)delta, (float)CurrentBatchWarningLimit, CurrentBatch->GetNumJobsRunning(), 
 					(uint32)CurrentBatch->NumJobs());
 				CurrentBatch->DebugDumpUnfinishedJobs();
