Added a CVar to allow commandlets to use TextureGraphEngine

diff --git a/Engine/Plugins/TextureGraph/Source/TextureGraphEngine/2D/Tex.cpp b/Engine/Plugins/TextureGraph/Source/TextureGraphEngine/2D/Tex.cpp
--- a/Engine/Plugins/TextureGraph/Source/TextureGraphEngine/2D/Tex.cpp
+++ b/Engine/Plugins/TextureGraph/Source/TextureGraphEngine/2D/Tex.cpp
@@ -1014,8 +1014,8 @@ AsyncTiledBlobRef Tex::ToSingleBlob(CHashPtr Hash, bool TransferToRT /* = false
 		TransferToRT = true;	//Force transfer if desc has mips generation on and we are not RT.
 #endif 
 #if WITH_EDITOR
-	// No need to run during command-let execution
-	if (!GEditor || !FSlateApplication::IsInitialized() || !FApp::CanEverRender())
+	// Early out if commandlet execution is not allowed
+	if (TextureGraphEngine::BlockCommandlets())
 		return static_cast<AsyncTiledBlobRef>(cti::make_ready_continuable<TiledBlobRef>(TiledBlobRef()));
 
 	if ((TransferToRT && !IsValid(RT)) || !IsValid(Texture) || IsValidVirtualTexture())		// For cases where we want to transfer UTexture2D to UTextureRenderTarget2D
diff --git a/Engine/Plugins/TextureGraph/Source/TextureGraphEngine/Export/T_ExportAsUAsset.cpp b/Engine/Plugins/TextureGraph/Source/TextureGraphEngine/Export/T_ExportAsUAsset.cpp
--- a/Engine/Plugins/TextureGraph/Source/TextureGraphEngine/Export/T_ExportAsUAsset.cpp
+++ b/Engine/Plugins/TextureGraph/Source/TextureGraphEngine/Export/T_ExportAsUAsset.cpp
@@ -56,7 +56,7 @@ cti::continuable<int32> Job_ExportAsUAsset::PreExecAsync(ENamedThreads::Type exe
 cti::continuable<int32> Job_ExportAsUAsset::ExecAsync(ENamedThreads::Type execThread, ENamedThreads::Type returnThread)
 {
 #if WITH_EDITOR
-	if (!GEditor || !FSlateApplication::IsInitialized() || !FApp::CanEverRender())
+	if (TextureGraphEngine::BlockCommandlets())
 		return cti::make_ready_continuable(0);
 
 	UE_LOG(LogData, Log, TEXT("Exporting As UAsset: %s [Width: %d Height: %d]"), *(Setting.Name.ToString()), Setting.Width, Setting.Height);
diff --git a/Engine/Plugins/TextureGraph/Source/TextureGraphEngine/Export/TextureExporter.cpp b/Engine/Plugins/TextureGraph/Source/TextureGraphEngine/Export/TextureExporter.cpp
--- a/Engine/Plugins/TextureGraph/Source/TextureGraphEngine/Export/TextureExporter.cpp
+++ b/Engine/Plugins/TextureGraph/Source/TextureGraphEngine/Export/TextureExporter.cpp
@@ -282,7 +282,7 @@ AsyncBool TextureExporter::ExportAsUAsset(UMixInterface* MixObj, const FExportSe
 }
 FExportMapSettings TextureExporter::GetExportSettingsForTarget(FExportSettings& ExportSettings, TiledBlobPtr BlobObj, FName Name)
 {
-	if (!FSlateApplication::IsInitialized() || !FApp::CanEverRender())
+	if (TextureGraphEngine::BlockCommandlets())
 		return FExportMapSettings();
 	ExportSettings.PresetName = "ExportAsUAssetPreset";
 	
diff --git a/Engine/Plugins/TextureGraph/Source/TextureGraphEngine/Export/TextureExporter.h b/Engine/Plugins/TextureGraph/Source/TextureGraphEngine/Export/TextureExporter.h
--- a/Engine/Plugins/TextureGraph/Source/TextureGraphEngine/Export/TextureExporter.h
+++ b/Engine/Plugins/TextureGraph/Source/TextureGraphEngine/Export/TextureExporter.h
@@ -78,7 +78,7 @@ public:
 
 	void OnMapExportDone(FExportMapSettings MapSetting)
 	{
-		if (!FSlateApplication::IsInitialized() || !FApp::CanEverRender())
+		if (TextureGraphEngine::BlockCommandlets())
 			return;
 
 		MapsExported++;
diff --git a/Engine/Plugins/TextureGraph/Source/TextureGraphEngine/TextureGraphEngine.cpp b/Engine/Plugins/TextureGraph/Source/TextureGraphEngine/TextureGraphEngine.cpp
--- a/Engine/Plugins/TextureGraph/Source/TextureGraphEngine/TextureGraphEngine.cpp
+++ b/Engine/Plugins/TextureGraph/Source/TextureGraphEngine/TextureGraphEngine.cpp
@@ -25,6 +25,13 @@
 DECLARE_CYCLE_STAT(TEXT("TextureGraphEngine-Update"), STAT_TextureGraphEngine_Update, STATGROUP_TextureGraphEngine);
 DEFINE_LOG_CATEGORY(LogTextureGraphEngine);
 
+
+static TAutoConsoleVariable<bool> CVarAllowTextureGraphInCommandlets(
+	TEXT("TextureGraph.AllowCommandlets"),
+	false,
+	TEXT("Allow TextureGraph engine to be initialized when running via Commandlets")
+	);
+
 void FTextureGraphEngineModule::StartupModule()
 {
 	FDefaultGameModuleImpl::StartupModule();
@@ -129,8 +136,8 @@ void TextureGraphEngine::InitEngineInternal()
 
 	// Currently we're only supporting editor only 
 #if WITH_EDITOR
-	// No need to run during command-let execution
-	if (!GEditor || !FSlateApplication::IsInitialized() || !FApp::CanEverRender())
+	// Early out if commandlet execution is not allowed
+	if (TextureGraphEngine::BlockCommandlets())
 		return;
 
 	DeviceManagerObj = std::make_unique<DeviceManager>();
@@ -285,6 +292,15 @@ void TextureGraphEngine::SetRunEngine()
 	}
 }
 
+bool TextureGraphEngine::BlockCommandlets()
+{
+	return !CVarAllowTextureGraphInCommandlets.GetValueOnAnyThread()
+	#if WITH_EDITOR
+		&& !(GEditor && FSlateApplication::IsInitialized() && FApp::CanEverRender())
+	#endif
+	;
+}
+
 void TextureGraphEngine::FirstRunInit()
 {
 	check(IsInGameThread());
@@ -295,8 +311,8 @@ void TextureGraphEngine::FirstRunInit()
 	MaterialMgrObj = TStrongObjectPtr<UMaterialManager>(UMaterialManager::CreateNew<UMaterialManager>());
 	
 #if WITH_EDITOR
-	// No need to run during command-let execution
-	if (!GEditor || !FSlateApplication::IsInitialized() || !FApp::CanEverRender())
+	// Early out if commandlet execution is not allowed
+	if (TextureGraphEngine::BlockCommandlets())
 	{
 		return;
 	}
@@ -339,8 +355,8 @@ void TextureGraphEngine::Update(float dt)
 
 	// Currently we're only supporting editor only 
 #if WITH_EDITOR
-	// No need to run during command-let execution
-	if (!GEditor || !FSlateApplication::IsInitialized() || !FApp::CanEverRender())
+	// Early out if commandlet execution is not allowed
+	if (TextureGraphEngine::BlockCommandlets())
 		return;
 
 	// If engine is destroying then we don't do anything
diff --git a/Engine/Plugins/TextureGraph/Source/TextureGraphEngine/TextureGraphEngine.h b/Engine/Plugins/TextureGraph/Source/TextureGraphEngine/TextureGraphEngine.h
--- a/Engine/Plugins/TextureGraph/Source/TextureGraphEngine/TextureGraphEngine.h
+++ b/Engine/Plugins/TextureGraph/Source/TextureGraphEngine/TextureGraphEngine.h
@@ -4,6 +4,7 @@
 #include "CoreMinimal.h"
 #include "Profiling/StatGroup.h"
 #include "UObject/StrongObjectPtr.h"
+#include "HAL/IConsoleManager.h"
 
 #include <memory>
 #include <mutex>
@@ -70,7 +71,6 @@ public:
 typedef std::shared_ptr<EngineObserverSource>			EngineObserverSourcePtr;
 typedef std::shared_ptr<FTextureGraphErrorReporter>			ErrorReporterPtr;
 
-
 //////////////////////////////////////////////////////////////////////////
 /// The Engine is supposed to maintain a global state for the application
 /// and giving access to various shared data structures.
@@ -166,6 +166,7 @@ public:
 	FORCEINLINE static FTextureGraphErrorReporter*		GetErrorReporter(UMixInterface* MixKey) { return GInstance ? GInstance->ErrorReporters.Find(MixKey)->get() : nullptr; }
 
 	FORCEINLINE static bool								IsRunning() { return GInstance ? GInstance->bRunEngine.load() : false; }
+	static bool											BlockCommandlets();
 };
 
 
