Fixing a horde compile error due to declaration shadowing between EBaseNodeRefreshPriority::Normal and  EBlendModes::Normal. Wrapped EBlendModes in a namespace to minimize API changes and breaking BPs

diff --git a/Engine/Plugins/TextureGraph/Source/TextureGraph/Public/Expressions/Maths/TG_Expression_Blend.h b/Engine/Plugins/TextureGraph/Source/TextureGraph/Public/Expressions/Maths/TG_Expression_Blend.h
--- a/Engine/Plugins/TextureGraph/Source/TextureGraph/Public/Expressions/Maths/TG_Expression_Blend.h
+++ b/Engine/Plugins/TextureGraph/Source/TextureGraph/Public/Expressions/Maths/TG_Expression_Blend.h
@@ -22,7 +22,7 @@ public:
 
 	// Blend mode determine how inputs texture mix together.
 	UPROPERTY(EditAnywhere, BlueprintReadWrite, Category = NoCategory, meta = (TGType = "TG_Setting"))
-	TEnumAsByte<EBlendModes> BlendMode = EBlendModes::Normal;
+	TEnumAsByte<EBlendModes::Type> BlendMode = EBlendModes::Normal;
 
 	// The opacity to be used in conjunction with the mask. The final blend value = 'Mask[RedChannel] * Opacity'
 	UPROPERTY(EditAnywhere, Category = NoCategory, meta = (TGType = "TG_Setting", UIMin = "0", ClampMin = "0", UIMax = "1", ClampMax = "1"))
diff --git a/Engine/Plugins/TextureGraph/Source/TextureGraphEngine/2D/BlendModes.h b/Engine/Plugins/TextureGraph/Source/TextureGraphEngine/2D/BlendModes.h
--- a/Engine/Plugins/TextureGraph/Source/TextureGraphEngine/2D/BlendModes.h
+++ b/Engine/Plugins/TextureGraph/Source/TextureGraphEngine/2D/BlendModes.h
@@ -6,17 +6,20 @@
 #include "BlendModes.generated.h"
 
 UENUM(BlueprintType)
-enum EBlendModes : int
+namespace EBlendModes
 {
-	Normal						UMETA(DisplayName = "Copy"),
-	Add							UMETA(DisplayName = "Add"),
-	Subtract					UMETA(DisplayName = "Subtract"),
-	Multiply					UMETA(DisplayName = "Multiply"),
-	Divide						UMETA(DisplayName = "Divide"),
-	Difference					UMETA(DisplayName = "Difference"),
-	Max 						UMETA(DisplayName = "Max"),
-	Min 						UMETA(DisplayName = "Min"),
-	Step						UMETA(DisplayName = "Step"),
-	Overlay						UMETA(DisplayName = "Overlay"),
-	// Distort						UMETA(DisplayName = "Distort")
-};
+	enum Type : int
+	{
+		Normal						UMETA(DisplayName = "Copy"),
+		Add							UMETA(DisplayName = "Add"),
+		Subtract					UMETA(DisplayName = "Subtract"),
+		Multiply					UMETA(DisplayName = "Multiply"),
+		Divide						UMETA(DisplayName = "Divide"),
+		Difference					UMETA(DisplayName = "Difference"),
+		Max 						UMETA(DisplayName = "Max"),
+		Min 						UMETA(DisplayName = "Min"),
+		Step						UMETA(DisplayName = "Step"),
+		Overlay						UMETA(DisplayName = "Overlay"),
+		// Distort						UMETA(DisplayName = "Distort")
+	};
+}
diff --git a/Engine/Plugins/TextureGraph/Source/TextureGraphEngine/Transform/Expressions/T_Blend.cpp b/Engine/Plugins/TextureGraph/Source/TextureGraphEngine/Transform/Expressions/T_Blend.cpp
--- a/Engine/Plugins/TextureGraph/Source/TextureGraphEngine/Transform/Expressions/T_Blend.cpp
+++ b/Engine/Plugins/TextureGraph/Source/TextureGraphEngine/Transform/Expressions/T_Blend.cpp
@@ -28,7 +28,7 @@ T_Blend::~T_Blend()
 {
 }
 
-TiledBlobPtr T_Blend::Create(MixUpdateCyclePtr InCycle, BufferDescriptor DesiredDesc, int InTargetId, EBlendModes InBlendMode, const FBlendSettings* InBlendSettings)
+TiledBlobPtr T_Blend::Create(MixUpdateCyclePtr InCycle, BufferDescriptor DesiredDesc, int InTargetId, EBlendModes::Type InBlendMode, const FBlendSettings* InBlendSettings)
 {
 	switch(InBlendMode)
 	{
diff --git a/Engine/Plugins/TextureGraph/Source/TextureGraphEngine/Transform/Expressions/T_Blend.h b/Engine/Plugins/TextureGraph/Source/TextureGraphEngine/Transform/Expressions/T_Blend.h
--- a/Engine/Plugins/TextureGraph/Source/TextureGraphEngine/Transform/Expressions/T_Blend.h
+++ b/Engine/Plugins/TextureGraph/Source/TextureGraphEngine/Transform/Expressions/T_Blend.h
@@ -172,7 +172,7 @@ public:
 	/// Static functions
 	//////////////////////////////////////////////////////////////////////////
 
-	static UE_API TiledBlobPtr				Create(MixUpdateCyclePtr InCycle, BufferDescriptor DesiredDesc, int InTargetId, EBlendModes InBlendMode, const FBlendSettings* InBlendSettings);
+	static UE_API TiledBlobPtr				Create(MixUpdateCyclePtr InCycle, BufferDescriptor DesiredDesc, int InTargetId, EBlendModes::Type InBlendMode, const FBlendSettings* InBlendSettings);
 	static UE_API TiledBlobPtr				CreateNormal(MixUpdateCyclePtr InCycle, BufferDescriptor DesiredDesc, int InTargetId, const FBlendSettings* InBlendSettings);
 	static UE_API TiledBlobPtr				CreateAdd(MixUpdateCyclePtr InCycle, BufferDescriptor DesiredDesc, int InTargetId, const FBlendSettings* InBlendSettings);
 	static UE_API TiledBlobPtr				CreateSubtract(MixUpdateCyclePtr InCycle, BufferDescriptor DesiredDesc, int InTargetId, const FBlendSettings* InBlendSettings);
