[Core] Change from std::function to TFunction to prevent crash on exit with compact set

#tests Perform a regular shutdown of the engine

diff --git a/Engine/Plugins/TextureGraph/Source/TextureGraph/Private/TG_GraphEvaluation.cpp b/Engine/Plugins/TextureGraph/Source/TextureGraph/Private/TG_GraphEvaluation.cpp
--- a/Engine/Plugins/TextureGraph/Source/TextureGraph/Private/TG_GraphEvaluation.cpp
+++ b/Engine/Plugins/TextureGraph/Source/TextureGraph/Private/TG_GraphEvaluation.cpp
@@ -547,7 +547,7 @@ void FloatToEnum_Converter(FTG_Evaluation::VarConverterInfo& Info)
 }
 
 //////////////////////////////////////////////////////////////////////////////////////////////////////////////////////////////////////
-#define VAR_CONVERTER(nameFrom, nameTo, function)	{ MakeConvertKey(TEXT(#nameFrom), TEXT(#nameTo)), std::bind(&function, std::placeholders::_1) }
+#define VAR_CONVERTER(nameFrom, nameTo, function)	{ MakeConvertKey(TEXT(#nameFrom), TEXT(#nameTo)), &function }
 #define VAR_CONVERTER_DEF(nameFrom, nameTo, function)		VAR_CONVERTER(#nameFrom, #nameTo, #function)
 #define VAR_CONVERTER_NULL(nameFrom, nameTo)	{ MakeConvertKey(TEXT(#nameFrom), TEXT(#nameTo)), nullptr }
 
diff --git a/Engine/Plugins/TextureGraph/Source/TextureGraph/Private/TG_Var.cpp b/Engine/Plugins/TextureGraph/Source/TextureGraph/Private/TG_Var.cpp
--- a/Engine/Plugins/TextureGraph/Source/TextureGraph/Private/TG_Var.cpp
+++ b/Engine/Plugins/TextureGraph/Source/TextureGraph/Private/TG_Var.cpp
@@ -643,7 +643,7 @@ void FTG_Var::CopyFrom(FTG_Var* InVar)
 }
 
 //////////////////////////////////////////////////////////////////////////////////////////////////////////////////////////////////////
-#define VAR_ARCHIVE_SERIALIZER(name, function)	{ TEXT(name), std::bind(&function, std::placeholders::_1) }
+#define VAR_ARCHIVE_SERIALIZER(name, function)	{ TEXT(name), &function }
 #define VAR_ARCHIVE_SERIALIZER_DEF(name)		VAR_ARCHIVE_SERIALIZER(#name, VarArchiveSerializer_##name)
 
 template <typename T_ValueType>
diff --git a/Engine/Plugins/TextureGraph/Source/TextureGraph/Public/TG_GraphEvaluation.h b/Engine/Plugins/TextureGraph/Source/TextureGraph/Public/TG_GraphEvaluation.h
--- a/Engine/Plugins/TextureGraph/Source/TextureGraph/Public/TG_GraphEvaluation.h
+++ b/Engine/Plugins/TextureGraph/Source/TextureGraph/Public/TG_GraphEvaluation.h
@@ -177,7 +177,7 @@ struct FTG_Evaluation
 		int Index = -1;
 		FTG_EvaluationContext* Context = nullptr;
 	};
-	typedef std::function<void(VarConverterInfo& Info)> VarConverter;
+	typedef TFunction<void(VarConverterInfo& Info)> VarConverter;
 
 	typedef TMap				<FName, VarConverter> ConverterMap;
 	static UE_API ConverterMap			DefaultConverters;
@@ -192,7 +192,7 @@ struct FTG_Evaluation
 		int Index = -1;
 		FTG_EvaluationContext* Context = nullptr;
     };
-    typedef std::function<bool(VarConformerInfo& Info)> VarConformer;
+    typedef TFunction<bool(VarConformerInfo& Info)> VarConformer;
 	
 	static UE_API bool AreArgumentsCompatible(const FTG_Argument& ArgFrom, const FTG_Argument& ArgTo, FName& ConverterKey);
 
diff --git a/Engine/Plugins/TextureGraph/Source/TextureGraph/Public/TG_Var.h b/Engine/Plugins/TextureGraph/Source/TextureGraph/Public/TG_Var.h
--- a/Engine/Plugins/TextureGraph/Source/TextureGraph/Public/TG_Var.h
+++ b/Engine/Plugins/TextureGraph/Source/TextureGraph/Public/TG_Var.h
@@ -61,7 +61,7 @@ struct FTG_Var
 		FTG_Var* Var = nullptr;
 		FArchive& Ar;
 	};
-	typedef std::function<void(VarArchiveSerialInfo&)>	VarArchiveSerializer;
+	typedef TFunctionRef<void(VarArchiveSerialInfo&)>	VarArchiveSerializer;
 	typedef TMap<FName, VarArchiveSerializer>			VarArchiveSerializerMap;
 	static UE_API VarArchiveSerializerMap						DefaultArchiveSerializers;
 
