[TextureGraph] Fix crash when loading saved TextureGraphs after change of TG_Var function pointers from std::function to TFunctionRef

diff --git a/Engine/Plugins/TextureGraph/Source/TextureGraph/Public/TG_Var.h b/Engine/Plugins/TextureGraph/Source/TextureGraph/Public/TG_Var.h
--- a/Engine/Plugins/TextureGraph/Source/TextureGraph/Public/TG_Var.h
+++ b/Engine/Plugins/TextureGraph/Source/TextureGraph/Public/TG_Var.h
@@ -61,15 +61,15 @@ struct FTG_Var
 		FTG_Var* Var = nullptr;
 		FArchive& Ar;
 	};
-	typedef TFunctionRef<void(VarArchiveSerialInfo&)>	VarArchiveSerializer;
+	typedef TFunction<void(VarArchiveSerialInfo&)>		VarArchiveSerializer;
 	typedef TMap<FName, VarArchiveSerializer>			VarArchiveSerializerMap;
-	static UE_API VarArchiveSerializerMap						DefaultArchiveSerializers;
+	static UE_API VarArchiveSerializerMap				DefaultArchiveSerializers;
+
 
 	// Serialize the var value
 	// Called from the Pin serialize()
 	UE_API void Serialize(FArchive& Ar, FTG_Id PinId, const FTG_Argument& Argument);
-
-
+	
 	UPROPERTY(Transient)
 	FTG_Id						PinId;
 
