[TG] Fixing VT forced loading  for Material expression

diff --git a/Engine/Plugins/TextureGraph/Source/TextureGraphEngine/FxMat/FxMaterial_DrawMaterial.h b/Engine/Plugins/TextureGraph/Source/TextureGraphEngine/FxMat/FxMaterial_DrawMaterial.h
--- a/Engine/Plugins/TextureGraph/Source/TextureGraphEngine/FxMat/FxMaterial_DrawMaterial.h
+++ b/Engine/Plugins/TextureGraph/Source/TextureGraphEngine/FxMat/FxMaterial_DrawMaterial.h
@@ -9,6 +9,7 @@
 #include <TextureResource.h>
 #include <SceneManagement.h>
 #include <PostProcess/DrawRectangle.h>
+#include <EngineModule.h>
 
 #include <DataDrivenShaderPlatformInfo.h>
 #include <MaterialShader.h>
@@ -178,6 +179,8 @@ public:
 		const FMaterialRenderProxy* MaterialProxy = Material->GetRenderProxy();
 		MaterialProxy->UpdateUniformExpressionCacheIfNeeded(SceneView->GetFeatureLevel());
 		FMaterialRenderProxy::UpdateDeferredCachedUniformExpressions();
+		GetRendererModule().RequestVirtualTextureTiles(MaterialProxy, Target->GetSizeXY(), SceneView->GetFeatureLevel());
+		GetRendererModule().LoadPendingVirtualTextureTiles(RHI, SceneView->GetFeatureLevel());
 
 		const FMaterial& RenderMaterial = MaterialProxy->GetMaterialWithFallback(SceneView->GetFeatureLevel(), MaterialProxy);
 		const FMaterialShaderMap* const MaterialShaderMap = RenderMaterial.GetRenderingThreadShaderMap();
