TG: [Linux] Fix crash on loading a texture in a TextureExpression

diff --git a/Engine/Plugins/TextureGraph/Source/TextureGraphEngine/Device/FX/Device_FX.cpp b/Engine/Plugins/TextureGraph/Source/TextureGraphEngine/Device/FX/Device_FX.cpp
--- a/Engine/Plugins/TextureGraph/Source/TextureGraphEngine/Device/FX/Device_FX.cpp
+++ b/Engine/Plugins/TextureGraph/Source/TextureGraphEngine/Device/FX/Device_FX.cpp
@@ -689,6 +689,7 @@ int32 Device_FX::InitRTResource(TexPtr TextureObj, UTextureRenderTarget* RT)
 	FRHICommandListImmediate& RHICmdList = FRHICommandListImmediate::Get();
 
 	RT->GetResource()->InitResource(RHICmdList);
+	
 	FDeferredUpdateResource::UpdateResources(RHICmdList);
 
 	FTextureRenderTarget2DResource* RTRes = (FTextureRenderTarget2DResource*)RT->GetRenderTargetResource();
@@ -698,10 +699,10 @@ int32 Device_FX::InitRTResource(TexPtr TextureObj, UTextureRenderTarget* RT)
 		TextureObj->GetDescriptor().Format_HashValue(), RT, RTRes->GetSizeX(), RTRes->GetSizeY());
 
 	FTextureRHIRef RHITexture = RTRes->GetTextureRHI();
-	check(RHITexture);
-
-	RHICmdList.BindDebugLabelName(RHITexture, *TextureObj->GetDescriptor().Name);
-
+	if(RHITexture)
+	{
+		RHICmdList.BindDebugLabelName(RHITexture, *TextureObj->GetDescriptor().Name);
+	}
 	return 0;
 }
 
