Reconciled offline work

#FailedValidations Preflight

diff --git a/Engine/Plugins/TextureGraph/Source/TextureGraph/Private/TG_GraphEvaluation.cpp b/Engine/Plugins/TextureGraph/Source/TextureGraph/Private/TG_GraphEvaluation.cpp
--- a/Engine/Plugins/TextureGraph/Source/TextureGraph/Private/TG_GraphEvaluation.cpp
+++ b/Engine/Plugins/TextureGraph/Source/TextureGraph/Private/TG_GraphEvaluation.cpp
@@ -52,20 +52,28 @@ void FTG_Evaluation::TransferVarToPin(UTG_Pin* InPin, FTG_EvaluationContext* Con
 	if (Arg.IsInput())
 	{
 		bool bNeedsConversion = InPin->ConnectionNeedsConversion();
-		FName ArrayConverterKey = InPin->GetInputVarConverterKey();
+		FName VarConverterKey = InPin->GetInputVarConverterKey();
 		FTG_Var* OutVar = InPin->EditConvertedVar(); // Generate the converted version of the connected input into the CoonvertedVar of the Pin
 		bool bIsEnum = InPin->GetArgument().ArgumentType.IsEnum();
 
-		if (!bNeedsConversion && (Var->IsArray() && !InPin->EditSelfVar()->IsArray()))
+		// Filter the case if a var is an array
+		if (Var->IsArray() && !InPin->EditSelfVar()->IsArray())
 		{
-			ArrayConverterKey = MakeConvertKey("FTG_VariantArray", Arg.CPPTypeName);
-			OutVar = InPin->EditSelfVar();
-			bNeedsConversion = true;
+			// Introduce a conversion specifically from VarianArray to the target type
+			// if there was ALREADY a conversion it is likely compatible anyway
+			VarConverterKey = MakeConvertKey("FTG_VariantArray", Arg.CPPTypeName);
+			// If no conversion was installed
+			if (!bNeedsConversion)
+			{
+				// TODO: SG : I need to double check that this is needed, it doesn't seem necessary
+				OutVar = InPin->EditSelfVar();
+				bNeedsConversion = true;
+			}
 		}
 
 		if (bNeedsConversion)
 		{
-			FTG_Evaluation::VarConverter* Converter = FTG_Evaluation::DefaultConverters.Find(ArrayConverterKey);
+			FTG_Evaluation::VarConverter* Converter = FTG_Evaluation::DefaultConverters.Find(VarConverterKey);
 			// If converter is null then the var is compatible, no need to convert.
 			// Else COnverter is valid and we need to run it 
 			if ((*Converter))
diff --git a/Engine/Plugins/TextureGraph/Source/TextureGraphEditor/Private/TG_EditorCommands.cpp b/Engine/Plugins/TextureGraph/Source/TextureGraphEditor/Private/TG_EditorCommands.cpp
--- a/Engine/Plugins/TextureGraph/Source/TextureGraphEditor/Private/TG_EditorCommands.cpp
+++ b/Engine/Plugins/TextureGraph/Source/TextureGraphEditor/Private/TG_EditorCommands.cpp
@@ -22,7 +22,7 @@ void FTG_EditorCommands::RegisterCommands()
 {
 	UI_COMMAND(RunGraph, "Update", "Update the current graph.", EUserInterfaceActionType::Button, FInputChord());
 	UI_COMMAND(AutoUpdateGraph, "Auto Update", "Toggles to update graph automatically or not.", EUserInterfaceActionType::ToggleButton, FInputChord());
-	UI_COMMAND(LogGraph, "Log", "Log the current graph.", EUserInterfaceActionType::Button, FInputChord());
+	UI_COMMAND(LogGraph, "Log", "Log the current graph.", EUserInterfaceActionType::Button, FInputChord(EKeys::L, EModifierKey::Alt));
 	UI_COMMAND(ExportAsUAsset, "Export", "Export all the output nodes as Texture UAssets", EUserInterfaceActionType::Button, FInputChord());
 	UI_COMMAND(ConvertInputParameterToConstant, "Convert input to Constant", "Convert input parameter to or from a constant.", EUserInterfaceActionType::Button, FInputChord());
 	UI_COMMAND(ConvertInputParameterFromConstant, "Convert input to Parameter", "Convert input parameter to or from a constant.", EUserInterfaceActionType::Button, FInputChord());
