[TG}] FIx a bug with Shpae node when using Rectangle, the ehight field was disappearing/ disconnecting when reloading

#FailedValidations Preflight

diff --git a/Engine/Plugins/TextureGraph/Source/TextureGraph/Private/Expressions/Procedural/TG_Expression_Shape.cpp b/Engine/Plugins/TextureGraph/Source/TextureGraph/Private/Expressions/Procedural/TG_Expression_Shape.cpp
--- a/Engine/Plugins/TextureGraph/Source/TextureGraph/Private/Expressions/Procedural/TG_Expression_Shape.cpp
+++ b/Engine/Plugins/TextureGraph/Source/TextureGraph/Private/Expressions/Procedural/TG_Expression_Shape.cpp
@@ -43,7 +43,7 @@ bool UTG_Expression_Shape::CanEditChange(const FProperty* InProperty) const
 	// the CanEditChange function also checks down the PropertyChain,
 	// checking direct ownership fixes an issue here where deeper in the hierarchy,
 	// FTG_TextureDescriptor::Height was also being affected because of the logic for TG_Expression_Shape::Height below 
-	if (!bEditCondition || this->GetClass() != InProperty->GetOwnerClass())
+	if (!bEditCondition && (this->GetClass() != InProperty->GetOwnerClass()))
 	{
 		return bEditCondition;
 	}
