TextureGraph material rendering now handles virtual texture feedback.
The material shader writes to the feedback buffer.
The tile render pass does a fixed number of warm up renders during which it processes the virtual texture feedback.
This allows the virtual texture cache to be populated before the final tile render.
There is a small cost for the multiple warmup passes and needing to flush the RHI to readback the feedback between each warmup pass.
Added "TG.VirtualTexture.WarmupFrames" to control the warmup length. Defaults to 8.
Also fixed a RHIValidation error with the texture graph histogram pass.

diff --git a/Engine/Plugins/TextureGraph/Shaders/TextureGraphMaterialShader.usf b/Engine/Plugins/TextureGraph/Shaders/TextureGraphMaterialShader.usf
--- a/Engine/Plugins/TextureGraph/Shaders/TextureGraphMaterialShader.usf
+++ b/Engine/Plugins/TextureGraph/Shaders/TextureGraphMaterialShader.usf
@@ -42,7 +42,10 @@ void TextureGraphMaterialShaderPS(
 
 	ResolvedView = ResolveView();
 	FMaterialPixelParameters MaterialParameters = MakeInitializedMaterialPixelParameters();
-	FPixelMaterialInputs PixelMaterialInputs;
+
+#if MATERIAL_VIRTUALTEXTURE_FEEDBACK
+	InitializeVirtualTextureFeedback(MaterialParameters.VirtualTextureFeedback);
+#endif
 
 	float2 ViewportUV = InUV;
 
@@ -57,6 +60,7 @@ void TextureGraphMaterialShaderPS(
 	SvPosition.z = max(SvPosition.z, 1e-18);
 
 	// fill out other related material parameters
+	FPixelMaterialInputs PixelMaterialInputs;
 	CalcMaterialParametersPost(MaterialParameters, PixelMaterialInputs, SvPosition, true);
 	
 	// Extract the value we want to represent
@@ -97,6 +101,14 @@ void TextureGraphMaterialShaderPS(
 		MaterialParameters.WorldPosition_CamRelative,
 		0.0f // Curvature
 	);
+
+#if MATERIAL_VIRTUALTEXTURE_FEEDBACK
+	FinalizeVirtualTextureFeedback(
+		MaterialParameters.VirtualTextureFeedback,
+		MaterialParameters.SvPosition,
+		View.VTFeedbackBuffer
+	);
+#endif
 	
 	MatBaseColor = Export.BaseColor;
 	MatMetallic = Export.Metallic;
@@ -133,5 +145,3 @@ void TextureGraphMaterialShaderPS(
 
 	OutColor = float4(Color, 1.0);
 }
-
-
diff --git a/Engine/Plugins/TextureGraph/Source/TextureGraphEngine/FxMat/FxMaterial_DrawMaterial.cpp b/Engine/Plugins/TextureGraph/Source/TextureGraphEngine/FxMat/FxMaterial_DrawMaterial.cpp
--- a/Engine/Plugins/TextureGraph/Source/TextureGraphEngine/FxMat/FxMaterial_DrawMaterial.cpp
+++ b/Engine/Plugins/TextureGraph/Source/TextureGraphEngine/FxMat/FxMaterial_DrawMaterial.cpp
@@ -1,12 +1,28 @@
 // Copyright Epic Games, Inc. All Rights Reserved
 #include "FxMaterial_DrawMaterial.h"
+
+#include "EngineModule.h"
 #include "SceneView.h"
+#include "VT/VirtualTextureFeedbackResource.h"
 
 IMPLEMENT_MATERIAL_SHADER_TYPE(, FTextureGraphMaterialShaderVS, TEXT("/Plugin/TextureGraph/TextureGraphMaterialShader.usf"), TEXT("TextureGraphMaterialShaderVS"), SF_Vertex);
 IMPLEMENT_MATERIAL_SHADER_TYPE(, FTextureGraphMaterialShaderPS, TEXT("/Plugin/TextureGraph/TextureGraphMaterialShader.usf"), TEXT("TextureGraphMaterialShaderPS"), SF_Pixel);
 
 FName FTextureGraphMaterialShaderPS::PSCONTROL_ARG = TEXT("PSControl");
 
+static int32 GVirtualTextureWarmupFrames = 8;
+static FAutoConsoleVariableRef CVarVirtualTextureWarmupFrames(
+	TEXT("TG.VirtualTexture.WarmupFrames"),
+	GVirtualTextureWarmupFrames,
+	TEXT("Number of warmup render frames when sampling from virtual textures."),
+	ECVF_RenderThreadSafe);
+
+static int32 GVirtualTextureFeedbackFactor = 16;
+static FAutoConsoleVariableRef CVarVirtualTextureFeedbackFactor(
+	TEXT("TG.VirtualTexture.FeedbackFactor"),
+	GVirtualTextureFeedbackFactor,
+	TEXT("Texel size of tile used to generate one item of virtual texture feedback."),
+	ECVF_RenderThreadSafe);
 
 FSceneView* FxMaterial_DrawMaterialBase::CreateSceneView(UTextureRenderTarget2D* RenderTarget, const FIntPoint& TargetSizeXY)
 {
@@ -74,6 +90,13 @@ FSceneView* FxMaterial_DrawMaterialBase::CreateSceneView(UTextureRenderTarget2D*
 
 		ViewUniformShaderParameters.MobilePreviewMode = 0.0f;
 
+		// Setup virtual texture feedback.
+		// Use a local static frame counter to vary feedback location independent of global render frame count.
+		static uint32 FrameIndex = 0;
+		VirtualTexture::FFeedbackShaderParams VirtualTextureFeedbackShaderParams;
+		VirtualTexture::GetFeedbackShaderParams(FrameIndex++, GVirtualTextureFeedbackFactor, VirtualTextureFeedbackShaderParams);
+		VirtualTexture::UpdateViewUniformShaderParameters(VirtualTextureFeedbackShaderParams, ViewUniformShaderParameters);
+
 		{
 			QUICK_SCOPE_CYCLE_COUNTER(STAT_Slate_CreateViewUniformBufferImmediate);
 			View->ViewUniformBuffer = TUniformBufferRef<FViewUniformShaderParameters>::CreateUniformBufferImmediate(ViewUniformShaderParameters, UniformBuffer_SingleFrame);
@@ -83,9 +106,37 @@ FSceneView* FxMaterial_DrawMaterialBase::CreateSceneView(UTextureRenderTarget2D*
 	return View;
 }
 
+bool FxMaterial_DrawMaterialBase::HasVirtualTextureFeedback(FMaterial const& RenderMaterial, EShaderPlatform ShaderPlatform) const
+{
+	// todo: Would be better if there was a bGeneratesVirtualTextureFeedback flag on the cached material data, since it
+	// is possible to have virtual texture samples that don't generate feedback.
+	return UseVirtualTexturing(ShaderPlatform) && !RenderMaterial.GetUniformVirtualTextureExpressions().IsEmpty();
+}
+
+int32 FxMaterial_DrawMaterialBase::GetVirtualTextureWarmupFrames() const
+{
+	return FMath::Max(GVirtualTextureWarmupFrames, 1);
+}
+
+void FxMaterial_DrawMaterialBase::VirtualTextureFeedbackBegin(FRDGBuilder& GraphBuilder, FIntPoint TargetSize, ERHIFeatureLevel::Type FeatureLevel) const
+{
+	VirtualTexture::BeginFeedback(GraphBuilder, TargetSize, GVirtualTextureFeedbackFactor, /*bInExtendFeedbackForDebug*/false, FeatureLevel);
+}
+
+void FxMaterial_DrawMaterialBase::VirtualTextureFeedbackEnd(FRDGBuilder& GraphBuilder) const
+{
+	VirtualTexture::EndFeedback(GraphBuilder);
+}
 
- bool FxMaterial_DrawMaterialBase::ValidateMaterialShaderMap(UMaterial * InMaterial, FShaderType* InPixelShaderType)
- {
+void FxMaterial_DrawMaterialBase::VirtualTextureFeedbackSync(FRHICommandListImmediate& RHICmdList, ERHIFeatureLevel::Type FeatureLevel) const
+{
+	// Need a flush before we can read back the virtual texture feedback buffer.
+	RHICmdList.ImmediateFlush(EImmediateFlushType::FlushRHIThread);
+	GetRendererModule().SyncVirtualTextureUpdates(RHICmdList, FeatureLevel);
+}
+
+bool FxMaterial_DrawMaterialBase::ValidateMaterialShaderMap(UMaterial * InMaterial, FShaderType* InPixelShaderType)
+{
 	// Pixel shader combined with Material
 	const FMaterialRenderProxy* MaterialProxy = InMaterial->GetRenderProxy();
 
diff --git a/Engine/Plugins/TextureGraph/Source/TextureGraphEngine/FxMat/FxMaterial_DrawMaterial.h b/Engine/Plugins/TextureGraph/Source/TextureGraphEngine/FxMat/FxMaterial_DrawMaterial.h
--- a/Engine/Plugins/TextureGraph/Source/TextureGraphEngine/FxMat/FxMaterial_DrawMaterial.h
+++ b/Engine/Plugins/TextureGraph/Source/TextureGraphEngine/FxMat/FxMaterial_DrawMaterial.h
@@ -106,6 +106,11 @@ public:
 
 	FSceneView* CreateSceneView(UTextureRenderTarget2D* RenderTarget, const FIntPoint& TargetSizeXY);
 
+	bool HasVirtualTextureFeedback(FMaterial const& RenderMaterial, EShaderPlatform ShaderPlatform) const;
+	int32 GetVirtualTextureWarmupFrames() const;
+	void VirtualTextureFeedbackBegin(FRDGBuilder& GraphBuilder, FIntPoint TargetSize, ERHIFeatureLevel::Type FeatureLevel) const;
+	void VirtualTextureFeedbackEnd(FRDGBuilder& GraphBuilder) const;
+	void VirtualTextureFeedbackSync(FRHICommandListImmediate& RHICmdList, ERHIFeatureLevel::Type FeatureLevel) const;
 
 	static bool ValidateMaterialShaderMap(UMaterial* InMaterial, FShaderType* InPixelShaderType);
 };
@@ -115,6 +120,11 @@ class FxMaterial_DrawMaterial : public FxMaterial_DrawMaterialBase
 {
 public:
 
+	BEGIN_SHADER_PARAMETER_STRUCT(FParameters, )
+		SHADER_PARAMETER_STRUCT_INCLUDE(VSH_Type::FParameters, VS)
+		SHADER_PARAMETER_STRUCT_INCLUDE(FSH_Type::FParameters, PS)
+		RENDER_TARGET_BINDING_SLOTS()
+	END_SHADER_PARAMETER_STRUCT()
 
 
 	// type name for the permutation domain
@@ -165,57 +175,75 @@ public:
 
 	void MyBlit(FRHICommandListImmediate& RHI, UTextureRenderTarget2D* RenderTarget, FRHITexture* Target, const RenderMesh* MeshObj, int32 TargetId, FGraphicsPipelineStateInitializer* InPSO = nullptr)
 	{
-		BindTexturesForBlitting();
-
-		// Make a Scene view for this pass
-		FSceneView* SceneView = CreateSceneView(RenderTarget, Target->GetSizeXY());
-
-		//const typename FSH_Type::FParameters* params = reinterpret_cast<typename FSH_Type::FParameters*>(params_);
-		QUICK_SCOPE_CYCLE_COUNTER(STAT_ShaderPlugin_PixelShader); // Used to gather CPU profiling data for the UE4 session frontend
-		SCOPED_DRAW_EVENT(RHI, ShaderPlugin_Pixel); // Used to profile GPU activity and add metadata to be consumed by for example RenderDoc
+		const ERHIFeatureLevel::Type FeatureLevel = GMaxRHIFeatureLevel;
 
+		BindTexturesForBlitting();
 
 		// Pixel shader combined with Material
 		const FMaterialRenderProxy* MaterialProxy = Material->GetRenderProxy();
-		MaterialProxy->UpdateUniformExpressionCacheIfNeeded(SceneView->GetFeatureLevel());
+		MaterialProxy->UpdateUniformExpressionCacheIfNeeded(FeatureLevel);
 		FMaterialRenderProxy::UpdateDeferredCachedUniformExpressions();
-		GetRendererModule().RequestVirtualTextureTiles(MaterialProxy, Target->GetSizeXY(), SceneView->GetFeatureLevel());
-		GetRendererModule().LoadPendingVirtualTextureTiles(RHI, SceneView->GetFeatureLevel());
 
-		const FMaterial& RenderMaterial = MaterialProxy->GetMaterialWithFallback(SceneView->GetFeatureLevel(), MaterialProxy);
+		const FMaterial& RenderMaterial = MaterialProxy->GetMaterialWithFallback(FeatureLevel, MaterialProxy);
 		const FMaterialShaderMap* const MaterialShaderMap = RenderMaterial.GetRenderingThreadShaderMap();
 
 		TShaderRef<VSH_Type> VertexShader = MaterialShaderMap->GetShader<VSH_Type>();
-		FRHIVertexShader* RHIVertexShader = VertexShader.GetVertexShader();
 		TShaderRef<FSH_Type> PixelShader = MaterialShaderMap->GetShader<FSH_Type>();
-		FRHIPixelShader* RHIPixelShader = PixelShader.GetPixelShader();
-
-		// Target render target needs to transition to RTV for rendering
-		RHI.Transition(FRHITransitionInfo(Target, ERHIAccess::Unknown, ERHIAccess::RTV));
 
-		FRHIRenderPassInfo passInfo(Target, ERenderTargetActions::Clear_Store);
-		RHI.BeginRenderPass(passInfo, TEXT("FxMaterial_UMaterial"));
+		const bool bUseVirtualTexturing = HasVirtualTextureFeedback(RenderMaterial, GMaxRHIShaderPlatform);
+		const int32 RenderCount = bUseVirtualTexturing ? GetVirtualTextureWarmupFrames() : 1;
 
-	
-		// Set the graphic pipeline state.
-		FGraphicsPipelineStateInitializer PSO;
-		
-		InitPSO_Default(PSO, RHIVertexShader, RHIPixelShader);
-
-		RHI.ApplyCachedRenderTargets(PSO);
-
-		SetGraphicsPipelineState(RHI, PSO, 0);
-
-		SetShaderParametersMixedVS(RHI, VertexShader, VSHParams, *SceneView);
-		SetShaderParametersMixedPS(RHI, PixelShader, FSHParams, *SceneView, MaterialProxy, RenderMaterial);
-
-		RHI.SetStreamSource(0, GQuadBuffer.VertexBufferRHI, 0);
-		RHI.DrawPrimitive(0, 2, 1);
-
-		RHI.EndRenderPass();
+		QUICK_SCOPE_CYCLE_COUNTER(STAT_ShaderPlugin_PixelShader); // Used to gather CPU profiling data for the session frontend
+		SCOPED_DRAW_EVENT(RHI, ShaderPlugin_Pixel); // Used to profile GPU activity and add metadata to be consumed by for example RenderDoc
 
-		// Target render target has been rendered, transition to the default SRV state for read
-		RHI.Transition(FRHITransitionInfo(Target, ERHIAccess::RTV, ERHIAccess::SRVMask));
+		for (int32 RenderIndex = 0; RenderIndex < RenderCount; ++RenderIndex)
+		{
+			FRDGBuilder GraphBuilder(RHI);
+			FRDGTextureRef TargetRDG = GraphBuilder.RegisterExternalTexture(CreateRenderTarget(Target, TEXT("FxMaterial_Target")));
+
+			// Call VirtualTextureFeedbackBegin() before CreateSceneView() so that we have feedback buffer ready for creation of scene uniform buffer.
+			if (bUseVirtualTexturing)
+			{
+				VirtualTextureFeedbackBegin(GraphBuilder, Target->GetSizeXY(), FeatureLevel);
+			}
+
+			FSceneView* SceneView = CreateSceneView(RenderTarget, Target->GetSizeXY());
+
+			auto* PassParameters = GraphBuilder.AllocParameters<FParameters>();
+			PassParameters->VS = VSHParams;
+			PassParameters->PS = FSHParams;
+			PassParameters->RenderTargets[0] = FRenderTargetBinding(TargetRDG, ERenderTargetLoadAction::EClear);
+
+			GraphBuilder.AddPass(
+				RDG_EVENT_NAME("FxMaterial (TargetId: %d, RenderIndex: %d)", TargetId, RenderIndex),
+				PassParameters,
+				ERDGPassFlags::Raster,
+				[&](FRHICommandList& RHICmdList)
+			{
+				FGraphicsPipelineStateInitializer PSO;
+				InitPSO_Default(PSO, VertexShader.GetVertexShader(), PixelShader.GetPixelShader());
+				RHICmdList.ApplyCachedRenderTargets(PSO);
+				SetGraphicsPipelineState(RHICmdList, PSO, 0);
+
+				SetShaderParametersMixedVS(RHICmdList, VertexShader, PassParameters->VS, *SceneView);
+				SetShaderParametersMixedPS(RHICmdList, PixelShader, PassParameters->PS, *SceneView, MaterialProxy, RenderMaterial);
+
+				RHICmdList.SetStreamSource(0, GQuadBuffer.VertexBufferRHI, 0);
+				RHICmdList.DrawPrimitive(0, 2, 1);
+			});
+
+			if (bUseVirtualTexturing)
+			{
+				VirtualTextureFeedbackEnd(GraphBuilder);
+			}
+
+			GraphBuilder.Execute();
+
+			if (bUseVirtualTexturing)
+			{
+				VirtualTextureFeedbackSync(RHI, FeatureLevel);
+			}
+		}
 	}
 
 	static bool ValidateMaterial(UMaterialInterface* InMaterial)
diff --git a/Engine/Plugins/TextureGraph/Source/TextureGraphEngine/Transform/Utility/T_TextureHistogram.cpp b/Engine/Plugins/TextureGraph/Source/TextureGraphEngine/Transform/Utility/T_TextureHistogram.cpp
--- a/Engine/Plugins/TextureGraph/Source/TextureGraphEngine/Transform/Utility/T_TextureHistogram.cpp
+++ b/Engine/Plugins/TextureGraph/Source/TextureGraphEngine/Transform/Utility/T_TextureHistogram.cpp
@@ -171,12 +171,12 @@ public:
 					)
 				);
 
-				RHI.Transition(FRHITransitionInfo(TilesHistogramBuffer.UAV, ERHIAccess::UAVCompute, ERHIAccess::UAVCompute));
+				RHI.Transition(FRHITransitionInfo(TilesHistogramBuffer.UAV, ERHIAccess::UAVMask, ERHIAccess::UAVMask));
 			}
 
 			// All the per tiles histograms have been collected
 			// transition to srv read
-			RHI.Transition(FRHITransitionInfo(TilesHistogramBuffer.UAV, ERHIAccess::UAVCompute, ERHIAccess::SRVCompute));
+			RHI.Transition(FRHITransitionInfo(TilesHistogramBuffer.UAV, ERHIAccess::UAVMask, ERHIAccess::SRVMask));
 
 			// Last pass
 			TShaderMapRef<CSH_Histogram> computeShader(GetGlobalShaderMap(GMaxRHIFeatureLevel), PermutationDomain);
diff --git a/Engine/Source/Runtime/RenderCore/Public/RendererInterface.h b/Engine/Source/Runtime/RenderCore/Public/RendererInterface.h
--- a/Engine/Source/Runtime/RenderCore/Public/RendererInterface.h
+++ b/Engine/Source/Runtime/RenderCore/Public/RendererInterface.h
@@ -861,6 +861,9 @@ public:
 	/** Evict all data from virtual texture cache for a given allocated virtual texture. */
 	virtual void FlushVirtualTextureCache(IAllocatedVirtualTexture* AllocatedVT, const FVector2f& InUV0, const FVector2f& InUV1) = 0;
 
+	/** Do a full virtual texture system update to read and process feedback. This will sync until all outstanding page rendering or streaming is complete. */
+	virtual void SyncVirtualTextureUpdates(FRHICommandListImmediate& RHICmdList, ERHIFeatureLevel::Type FeatureLevel) = 0;
+
 	/** Allocate a buffer and record all nanite page requests until the next call to either SetNaniteRequestRecordBuffer or GetNaniteRequestRecordBuffer. */
 	virtual void SetNaniteRequestRecordBuffer(uint64 Handle) = 0;
 	/** Fetch the page requests recorded since the last call to SetNaniteRequestRecordBuffer. Returns the handle that was passed in. */
diff --git a/Engine/Source/Runtime/Renderer/Internal/VT/VirtualTextureFeedbackResource.h b/Engine/Source/Runtime/Renderer/Internal/VT/VirtualTextureFeedbackResource.h
--- a/Engine/Source/Runtime/Renderer/Internal/VT/VirtualTextureFeedbackResource.h
+++ b/Engine/Source/Runtime/Renderer/Internal/VT/VirtualTextureFeedbackResource.h
@@ -36,7 +36,7 @@ namespace VirtualTexture
 	 * If bInExtendFeedbackForDebug is true we will allocate room at the end of the feedback buffer to capture extra debug info.
 	 * Note that there is only one global feedback buffer alive at any time and its buffer size is later used internally by GetFeedbackShaderParams().
 	 */
-	void BeginFeedback(FRDGBuilder& GraphBuilder, FIntPoint InViewportSize, uint32 InVirtualTextureFeedbackTileSize, bool bInExtendFeedbackForDebug, ERHIFeatureLevel::Type InFeatureLevel);
+	RENDERER_API void BeginFeedback(FRDGBuilder& GraphBuilder, FIntPoint InViewportSize, uint32 InVirtualTextureFeedbackTileSize, bool bInExtendFeedbackForDebug, ERHIFeatureLevel::Type InFeatureLevel);
 
 	/** End a virtual texture feedback scope. */
 	RENDERER_API void EndFeedback(FRDGBuilder& GraphBuilder);
@@ -45,7 +45,7 @@ namespace VirtualTexture
 	 * Get virtual texture feedback parameters to use for binding to any shader that is sampling virtual textures. 
 	 * This is only valid within a BeginFeedback()/EndFeedback() scope.
 	 */
-	void GetFeedbackShaderParams(int32 InFrameIndex, int32 InVirtualTextureFeedbackTileSize, FFeedbackShaderParams& OutParams);
+	RENDERER_API void GetFeedbackShaderParams(int32 InFrameIndex, int32 InVirtualTextureFeedbackTileSize, FFeedbackShaderParams& OutParams);
 	
 	/** 
 	 * Get virtual texture feedback parameters to use for binding to any shader that is sampling virtual textures. 
diff --git a/Engine/Source/Runtime/Renderer/Private/RendererModule.h b/Engine/Source/Runtime/Renderer/Private/RendererModule.h
--- a/Engine/Source/Runtime/Renderer/Private/RendererModule.h
+++ b/Engine/Source/Runtime/Renderer/Private/RendererModule.h
@@ -120,6 +120,7 @@ public:
 	virtual uint64 GetVirtualTextureRequestRecordBuffer(TSet<uint64>& OutPageRequests) override;
 	virtual void FlushVirtualTextureCache() override;
 	virtual void FlushVirtualTextureCache(IAllocatedVirtualTexture* AllocatedVT, const FVector2f& InUV0, const FVector2f& InUV1) override;
+	virtual void SyncVirtualTextureUpdates(FRHICommandListImmediate& RHICmdList, ERHIFeatureLevel::Type FeatureLevel) override;
 
 	virtual void SetNaniteRequestRecordBuffer(uint64 Handle) override;
 	virtual uint64 GetNaniteRequestRecordBuffer(TArray<uint32>& OutPageRequests) override; 	
diff --git a/Engine/Source/Runtime/Renderer/Private/SceneRendering.cpp b/Engine/Source/Runtime/Renderer/Private/SceneRendering.cpp
--- a/Engine/Source/Runtime/Renderer/Private/SceneRendering.cpp
+++ b/Engine/Source/Runtime/Renderer/Private/SceneRendering.cpp
@@ -5590,6 +5590,15 @@ void FRendererModule::FlushVirtualTextureCache(IAllocatedVirtualTexture* Allocat
 	}
 }
 
+void FRendererModule::SyncVirtualTextureUpdates(FRHICommandListImmediate& RHICmdList, ERHIFeatureLevel::Type FeatureLevel)
+{
+	FRDGBuilder GraphBuilder(RHICmdList);
+	FVirtualTextureUpdateSettings Settings;
+	Settings.EnableThrottling(false);
+	FVirtualTextureSystem::Get().Update(GraphBuilder, FeatureLevel, nullptr, Settings);
+	GraphBuilder.Execute();
+}
+
 uint64 FRendererModule::GetNaniteRequestRecordBuffer(TArray<uint32>& OutPageRequests)
 {
 #if WITH_EDITOR
diff --git a/Engine/Source/Runtime/Renderer/Private/VT/VirtualTextureSystem.cpp b/Engine/Source/Runtime/Renderer/Private/VT/VirtualTextureSystem.cpp
--- a/Engine/Source/Runtime/Renderer/Private/VT/VirtualTextureSystem.cpp
+++ b/Engine/Source/Runtime/Renderer/Private/VT/VirtualTextureSystem.cpp
@@ -190,8 +190,9 @@ FVirtualTextureUpdateSettings::FVirtualTextureUpdateSettings()
 	bEnableFeedback = CVarVTEnableFeedback.GetValueOnRenderThread() != 0;
 	bEnableFeedbackProduce = CVarVTEnableFeedbackProduce.GetValueOnRenderThread() != 0;
 	bEnablePlayback = CVarVTEnablePlayback.GetValueOnRenderThread() != 0;
-	bForceContinuousUpdate = CVarVTForceContinuousUpdate.GetValueOnRenderThread() != 0;
 	bParallelFeedbackTasks = CVarVTParallelFeedbackTasks.GetValueOnRenderThread() != 0;
+	bForceContinuousUpdate = CVarVTForceContinuousUpdate.GetValueOnRenderThread() != 0;
+	bForceSyncPageUpdate = false;
 	NumFeedbackTasks = CVarVTNumFeedbackTasks.GetValueOnRenderThread();
 	NumGatherTasks = CVarVTNumGatherTasks.GetValueOnRenderThread();
 	MaxGatherPagesBeforeFlush = CVarVTPageUpdateFlushCount.GetValueOnRenderThread();
@@ -2426,7 +2427,7 @@ void FVirtualTextureSystem::SubmitRequests(FRHICommandList& RHICmdList, ERHIFeat
 		for (uint32 RequestIndex = 0u; RequestIndex < RequestList->GetNumLoadRequests(); ++RequestIndex)
 		{
 			const bool bLockTile = RequestList->IsLocked(RequestIndex);
-			const bool bForceProduceTile = !bAsync || (bLockTile && bSyncProduceLockedTiles);
+			const bool bForceProduceTile = !bAsync || (bLockTile && bSyncProduceLockedTiles) || Settings.bForceSyncPageUpdate;
 			const FVirtualTextureLocalTile TileToLoad = RequestList->GetLoadRequest(RequestIndex).GetTile();
 			const FVirtualTextureProducerHandle ProducerHandle = TileToLoad.GetProducerHandle();
 			const FVirtualTextureProducer& Producer = Producers.GetProducer(ProducerHandle);
@@ -2441,7 +2442,7 @@ void FVirtualTextureSystem::SubmitRequests(FRHICommandList& RHICmdList, ERHIFeat
 				}
 			}
 
-			const EVTRequestPagePriority Priority = bLockTile ? EVTRequestPagePriority::High : EVTRequestPagePriority::Normal;
+			const EVTRequestPagePriority Priority = bLockTile || Settings.bForceSyncPageUpdate ? EVTRequestPagePriority::High : EVTRequestPagePriority::Normal;
 			FVTRequestPageResult RequestPageResult = Producer.GetVirtualTexture()->RequestPageData(RHICmdList, ProducerHandle, ProducerTextureLayerMask, TileToLoad.Local_vLevel, TileToLoad.Local_vAddress, Priority);
 			if (RequestPageResult.Status == EVTRequestPageStatus::Pending && bForceProduceTile)
 			{
diff --git a/Engine/Source/Runtime/Renderer/Private/VT/VirtualTextureSystem.h b/Engine/Source/Runtime/Renderer/Private/VT/VirtualTextureSystem.h
--- a/Engine/Source/Runtime/Renderer/Private/VT/VirtualTextureSystem.h
+++ b/Engine/Source/Runtime/Renderer/Private/VT/VirtualTextureSystem.h
@@ -46,6 +46,7 @@ struct FVirtualTextureUpdateSettings
 			MaxRVTPageUploads = 99999;
 			MaxSVTPageUploads = 99999;
 			MaxPagesProduced = 99999;
+			bForceSyncPageUpdate = true;
 		}
 		return *this;
 	}
@@ -69,8 +70,9 @@ struct FVirtualTextureUpdateSettings
 	bool bEnableFeedback;
 	bool bEnableFeedbackProduce;
 	bool bEnablePlayback;
-	bool bForceContinuousUpdate;
 	bool bParallelFeedbackTasks;
+	bool bForceContinuousUpdate;
+	bool bForceSyncPageUpdate;
 	int32 NumFeedbackTasks;
 	int32 NumGatherTasks;
 	int32 MaxGatherPagesBeforeFlush;
