TextureGraph: Fix memory leak of SceneViewFamily and SceneView objects during material tile rendering.
#rnx

diff --git a/Engine/Plugins/TextureGraph/Source/TextureGraphEngine/FxMat/FxMaterial_DrawMaterial.cpp b/Engine/Plugins/TextureGraph/Source/TextureGraphEngine/FxMat/FxMaterial_DrawMaterial.cpp
--- a/Engine/Plugins/TextureGraph/Source/TextureGraphEngine/FxMat/FxMaterial_DrawMaterial.cpp
+++ b/Engine/Plugins/TextureGraph/Source/TextureGraphEngine/FxMat/FxMaterial_DrawMaterial.cpp
@@ -24,30 +24,12 @@ static FAutoConsoleVariableRef CVarVirtualTextureFeedbackFactor(
 	TEXT("Texel size of tile used to generate one item of virtual texture feedback."),
 	ECVF_RenderThreadSafe);
 
-FSceneView* FxMaterial_DrawMaterialBase::CreateSceneView(UTextureRenderTarget2D* RenderTarget, const FIntPoint& TargetSizeXY)
+FSceneView* FxMaterial_DrawMaterialBase::CreateSceneView(FSceneViewFamilyContext& ViewFamily, UTextureRenderTarget2D* RenderTarget, const FIntPoint& TargetSizeXY)
 {
-
-	// Create a new FSceneView
-	static const FEngineShowFlags DefaultShowFlags(ESFIM_Game);
-
-	FSceneViewFamilyContext* ViewFamilyContext = new FSceneViewFamilyContext
-	(
-		FSceneViewFamily::ConstructionValues
-		(
-			RenderTarget->GetRenderTargetResource(),
-			nullptr,
-			DefaultShowFlags
-		)
-		.SetTime(FGameTime())
-		//.SetGammaCorrection(DisplayGamma)
-		.SetRealtimeUpdate(true)
-	);
-
 	FIntRect ViewRect(FIntPoint(0, 0), TargetSizeXY);
 
-	// make a temporary view
 	FSceneViewInitOptions ViewInitOptions;
-	ViewInitOptions.ViewFamily = ViewFamilyContext;
+	ViewInitOptions.ViewFamily = &ViewFamily;
 	ViewInitOptions.SetViewRectangle(ViewRect);
 	ViewInitOptions.ViewOrigin = FVector::ZeroVector;
 	ViewInitOptions.ViewRotationMatrix = FMatrix::Identity;
@@ -64,7 +46,7 @@ FSceneView* FxMaterial_DrawMaterialBase::CreateSceneView(UTextureRenderTarget2D*
 	ViewInitOptions.OverlayColor = FLinearColor::White;
 
 	FSceneView* View = new FSceneView(ViewInitOptions);
-	ViewFamilyContext->Views.Add(View);
+	ViewFamily.Views.Add(View);
 
 	{
 		// Create the view's uniform buffer.
diff --git a/Engine/Plugins/TextureGraph/Source/TextureGraphEngine/FxMat/FxMaterial_DrawMaterial.h b/Engine/Plugins/TextureGraph/Source/TextureGraphEngine/FxMat/FxMaterial_DrawMaterial.h
--- a/Engine/Plugins/TextureGraph/Source/TextureGraphEngine/FxMat/FxMaterial_DrawMaterial.h
+++ b/Engine/Plugins/TextureGraph/Source/TextureGraphEngine/FxMat/FxMaterial_DrawMaterial.h
@@ -104,7 +104,7 @@ public:
 	{
 	}
 
-	FSceneView* CreateSceneView(UTextureRenderTarget2D* RenderTarget, const FIntPoint& TargetSizeXY);
+	FSceneView* CreateSceneView(FSceneViewFamilyContext& ViewFamily, UTextureRenderTarget2D* RenderTarget, const FIntPoint& TargetSizeXY);
 
 	bool HasVirtualTextureFeedback(FMaterial const& RenderMaterial, EShaderPlatform ShaderPlatform) const;
 	int32 GetVirtualTextureWarmupFrames() const;
@@ -207,7 +207,11 @@ public:
 				VirtualTextureFeedbackBegin(GraphBuilder, Target->GetSizeXY(), FeatureLevel);
 			}
 
-			FSceneView* SceneView = CreateSceneView(RenderTarget, Target->GetSizeXY());
+			FSceneViewFamilyContext SceneViewFamilyContext(
+				FSceneViewFamily::ConstructionValues(RenderTarget->GetRenderTargetResource(), nullptr, ESFIM_Game)
+				.SetTime(FGameTime())
+				.SetRealtimeUpdate(true));
+			FSceneView* SceneView = CreateSceneView(SceneViewFamilyContext, RenderTarget, Target->GetSizeXY());
 
 			auto* PassParameters = GraphBuilder.AllocParameters<FParameters>();
 			PassParameters->VS = VSHParams;
