[TG] Expose the control for VT num warmup frames on the Material and MAterialFUnction expressions in order to fine tune this parameter to a different value than the default if needed.

the default value is specifed in the CVar TG.VirtualTexture.NumWarmupFrames
and the default is 8

The MaterialBase expression configure the RenderMaterial_BP and the chosen NumWarmupFrames chosen by the user from the exposed setting (in expression_Material and Expression_MaterialFunction)
the default value for these settings is 0 meaning that internally it will fallback to the actual default CVar TG.VirtualTexture.NumWarmupFrames value.
is > 0 then the actual setting is used.

diff --git a/Engine/Plugins/TextureGraph/Source/TextureGraph/Private/Expressions/TG_Expression_MaterialBase.cpp b/Engine/Plugins/TextureGraph/Source/TextureGraph/Private/Expressions/TG_Expression_MaterialBase.cpp
--- a/Engine/Plugins/TextureGraph/Source/TextureGraph/Private/Expressions/TG_Expression_MaterialBase.cpp
+++ b/Engine/Plugins/TextureGraph/Source/TextureGraph/Private/Expressions/TG_Expression_MaterialBase.cpp
@@ -193,7 +193,7 @@ void UTG_Expression_MaterialBase::Evaluate(FTG_EvaluationContext* InContext)
 
 		for (int32 i = 0; i < AvailableMaterialAttributeIds.Num(); ++i)
 		{
-			const auto RenderMaterial = std::make_shared<RenderMaterial_BP>(AssetName, MIC);
+			const auto RenderMaterial = std::make_shared<RenderMaterial_BP>(AssetName, MIC, GetNumVirtualTextureWarmupFrames());
 			TiledBlobPtr Result = TextureHelper::GetBlack();
 			Result = CreateRenderMaterialJob(InContext, RenderMaterial, Output.GetBufferDescriptor(), AvailableMaterialAttributeIds[i]);
 			Outputs[i] = Result;
@@ -216,7 +216,7 @@ void UTG_Expression_MaterialBase::Evaluate(FTG_EvaluationContext* InContext)
 
 TiledBlobPtr UTG_Expression_MaterialBase::CreateRenderMaterialJob(FTG_EvaluationContext* InContext, const FString& InName, const FString& InMaterialPath, const BufferDescriptor& InDescriptor, EDrawMaterialAttributeTarget InDrawMaterialAttributeTarget)
 {
-	const RenderMaterial_BPPtr RenderMaterial = TextureGraphEngine::GetMaterialManager()->CreateMaterial_BP(InName, InMaterialPath);
+	const RenderMaterial_BPPtr RenderMaterial = TextureGraphEngine::GetMaterialManager()->CreateMaterial_BP(InName, InMaterialPath, GetNumVirtualTextureWarmupFrames());
 	return CreateRenderMaterialJob(InContext, RenderMaterial, InDescriptor, InDrawMaterialAttributeTarget);
 }
 
@@ -263,6 +263,7 @@ TiledBlobPtr UTG_Expression_MaterialBase::CreateRenderMaterialJob(FTG_Evaluation
 	MaterialJob->AddArg(ARG_VECTOR(PSControl, "PSControl"));
 	MaterialJob->AddArg(std::make_shared<JobArg_TileInfo>(tileInfo, "TileInfo")); /// Enable the tileinfo parameters
 	MaterialJob->AddArg(std::make_shared<JobArg_ForceTiling>()); /// Force hashing individual tiles differently
+	MaterialJob->AddArg(ARG_INT(GetNumVirtualTextureWarmupFrames(), "NumVirtualTextureWarmupFrames")); // Inject the num warmup frames as part of the caching key
 
 	BufferDescriptor Desc = InDescriptor;
 
diff --git a/Engine/Plugins/TextureGraph/Source/TextureGraph/Public/Expressions/Input/TG_Expression_Material.h b/Engine/Plugins/TextureGraph/Source/TextureGraph/Public/Expressions/Input/TG_Expression_Material.h
--- a/Engine/Plugins/TextureGraph/Source/TextureGraph/Public/Expressions/Input/TG_Expression_Material.h
+++ b/Engine/Plugins/TextureGraph/Source/TextureGraph/Public/Expressions/Input/TG_Expression_Material.h
@@ -50,6 +50,12 @@ public:
 	UFUNCTION(CallInEditor)
 	UE_API TArray<FName> GetRenderAttributeOptions() const;
 
+	// When Material references VT, can specify the number of warmup frames.
+	// Default value is 0 meaning that the CVar <TG.VirtualTexture.NumWarmupFrames> is used instead 
+	UPROPERTY(EditAnywhere, BlueprintReadWrite, Category = NoCategory, meta = (TGType = "TG_Setting", ClampMin = 0, ClampMax = 32, PinDisplayName = "Num Warmup Frames") )
+	int32 NumWarmupFrames = 0;
+	virtual int32 GetNumVirtualTextureWarmupFrames() const override { return NumWarmupFrames; }
+
 	UE_API virtual bool CanHandleAsset(UObject* Asset) override;
 	UE_API virtual void SetAsset(UObject* Asset) override;
 	virtual FText GetTooltipText() const override { return FText::FromString(TEXT("Renders a material into a quad and makes it available. It is automatically exposed as a graph input parameter.")); } 
diff --git a/Engine/Plugins/TextureGraph/Source/TextureGraph/Public/Expressions/Input/TG_Expression_MaterialFunction.h b/Engine/Plugins/TextureGraph/Source/TextureGraph/Public/Expressions/Input/TG_Expression_MaterialFunction.h
--- a/Engine/Plugins/TextureGraph/Source/TextureGraph/Public/Expressions/Input/TG_Expression_MaterialFunction.h
+++ b/Engine/Plugins/TextureGraph/Source/TextureGraph/Public/Expressions/Input/TG_Expression_MaterialFunction.h
@@ -28,6 +28,12 @@ public:
 	TObjectPtr<UMaterialFunctionInterface> MaterialFunction;
 	UE_API void SetMaterialFunction(UMaterialFunctionInterface* InMaterialFunction);
 
+	// When MaterialFunction references VT, can specify the number of warmup frames.
+	// Default value is 0 meaning that the CVar <TG.VirtualTexture.NumWarmupFrames> is used instead 
+	UPROPERTY(EditAnywhere, BlueprintReadWrite, Category = NoCategory, meta = (TGType = "TG_Setting", ClampMin = 0, ClampMax = 32, PinDisplayName = "Num Warmup Frames") )
+	int32 NumWarmupFrames = 0;
+	virtual int32 GetNumVirtualTextureWarmupFrames() const override { return NumWarmupFrames; }
+
 	UE_API virtual bool CanHandleAsset(UObject* Asset) override;
 	UE_API virtual void SetAsset(UObject* Asset) override;
 
diff --git a/Engine/Plugins/TextureGraph/Source/TextureGraph/Public/Expressions/TG_Expression_MaterialBase.h b/Engine/Plugins/TextureGraph/Source/TextureGraph/Public/Expressions/TG_Expression_MaterialBase.h
--- a/Engine/Plugins/TextureGraph/Source/TextureGraph/Public/Expressions/TG_Expression_MaterialBase.h
+++ b/Engine/Plugins/TextureGraph/Source/TextureGraph/Public/Expressions/TG_Expression_MaterialBase.h
@@ -68,6 +68,9 @@ public:
 
 	static UE_API FName CPPTypeNameFromMaterialParamType(EMaterialParameterType InMatType);
 
+	// 0 will fallback to CVar <TG.VirtualTexture.NumWarmupFrames> value.
+	virtual int32 GetNumVirtualTextureWarmupFrames() const { return 0; }
+
 protected:
 	// A local Material Instance Constant is recreated from the reference material assigned through SetMaterialInternal
 	// All static switches permutations are cached in this map and only created on demand
diff --git a/Engine/Plugins/TextureGraph/Source/TextureGraphEngine/FxMat/FxMaterial_DrawMaterial.cpp b/Engine/Plugins/TextureGraph/Source/TextureGraphEngine/FxMat/FxMaterial_DrawMaterial.cpp
--- a/Engine/Plugins/TextureGraph/Source/TextureGraphEngine/FxMat/FxMaterial_DrawMaterial.cpp
+++ b/Engine/Plugins/TextureGraph/Source/TextureGraphEngine/FxMat/FxMaterial_DrawMaterial.cpp
@@ -10,10 +10,10 @@ IMPLEMENT_MATERIAL_SHADER_TYPE(, FTextureGraphMaterialShaderPS, TEXT("/Plugin/Te
 
 FName FTextureGraphMaterialShaderPS::PSCONTROL_ARG = TEXT("PSControl");
 
-static int32 GVirtualTextureWarmupFrames = 8;
-static FAutoConsoleVariableRef CVarVirtualTextureWarmupFrames(
-	TEXT("TG.VirtualTexture.WarmupFrames"),
-	GVirtualTextureWarmupFrames,
+static int32 GTG_VirtualTextureNumWarmupFrames = 8;
+static FAutoConsoleVariableRef CVarTGVirtualTextureNumWarmupFrames(
+	TEXT("TG.VirtualTexture.NumWarmupFrames"),
+	GTG_VirtualTextureNumWarmupFrames,
 	TEXT("Number of warmup render frames when sampling from virtual textures."),
 	ECVF_RenderThreadSafe);
 
@@ -95,9 +95,9 @@ bool FxMaterial_DrawMaterialBase::HasVirtualTextureFeedback(FMaterial const& Ren
 	return UseVirtualTexturing(ShaderPlatform) && !RenderMaterial.GetUniformVirtualTextureExpressions().IsEmpty();
 }
 
-int32 FxMaterial_DrawMaterialBase::GetVirtualTextureWarmupFrames() const
+int32 FxMaterial_DrawMaterialBase::GetDefaultVirtualTextureWarmupFrames() const
 {
-	return FMath::Max(GVirtualTextureWarmupFrames, 1);
+	return FMath::Max(GTG_VirtualTextureNumWarmupFrames, 1);
 }
 
 void FxMaterial_DrawMaterialBase::VirtualTextureFeedbackBegin(FRDGBuilder& GraphBuilder, FIntPoint TargetSize, ERHIFeatureLevel::Type FeatureLevel) const
diff --git a/Engine/Plugins/TextureGraph/Source/TextureGraphEngine/FxMat/FxMaterial_DrawMaterial.h b/Engine/Plugins/TextureGraph/Source/TextureGraphEngine/FxMat/FxMaterial_DrawMaterial.h
--- a/Engine/Plugins/TextureGraph/Source/TextureGraphEngine/FxMat/FxMaterial_DrawMaterial.h
+++ b/Engine/Plugins/TextureGraph/Source/TextureGraphEngine/FxMat/FxMaterial_DrawMaterial.h
@@ -107,7 +107,7 @@ public:
 	FSceneView* CreateSceneView(FSceneViewFamilyContext& ViewFamily, UTextureRenderTarget2D* RenderTarget, const FIntPoint& TargetSizeXY);
 
 	bool HasVirtualTextureFeedback(FMaterial const& RenderMaterial, EShaderPlatform ShaderPlatform) const;
-	int32 GetVirtualTextureWarmupFrames() const;
+	int32 GetDefaultVirtualTextureWarmupFrames() const;
 	void VirtualTextureFeedbackBegin(FRDGBuilder& GraphBuilder, FIntPoint TargetSize, ERHIFeatureLevel::Type FeatureLevel) const;
 	void VirtualTextureFeedbackEnd(FRDGBuilder& GraphBuilder) const;
 	void VirtualTextureFeedbackSync(FRHICommandListImmediate& RHICmdList, ERHIFeatureLevel::Type FeatureLevel) const;
@@ -139,6 +139,10 @@ protected:
 	typename VSH_Type::FParameters	VSHParams;			/// Params for the vertex shader
 	typename FSH_Type::FParameters	FSHParams;			/// Params for the fragment shader
 
+public:
+	// If Material uses VT, we need to run several frames to make sure the VT is fully loaded.
+	// 0 indicate to use the actual default value specified from the CVar <TG.VirtualTexture.NumWarmupFrames>
+	int32							VirtualTextureNumWarmupFrames = 0;
 
 public:
 
@@ -147,17 +151,31 @@ public:
 	{
 	}
 
-	FxMaterial_DrawMaterial(UMaterialInterface* InMaterial, const VSHPermutationDomain& InVSHPermutationDomain, const FSHPermutationDomain& InFSHPermutationDomain)
+	FxMaterial_DrawMaterial(UMaterialInterface* InMaterial, const VSHPermutationDomain& InVSHPermutationDomain, const FSHPermutationDomain& InFSHPermutationDomain, int32 InVirtualTextureNumWarmupFrames)
 		: 	FxMaterial_DrawMaterialBase(),
 		Material(InMaterial),
 		VSHPermDomain(InVSHPermutationDomain),
-		FSHPermDomain(InFSHPermutationDomain)
+		FSHPermDomain(InFSHPermutationDomain),
+		VirtualTextureNumWarmupFrames(InVirtualTextureNumWarmupFrames)
 	{
 	}
 
 	virtual std::shared_ptr<FxMaterial> Clone() override
 	{
-		return std::static_pointer_cast<FxMaterial>(std::make_shared<FxMaterial_DrawMaterial<VSH_Type, FSH_Type>>(Material, VSHPermDomain, FSHPermDomain));
+		return std::static_pointer_cast<FxMaterial>(std::make_shared<FxMaterial_DrawMaterial<VSH_Type, FSH_Type>>(Material, VSHPermDomain, FSHPermDomain, VirtualTextureNumWarmupFrames));
+	}
+	
+	int32 GetVirtualTextureWarmupFrames() const
+	{
+		// If the instance value is invalid, then fallback to the default
+		if (VirtualTextureNumWarmupFrames < 1)
+		{
+			return GetDefaultVirtualTextureWarmupFrames();
+		}
+		else
+		{
+			return VirtualTextureNumWarmupFrames;
+		}
 	}
 
 	virtual FxMetadataSet GetMetadata() const override
diff --git a/Engine/Plugins/TextureGraph/Source/TextureGraphEngine/FxMat/MaterialManager.cpp b/Engine/Plugins/TextureGraph/Source/TextureGraphEngine/FxMat/MaterialManager.cpp
--- a/Engine/Plugins/TextureGraph/Source/TextureGraphEngine/FxMat/MaterialManager.cpp
+++ b/Engine/Plugins/TextureGraph/Source/TextureGraphEngine/FxMat/MaterialManager.cpp
@@ -41,21 +41,21 @@ UMaterial* UMaterialManager::LoadMaterial(const FString& Path)
 	return Mat;
 }
 
-RenderMaterial_BPPtr UMaterialManager::CreateMaterial_BP(FString Name, const FString& Path)
+RenderMaterial_BPPtr UMaterialManager::CreateMaterial_BP(FString Name, const FString& Path, int32 InNumWarmupFrames)
 {
 	UMaterial* Mat = LoadMaterial(Path);
-	return CreateMaterial_BP(Name, Mat);
+	return CreateMaterial_BP(Name, Mat, InNumWarmupFrames);
 }
 
-RenderMaterial_BPPtr UMaterialManager::CreateMaterial_BP(FString InName, UMaterial* InMaterial)
+RenderMaterial_BPPtr UMaterialManager::CreateMaterial_BP(FString InName, UMaterial* InMaterial, int32 InNumWarmupFrames)
 {
-	return std::make_shared<RenderMaterial_BP>(InName, InMaterial);
+	return std::make_shared<RenderMaterial_BP>(InName, InMaterial, InNumWarmupFrames);
 }
 
-RenderMaterial_BP_NoTilePtr UMaterialManager::CreateMaterial_BP_NoTile(FString Name, const FString& Path)
+RenderMaterial_BP_NoTilePtr UMaterialManager::CreateMaterial_BP_NoTile(FString Name, const FString& Path, int32 InNumWarmupFrames)
 {
 	UMaterial* Mat = LoadMaterial(Path);
-	return std::make_shared<RenderMaterial_BP_NoTile>(Name, Mat);
+	return std::make_shared<RenderMaterial_BP_NoTile>(Name, Mat, InNumWarmupFrames);
 }
 
 RenderMaterial_ThumbPtr UMaterialManager::CreateMaterial_Thumbnail(FString Name, const FString& Path)
@@ -64,8 +64,8 @@ RenderMaterial_ThumbPtr UMaterialManager::CreateMaterial_Thumbnail(FString Name,
 	return std::make_shared<RenderMaterial_Thumbnail>(Name, Mat);
 }
 
-RenderMaterial_BP_TileArgsPtr UMaterialManager::CreateMaterialOfType_BP_TileArgs(FString Name, const FString& Path)
+RenderMaterial_BP_TileArgsPtr UMaterialManager::CreateMaterialOfType_BP_TileArgs(FString Name, const FString& Path, int32 InNumWarmupFrames)
 {
 	UMaterial* Mat = LoadMaterial(Path);
-	return std::make_shared<RenderMaterial_BP_TileArgs>(Name, Mat);
+	return std::make_shared<RenderMaterial_BP_TileArgs>(Name, Mat, InNumWarmupFrames);
 }
diff --git a/Engine/Plugins/TextureGraph/Source/TextureGraphEngine/FxMat/MaterialManager.h b/Engine/Plugins/TextureGraph/Source/TextureGraphEngine/FxMat/MaterialManager.h
--- a/Engine/Plugins/TextureGraph/Source/TextureGraphEngine/FxMat/MaterialManager.h
+++ b/Engine/Plugins/TextureGraph/Source/TextureGraphEngine/FxMat/MaterialManager.h
@@ -31,11 +31,11 @@ public:
 	
 	UE_API class UMaterial*				LoadMaterial(const FString& Path);	
 
-	UE_API RenderMaterial_BPPtr			CreateMaterial_BP(FString Name, const FString& Path);
-	UE_API RenderMaterial_BPPtr			CreateMaterial_BP(FString InName, UMaterial* InMaterial);
-	UE_API RenderMaterial_BP_NoTilePtr		CreateMaterial_BP_NoTile(FString Name, const FString& Path);
+	UE_API RenderMaterial_BPPtr			CreateMaterial_BP(FString Name, const FString& Path, int32 InNumWarmupFrames = 0);
+	UE_API RenderMaterial_BPPtr			CreateMaterial_BP(FString InName, UMaterial* InMaterial, int32 InNumWarmupFrames = 0);
+	UE_API RenderMaterial_BP_NoTilePtr		CreateMaterial_BP_NoTile(FString Name, const FString& Path, int32 InNumWarmupFrames = 0);
 	UE_API RenderMaterial_ThumbPtr			CreateMaterial_Thumbnail(FString Name, const FString& Path);
-	UE_API RenderMaterial_BP_TileArgsPtr	CreateMaterialOfType_BP_TileArgs(FString Name, const FString& Path);
+	UE_API RenderMaterial_BP_TileArgsPtr	CreateMaterialOfType_BP_TileArgs(FString Name, const FString& Path, int32 InNumWarmupFrames = 0);
 
 	//////////////////////////////////////////////////////////////////////////
 
diff --git a/Engine/Plugins/TextureGraph/Source/TextureGraphEngine/FxMat/RenderMaterial_BP.cpp b/Engine/Plugins/TextureGraph/Source/TextureGraphEngine/FxMat/RenderMaterial_BP.cpp
--- a/Engine/Plugins/TextureGraph/Source/TextureGraphEngine/FxMat/RenderMaterial_BP.cpp
+++ b/Engine/Plugins/TextureGraph/Source/TextureGraphEngine/FxMat/RenderMaterial_BP.cpp
@@ -17,7 +17,6 @@
 #include "FxMaterial_DrawMaterial.h"
 #include "Job/Job.h"
 
-
 TSharedPtr<FMaterialInstanceStaticSwitchPermutationMap> FMaterialInstanceStaticSwitchPermutationMap::Create(UMaterialInterface* InMaterial)
 {
 	TSharedPtr<FMaterialInstanceStaticSwitchPermutationMap> MISSPMap;
@@ -104,10 +103,11 @@ bool RenderMaterial_BP::ValidateMaterialCompatible(UMaterialInterface* InMateria
 	return FxMaterial_QuadDrawMaterial::ValidateMaterial(InMaterial);
 }
 
-RenderMaterial_BP::RenderMaterial_BP(FString Name, UMaterialInterface* InMaterial)
+RenderMaterial_BP::RenderMaterial_BP(FString Name, UMaterialInterface* InMaterial, int32 InNumWarmupFrames)
 	: RenderMaterial(!Name.IsEmpty() ? Name : InMaterial->GetName())
 	, Material(InMaterial)
 	, FXMaterialObj(std::make_shared<FxMaterial_QuadDrawMaterial>())
+	, VirtualTextureNumWarmupFrames(InNumWarmupFrames)
 {
 	// Cast to or build a MIC to be used by this RenderMaterial and its duplicates
 	if (Material && Material->IsA< UMaterialInstanceConstant>())
@@ -140,6 +140,7 @@ RenderMaterial_BP::RenderMaterial_BP(FString Name, UMaterialInterface* InMateria
 	HashValue = std::make_shared<CHash>(DataUtil::Hash((uint8*)Material, sizeof(UMaterial*)), true);
 
 	std::static_pointer_cast<FxMaterial_QuadDrawMaterial>(FXMaterialObj)->Material = MaterialInstance.Get();
+	std::static_pointer_cast<FxMaterial_QuadDrawMaterial>(FXMaterialObj)->VirtualTextureNumWarmupFrames = VirtualTextureNumWarmupFrames;
 }
 
 RenderMaterial_BP::~RenderMaterial_BP()
@@ -204,7 +205,7 @@ std::shared_ptr<BlobTransform> RenderMaterial_BP::DuplicateInstance(FString InNa
 	if (InName.IsEmpty())
 		InName = Name;
 
-	RenderMaterial_BPPtr MaterialBP = std::make_shared<RenderMaterial_BP>(InName, MaterialInstance.Get()); //Reuse the same material instance for every invocations
+	RenderMaterial_BPPtr MaterialBP = std::make_shared<RenderMaterial_BP>(InName, MaterialInstance.Get(), GetVirtualTextureNumWarmupFrames()); //Reuse the same material instance for every invocations
 
 	MaterialBP->MaterialInstanceValidated = MaterialInstanceValidated;
 	std::shared_ptr<BlobTransform> Result = std::static_pointer_cast<RenderMaterial>(MaterialBP);
diff --git a/Engine/Plugins/TextureGraph/Source/TextureGraphEngine/FxMat/RenderMaterial_BP.h b/Engine/Plugins/TextureGraph/Source/TextureGraphEngine/FxMat/RenderMaterial_BP.h
--- a/Engine/Plugins/TextureGraph/Source/TextureGraphEngine/FxMat/RenderMaterial_BP.h
+++ b/Engine/Plugins/TextureGraph/Source/TextureGraphEngine/FxMat/RenderMaterial_BP.h
@@ -29,6 +29,11 @@ struct FMaterialInstanceStaticSwitchPermutationMap
 
 class RenderMaterial_BP : public RenderMaterial
 {
+public:
+	// When a Material contains a VirtualTexture, the default num of warmup frames run to stream the appropriate resolution
+	// This can be fine tuned by instance of RenderMaterial_BP from the constructor 
+	static UE_API int32 GetDefaultVirtualTextureNumWarmupFrames();
+
 protected:
 
 	UMaterialInterface*				Material = nullptr;			/// The base material that is used for this job 
@@ -40,13 +45,15 @@ protected:
 
 	FxMaterialPtr					FXMaterialObj;
 
+	int32 							VirtualTextureNumWarmupFrames = 0;  // Num warmup frames for this instance default to -1
+
 	UCanvas*						Canvas = nullptr;
 
 	UE_API void							DrawMaterial(UMaterialInterface* RenderMaterial, FVector2D ScreenPosition, FVector2D ScreenSize,
 									FVector2D CoordinatePosition, FVector2D CoordinateSize=FVector2D::UnitVector, float Rotation=0.f,
 									FVector2D PivotPoint=FVector2D(0.5f,0.5f)) const;
 public:
-									UE_API RenderMaterial_BP(FString Name, UMaterialInterface* InMaterial);
+									UE_API RenderMaterial_BP(FString Name, UMaterialInterface* InMaterial, int32 InNumWarmupFrames);
 
 	UE_API virtual							~RenderMaterial_BP() override;
 
@@ -85,6 +92,7 @@ public:
 	FORCEINLINE UMaterialInterface*					GetMaterial() { return Material; }
 	FORCEINLINE UMaterialInstanceConstant*			Instance() { return MaterialInstance.Get(); }
 	FORCEINLINE const UMaterialInstanceConstant*	Instance() const { return MaterialInstance.Get(); }
+	FORCEINLINE int32 GetVirtualTextureNumWarmupFrames() const { return VirtualTextureNumWarmupFrames; }
 };
 
 typedef std::shared_ptr<RenderMaterial_BP> RenderMaterial_BPPtr;
diff --git a/Engine/Plugins/TextureGraph/Source/TextureGraphEngine/FxMat/RenderMaterial_BP_NoTile.cpp b/Engine/Plugins/TextureGraph/Source/TextureGraphEngine/FxMat/RenderMaterial_BP_NoTile.cpp
--- a/Engine/Plugins/TextureGraph/Source/TextureGraphEngine/FxMat/RenderMaterial_BP_NoTile.cpp
+++ b/Engine/Plugins/TextureGraph/Source/TextureGraphEngine/FxMat/RenderMaterial_BP_NoTile.cpp
@@ -21,7 +21,7 @@ std::shared_ptr<BlobTransform> RenderMaterial_BP_TileArgs::DuplicateInstance(FSt
 	if (NewName.IsEmpty())
 		NewName = Name;
 
-	RenderMaterial_BP_TileArgsPtr MaterialBP = std::make_shared<RenderMaterial_BP_TileArgs>(NewName, Instance()); //We would want a new instance every time
+	RenderMaterial_BP_TileArgsPtr MaterialBP = std::make_shared<RenderMaterial_BP_TileArgs>(NewName, Instance(), GetVirtualTextureNumWarmupFrames()); //We would want a new instance every time
 	std::shared_ptr<BlobTransform> Result = std::static_pointer_cast<RenderMaterial_BP_TileArgs>(MaterialBP);
 
 	return Result;
diff --git a/Engine/Plugins/TextureGraph/Source/TextureGraphEngine/FxMat/RenderMaterial_BP_NoTile.h b/Engine/Plugins/TextureGraph/Source/TextureGraphEngine/FxMat/RenderMaterial_BP_NoTile.h
--- a/Engine/Plugins/TextureGraph/Source/TextureGraphEngine/FxMat/RenderMaterial_BP_NoTile.h
+++ b/Engine/Plugins/TextureGraph/Source/TextureGraphEngine/FxMat/RenderMaterial_BP_NoTile.h
@@ -13,7 +13,7 @@ class RenderMaterial_BP_NoTile : public RenderMaterial_BP
 {
 
 public:
-	RenderMaterial_BP_NoTile(FString InName, UMaterialInterface* InMaterial) : RenderMaterial_BP(InName, InMaterial) {}
+	RenderMaterial_BP_NoTile(FString InName, UMaterialInterface* InMaterial, int32 InVTNumWarmUpFrames) : RenderMaterial_BP(InName, InMaterial, InVTNumWarmUpFrames) {}
 	virtual bool					CanHandleTiles() const override { return false; };
 };
 
@@ -25,7 +25,7 @@ private:
 	FName							MemberName = "TileInfo";
 
 public:
-	RenderMaterial_BP_TileArgs(FString InName, UMaterialInterface* InMaterial) : RenderMaterial_BP(InName, InMaterial) {}
+	RenderMaterial_BP_TileArgs(FString InName, UMaterialInterface* InMaterial, int32 InVTNumWarmUpFrames) : RenderMaterial_BP(InName, InMaterial, InVTNumWarmUpFrames) {}
 	UE_API void							AddTileArgs(TransformArgs& Args);
 	UE_API virtual AsyncPrepareResult		PrepareResources(const TransformArgs& Args) override;
 	virtual bool					CanHandleTiles() const override { return true; };
diff --git a/Engine/Plugins/TextureGraph/Source/TextureGraphEngine/FxMat/RenderMaterial_Thumbnail.cpp b/Engine/Plugins/TextureGraph/Source/TextureGraphEngine/FxMat/RenderMaterial_Thumbnail.cpp
--- a/Engine/Plugins/TextureGraph/Source/TextureGraphEngine/FxMat/RenderMaterial_Thumbnail.cpp
+++ b/Engine/Plugins/TextureGraph/Source/TextureGraphEngine/FxMat/RenderMaterial_Thumbnail.cpp
@@ -14,7 +14,7 @@
 
 ////////////////////////////////////////////////////////////////////////// RenderMaterial_Thumbnail //////////////////////////////////////////////////////////////////////////
 
-RenderMaterial_Thumbnail::RenderMaterial_Thumbnail(FString InName, UMaterialInterface* InMaterial) : RenderMaterial_BP(InName, InMaterial)
+RenderMaterial_Thumbnail::RenderMaterial_Thumbnail(FString InName, UMaterialInterface* InMaterial) : RenderMaterial_BP(InName, InMaterial, 0)
 {
 }
 
