TextureGraph sRGB fixes

diff --git a/Engine/Plugins/TextureGraph/Source/TextureGraph/Private/Expressions/Output/TG_Expression_Output.cpp b/Engine/Plugins/TextureGraph/Source/TextureGraph/Private/Expressions/Output/TG_Expression_Output.cpp
--- a/Engine/Plugins/TextureGraph/Source/TextureGraph/Private/Expressions/Output/TG_Expression_Output.cpp
+++ b/Engine/Plugins/TextureGraph/Source/TextureGraph/Private/Expressions/Output/TG_Expression_Output.cpp
@@ -38,10 +38,6 @@ void UTG_Expression_Output::Evaluate(FTG_EvaluationContext* InContext)
 			BufferDescriptor OutputDesc = Output.EditTexture().GetBufferDescriptor();
 			const BufferDescriptor& InputDesc = Source.GetTexture()->GetDescriptor();
 
-			// When Output and Input desc don't match, it just ignores input's Desc sRGB value and uses Output desc value (which is currently always false) 
-			// Right now, the sRGB settings cannot be changed in the output node, so we can safely apply the input sRGB setting here.
-			OutputDesc.bIsSRGB = InputDesc.bIsSRGB;
-			
 			/// If the descriptors are not the same
 			if (OutputDesc != InputDesc)
 			{
@@ -56,7 +52,9 @@ void UTG_Expression_Output::Evaluate(FTG_EvaluationContext* InContext)
 				RenderJob
 					->AddArg(ARG_BLOB(Source.GetTexture().RasterBlob, "SourceTexture"))
 					->AddArg(WithUnbounded(ARG_INT(OutputDesc.ItemsPerPoint, "ItemsPerPoint")))
-					->AddArg(WithUnbounded(ARG_INT(Format, "Format")));
+					->AddArg(WithUnbounded(ARG_INT(Format, "Format")))
+					->AddArg(WithUnbounded(ARG_BOOL(OutputDesc.bIsSRGB, "sRGB")));
+
 
 				const FString Name = TEXT("Output");
 				OutputDesc.Name = Name;
@@ -135,6 +133,7 @@ void UTG_Expression_Output::UpdateBufferDescriptorValues()
 	OutputTexture.Descriptor.Width = OutputSettings.Width;
 	OutputTexture.Descriptor.Height = OutputSettings.Height;
 	OutputTexture.Descriptor.TextureFormat = OutputSettings.TextureFormat;
+	OutputTexture.Descriptor.bIsSRGB = OutputSettings.bSRGB;
 }
 
 void UTG_Expression_Output::InitializeOutputSettings()
diff --git a/Engine/Plugins/TextureGraph/Source/TextureGraphEngine/Data/RawBuffer.cpp b/Engine/Plugins/TextureGraph/Source/TextureGraphEngine/Data/RawBuffer.cpp
--- a/Engine/Plugins/TextureGraph/Source/TextureGraphEngine/Data/RawBuffer.cpp
+++ b/Engine/Plugins/TextureGraph/Source/TextureGraphEngine/Data/RawBuffer.cpp
@@ -84,6 +84,8 @@ bool BufferDescriptor::operator==(const BufferDescriptor& rhs) const
 	
 	if (bIsSame && Format != BufferFormat::Auto && rhs.Format != BufferFormat::Auto)
 		bIsSame &= Format == rhs.Format;
+	
+	bIsSame &= bIsSRGB == rhs.bIsSRGB;
 
 	return bIsSame;
 }
