Fix crash on exit when device native tasks might still be running

diff --git a/Engine/Plugins/TextureGraph/Source/TextureGraphEngine/Device/Device.cpp b/Engine/Plugins/TextureGraph/Source/TextureGraphEngine/Device/Device.cpp
--- a/Engine/Plugins/TextureGraph/Source/TextureGraphEngine/Device/Device.cpp
+++ b/Engine/Plugins/TextureGraph/Source/TextureGraphEngine/Device/Device.cpp
@@ -116,17 +116,49 @@ bool Device::CanFree() const
 
 void Device::Free()
 {
-	check(NativePreExecWait.empty());
-	check(NativePreExec.empty());
-	check(NativeExec.empty());
-	check(NativePostExec.empty());
-
-	FScopeLock Lock(Cache.GetMutex());
-	Cache.GetCache().Empty();
-	GarbageCollect();
+	const bool bEngineShutdown = TextureGraphEngine::IsDestroying();
 
+	// Stop threads first (wakes + joins)
 	Terminate();
 
+	// In normal operation enforce emptiness; during shutdown, drain instead of crashing
+	if (!bEngineShutdown)
+	{
+		check(NativePreExecWait.empty());
+		check(NativePreExec.empty());
+		check(NativeExec.empty());
+		check(NativePostExec.empty());
+	}
+	else
+	{
+		auto Drain = [this](auto& Q)
+			{
+				auto v = Q.to_vector_and_clear();
+				for (auto& w : v)
+					if (auto t = w.lock())
+						MarkTaskFinished(t, /*bRemoveFromOwnerQueue=*/true);
+				return v.size();
+			};
+
+		size_t Drained = 0;
+		Drained += Drain(NativePreExecWait);
+		Drained += Drain(NativePreExec);
+		Drained += Drain(NativeExec);
+		Drained += Drain(NativePostExec);
+		Drained += Drain(NonAsyncPostExecQueue);
+		NativeTasks.clear();
+
+		UE_LOG(LogDevice, Display,
+				TEXT("[%s] Shutdown drained %llu native tasks (expected on fast exit)"),
+				*Name(), (unsigned long long)Drained);
+	}
+
+	{
+		FScopeLock Lock(Cache.GetMutex());
+		Cache.GetCache().Empty();
+	}
+	GarbageCollect();
+
 	BufferFactory = nullptr;
 	E_ResourceUpdate().Broadcast();
 }
@@ -134,29 +166,26 @@ void Device::Free()
 void Device::Terminate()
 {
 	ShouldTerminate = true;
-	if (PreExecThread)
-	{
-		/// Add a dummy object just to release the conditional variable
-		NativePreExecWait.add(DeviceNativeTaskPtrW());
-		PreExecThread->join();
-		PreExecThread = nullptr;
-	}
 
-	if (ExecThread)
-	{
-		/// Add a dummy object just to release the conditional variable
-		NativeExec.add(DeviceNativeTaskPtrW());
-		ExecThread->join();
-		ExecThread = nullptr;
-	}
+	// If we were waiting on a dependency, nudge it to cancel quickly.
+	if (CurrentWaitTask)
+		CurrentWaitTask->Terminate();
 
-	if (PostExecThread)
-	{
-		/// Add a dummy object just to release the conditional variable
-		NativePostExec.add(DeviceNativeTaskPtrW());
-		PostExecThread->join();
-		PostExecThread = nullptr;
-	}
+	// Wake all waiters so join() won’t block on condition variables.
+	NativePreExecWait.add(DeviceNativeTaskPtrW());
+	NativeExec.add(DeviceNativeTaskPtrW());
+	NativePostExec.add(DeviceNativeTaskPtrW());
+
+	auto JoinAndReset = [](std::unique_ptr<std::thread>& T)
+		{
+			if (T && T->joinable())
+				T->join();
+			T.reset();
+		};
+
+	JoinAndReset(PreExecThread);
+	JoinAndReset(ExecThread);
+	JoinAndReset(PostExecThread);
 }
 
 DeviceBufferRef Device::CreateCompatible_Copy(DeviceBufferRef Source)
@@ -856,10 +885,19 @@ void Device::PreExecThreadFunc()
 				UE_LOG(LogDevice, VeryVerbose, TEXT("PreExec - %s"), *Task->GetDebugName());
 
 				CurrentWaitTask = Task;
-
+				ON_SCOPE_EXIT
+				{
+					CurrentWaitTask = nullptr;   // always cleared, even on continue/throw
+				};
+				
 				Task->Wait();
-
-				CurrentWaitTask = nullptr;
+				
+				if (Task->IsTerminated())
+				{
+					// Drop gracefully and finalize bookkeeping
+					MarkTaskFinished(Task, /*bRemoveFromOwnerQueue=*/true);
+					continue;
+				}
 
 				NativePreExec.add(Task);
 			}
@@ -988,6 +1026,13 @@ void Device::UpdateNative()
 
 void Device::AddNativeTask(DeviceNativeTaskPtr Task)
 {
+	if (ShouldTerminate || TextureGraphEngine::IsDestroying())
+	{
+		// Drop the task gracefully; don’t enqueue.
+		MarkTaskFinished(Task, /*bRemoveFromOwnerQueue=*/true);
+		return;
+	}
+	
 	// check(IsInGameThread());
 	check(!Task->IsDone());
 
diff --git a/Engine/Plugins/TextureGraph/Source/TextureGraphEngine/Device/DeviceNativeTask.cpp b/Engine/Plugins/TextureGraph/Source/TextureGraphEngine/Device/DeviceNativeTask.cpp
--- a/Engine/Plugins/TextureGraph/Source/TextureGraphEngine/Device/DeviceNativeTask.cpp
+++ b/Engine/Plugins/TextureGraph/Source/TextureGraphEngine/Device/DeviceNativeTask.cpp
@@ -221,8 +221,12 @@ void DeviceNativeTask::Wait()
 
 	while (!WaitFor(1))
 	{
+		if (bTerminate.load(std::memory_order_acquire))
+		{
+			// Cancelled by Device::Terminate(); exit immediately.
+			return;
+		}
 		std::this_thread::sleep_for(std::chrono::milliseconds(1));
-		check(!bTerminate);
 	}
 }
 
diff --git a/Engine/Plugins/TextureGraph/Source/TextureGraphEngine/Device/DeviceNativeTask.h b/Engine/Plugins/TextureGraph/Source/TextureGraphEngine/Device/DeviceNativeTask.h
--- a/Engine/Plugins/TextureGraph/Source/TextureGraphEngine/Device/DeviceNativeTask.h
+++ b/Engine/Plugins/TextureGraph/Source/TextureGraphEngine/Device/DeviceNativeTask.h
@@ -118,6 +118,7 @@ public:
 	FORCEINLINE bool				IsDone() const { return bIsDone || IsCulled(); }
 	FORCEINLINE void				SetIsAsync(bool bIsAsyncIn) { bIsAsync = bIsAsyncIn; }
 	FORCEINLINE const DeviceNativeTaskPtrVec& GetPrev() const { return Prev; }
+	FORCEINLINE bool				IsTerminated() const	{ return bTerminate.load(std::memory_order_acquire); }
 };
 
 //////////////////////////////////////////////////////////////////////////
