TG Export crash fix

diff --git a/Engine/Plugins/TextureGraph/Source/TextureGraph/Private/Blueprint/TG_AsyncExportTask.cpp b/Engine/Plugins/TextureGraph/Source/TextureGraph/Private/Blueprint/TG_AsyncExportTask.cpp
--- a/Engine/Plugins/TextureGraph/Source/TextureGraph/Private/Blueprint/TG_AsyncExportTask.cpp
+++ b/Engine/Plugins/TextureGraph/Source/TextureGraph/Private/Blueprint/TG_AsyncExportTask.cpp
@@ -50,7 +50,7 @@ const TArray<UTextureRenderTarget2D*>& UTG_AsyncExportTask::ActivateBlocking(Job
 	Super::ActivateBlocking(Batch);
 
 	bool bExportDone = false;
-	TextureExporter::ExportAsUAsset(TextureGraphPtr, TargetExportSettings, "").then([&bExportDone]()
+	TextureExporter::ExportAsUAsset(TextureGraphPtr, MakeShared<FExportSettings>(TargetExportSettings), "").then([&bExportDone]()
 		{
 			bExportDone = true;
 		});
diff --git a/Engine/Plugins/TextureGraph/Source/TextureGraph/Private/TG_HelperFunctions.cpp b/Engine/Plugins/TextureGraph/Source/TextureGraph/Private/TG_HelperFunctions.cpp
--- a/Engine/Plugins/TextureGraph/Source/TextureGraph/Private/TG_HelperFunctions.cpp
+++ b/Engine/Plugins/TextureGraph/Source/TextureGraph/Private/TG_HelperFunctions.cpp
@@ -14,7 +14,6 @@
 #include "Logging/MessageLog.h"
 #include "Misc/FileHelper.h"
 #include "TextureGraph.h"
-
 TArray<BlobPtr> FTG_HelperFunctions::GetTexturedOutputs(const UTG_Node* Node, FTG_EvaluationContext* TextureConversionContext /*= nullptr*/)
 
 {
@@ -197,33 +196,42 @@ JobBatchPtr FTG_HelperFunctions::InitExportBatch(UTextureGraphBase* InTextureGra
 	return Batch;
 }
 
-AsyncBool FTG_HelperFunctions::ExportAsync(UTextureGraphBase* InTextureGraph, FString ExportPath, FString AssetName, FExportSettings& TargetExportSettings, bool OverrideExportPath, bool OverwriteTextures /*= true*/, bool ExportAllOutputs /*= false*/,bool bSave /*= true*/)
+AsyncInt FTG_HelperFunctions::ExportAsync(UTextureGraphBase* InTextureGraph, FString ExportPath, FString AssetName, FExportSettings& TargetExportSettings, bool OverrideExportPath, bool OverwriteTextures /*= true*/, bool ExportAllOutputs /*= false*/,bool bSave /*= true*/)
 {
-	JobBatchPtr Batch = InitExportBatch(InTextureGraph, ExportPath, AssetName, TargetExportSettings, OverrideExportPath, OverwriteTextures, ExportAllOutputs, bSave);
+	// Make a copy of the settings so that we don't have any threading issues or issues if TG window is closed
+	// while the async operation is in progress
+	TSharedRef<FExportSettings> SessionSettings = MakeShared<FExportSettings>(TargetExportSettings); 
+	
+	JobBatchPtr Batch = InitExportBatch(InTextureGraph, ExportPath, AssetName, *SessionSettings, OverrideExportPath, OverwriteTextures, ExportAllOutputs, bSave);
 
 	if (!Batch)
-		return cti::make_ready_continuable<bool>(false);
-
+	{
+		// Nothing to export → notify and resolve to 0.
+		Util::OnGameThread([Sess = SessionSettings]() { Sess->OnDone.ExecuteIfBound(); });
+		return cti::make_ready_continuable<int32>(0);
+	}
+	
 	return RenderAsync(InTextureGraph, Batch)
-		.then([InTextureGraph, &TargetExportSettings, ExportPath](auto result) 
+		.then([InTextureGraph, SessionSettings, ExportPath](auto result) mutable
 		{
-			return TextureExporter::ExportAsUAsset(InTextureGraph, TargetExportSettings, ExportPath);
+			return TextureExporter::ExportAsUAsset(InTextureGraph, SessionSettings, ExportPath);
 		})
-		.then([InTextureGraph, &TargetExportSettings]()
+		.then([InTextureGraph, SessionSettings](int32 NumExports) mutable
 		{
-			InTextureGraph->InvalidateAll();
-
+			if (IsValid(InTextureGraph))
+			{
+				InTextureGraph->InvalidateAll();
+			}
+				
 			// Add analytics tag
 			if (FEngineAnalytics::IsAvailable())
 			{
 				TArray<FAnalyticsEventAttribute> Attributes;
-				Attributes.Add(FAnalyticsEventAttribute(TEXT("NumExports"),  TargetExportSettings.MapsExported));
-					
+				Attributes.Add(FAnalyticsEventAttribute(TEXT("NumExports"),  NumExports));
 				// Send Analytics event 
 				FEngineAnalytics::GetProvider().RecordEvent(TEXT("Editor.Usage.TextureGraph.Export"), Attributes);
 			}
-
-			return true;
+			return NumExports;
 		});
 }
 
@@ -363,4 +371,4 @@ AsyncBool FTG_HelperFunctions::RenderAsync(UTextureGraphBase* InTextureGraph, Jo
 				});
 		}
 	}); 
-}
\ No newline at end of file
+}
diff --git a/Engine/Plugins/TextureGraph/Source/TextureGraph/Public/TG_HelperFunctions.h b/Engine/Plugins/TextureGraph/Source/TextureGraph/Public/TG_HelperFunctions.h
--- a/Engine/Plugins/TextureGraph/Source/TextureGraph/Public/TG_HelperFunctions.h
+++ b/Engine/Plugins/TextureGraph/Source/TextureGraph/Public/TG_HelperFunctions.h
@@ -23,7 +23,7 @@ public:
 	static UE_API void EnsureOutputIsTexture(MixUpdateCyclePtr Cycle, UTG_Node* OutputNode);
 
 	static UE_API JobBatchPtr InitExportBatch(UTextureGraphBase* InTextureGraph, FString ExportPath, FString AssetName, FExportSettings& TargetExportSettings, bool OverrideExportPath, bool OverwriteTextures, bool ExportAllOutputs, bool bSave);
-	static UE_API AsyncBool ExportAsync(UTextureGraphBase* InTextureGraph, FString ExportPath, FString AssetName, FExportSettings& TargetExportSettings, bool OverrideExportPath, bool OverwriteTextures = true, bool ExportAllOutputs = false, bool bSave = true);
+	static UE_API AsyncInt ExportAsync(UTextureGraphBase* InTextureGraph, FString ExportPath, FString AssetName, FExportSettings& TargetExportSettings, bool OverrideExportPath, bool OverwriteTextures = true, bool ExportAllOutputs = false, bool bSave = true);
 
 	static UE_API JobBatchPtr InitRenderBatch(UTextureGraphBase* InTextureGraph, JobBatchPtr ExistingBatch = nullptr);
 	static UE_API AsyncBool RenderAsync(UTextureGraphBase* InTextureGraph, JobBatchPtr ExistingBatch = nullptr);
diff --git a/Engine/Plugins/TextureGraph/Source/TextureGraphEditor/Private/EdGraph/TG_EditorErrorReporter.cpp b/Engine/Plugins/TextureGraph/Source/TextureGraphEditor/Private/EdGraph/TG_EditorErrorReporter.cpp
--- a/Engine/Plugins/TextureGraph/Source/TextureGraphEditor/Private/EdGraph/TG_EditorErrorReporter.cpp
+++ b/Engine/Plugins/TextureGraph/Source/TextureGraphEditor/Private/EdGraph/TG_EditorErrorReporter.cpp
@@ -21,7 +21,7 @@ FTextureGraphErrorReport FTG_EditorErrorReporter::Report(int32 ErrorId, const FS
 	{
 		// we update the error message on the EdNode 
 		auto FormattedMsg = Report.GetFormattedMessage();
-		if (Editor)
+		if (Editor && IsValid(Editor->TG_EdGraph))
 		{
 			auto EdNode = Editor->TG_EdGraph->GetViewModelNode(NodeObj->GetId());
 
@@ -75,7 +75,7 @@ void FTG_EditorErrorReporter::Clear()
 			// Expect the reference object to be a TG_Node
 			const UTG_Node* NodeObj = Cast<UTG_Node>(Error.ReferenceObj);
 
-			if (NodeObj && Editor)
+			if (NodeObj && Editor && IsValid(Editor->TG_EdGraph))
 			{
 				auto EdNode = Editor->TG_EdGraph->GetViewModelNode(NodeObj->GetId());
 
diff --git a/Engine/Plugins/TextureGraph/Source/TextureGraphEditor/Private/STG_NodePreview.cpp b/Engine/Plugins/TextureGraph/Source/TextureGraphEditor/Private/STG_NodePreview.cpp
--- a/Engine/Plugins/TextureGraph/Source/TextureGraphEditor/Private/STG_NodePreview.cpp
+++ b/Engine/Plugins/TextureGraph/Source/TextureGraphEditor/Private/STG_NodePreview.cpp
@@ -17,7 +17,6 @@
 #include "Engine/Texture2D.h"
 
 #define LOCTEXT_NAMESPACE "STG_NodePreview"
-
 UE::ImageWidgets::IImageViewer::FImageInfo FNodeViewer::GetCurrentImageInfo() const
 {
 	if (NodeTexture)
@@ -389,7 +388,8 @@ void STG_NodePreviewWidget::Update() const
 		{
 			Blob = Variant.GetTexture().RasterBlob;
 			ImageSize = FIntPoint {(int32)Blob->GetWidth(), (int32)Blob->GetHeight()};
-
+			TWeakPtr<FNodeViewer> WeakViewer = NodeViewer;
+			TWeakPtr<UE::ImageWidgets::SImageViewport> WeakViewport = Viewport;
 			if (Blob->IsTiled())
 			{
 				Blob->OnFinalise()
@@ -398,19 +398,34 @@ void STG_NodePreviewWidget::Update() const
 							const TiledBlobPtr BlobTiled = std::static_pointer_cast<TiledBlob>(Blob);
 							return BlobTiled->CombineTiles(false, false);
 						})
-					.then([this, Blob]()
+					.then([WeakViewer, WeakViewport, Blob]()
 						{
-							FIntPoint ImageSize = FIntPoint {(int32)Blob->GetWidth(), (int32)Blob->GetHeight()};
-							NodeViewer->SetTexture(Blob);
-							Viewport->ResetZoom(ImageSize);
+							Util::OnGameThread([WeakViewer, WeakViewport, Blob]()
+							{
+								const auto PinnedViewer   = WeakViewer.Pin();
+								const auto PinnedViewport = WeakViewport.Pin();
+								if (!PinnedViewer || !PinnedViewport)
+								{
+									return; // widget/editor closed; nothing to do
+								}
+								FIntPoint ImageSize = FIntPoint {(int32)Blob->GetWidth(), (int32)Blob->GetHeight()};
+								PinnedViewer->SetTexture(Blob);
+								PinnedViewport->ResetZoom(ImageSize);
+							});
 						});
 			}
 			else
 			{
-				Blob->OnFinalise().then([this, Blob]()
-					{
-						NodeViewer->SetTexture(Blob);
-					});
+				Blob->OnFinalise().then([WeakViewer, Blob]()
+				{
+					Util::OnGameThread([WeakViewer, Blob]()
+						{
+							if (const auto PinnedViewer   = WeakViewer.Pin())
+							{
+								PinnedViewer->SetTexture(Blob);
+							}
+						});
+				});
 			}
 		}
 		else if(Variant.IsColor())
@@ -599,4 +614,4 @@ void STG_NodePreviewWidget::ToggleLock()
 	}
 }
 
-#undef LOCTEXT_NAMESPACE
+#undef LOCTEXT_NAMESPACE
\ No newline at end of file
diff --git a/Engine/Plugins/TextureGraph/Source/TextureGraphEngine/Export/TextureExporter.cpp b/Engine/Plugins/TextureGraph/Source/TextureGraphEngine/Export/TextureExporter.cpp
--- a/Engine/Plugins/TextureGraph/Source/TextureGraphEngine/Export/TextureExporter.cpp
+++ b/Engine/Plugins/TextureGraph/Source/TextureGraphEngine/Export/TextureExporter.cpp
@@ -86,14 +86,14 @@ bool TextureExporter::IsFilePathValid(const FName InFileName, const FName InFold
 
 	return OutErrors.IsEmpty();
 }
-AsyncBool TextureExporter::ExportAllMapsAsUAsset(UMixInterface* MixObj, const FString& OutputFolder, const FExportSettings& ExportSettings)
+AsyncInt TextureExporter::ExportAllMapsAsUAsset(UMixInterface* MixObj, const FString& OutputFolder, TSharedRef<FExportSettings> SettingsRef)
 {
 	FInvalidationDetails Details;
 	Details.Mix = MixObj;
 	Details.All();
 	JobBatchPtr Batch = JobBatch::Create(Details);
 	MixUpdateCyclePtr Cycle = Batch->GetCycle();
-	const ExportMaps Maps = ExportSettings.ExportPreset;
+	const ExportMaps Maps = SettingsRef->ExportPreset;
 	for (const auto& ExportMap : Maps)
 	{
 		auto& MapSettings = ExportMap.second;
@@ -105,35 +105,37 @@ AsyncBool TextureExporter::ExportAllMapsAsUAsset(UMixInterface* MixObj, const FS
 		T_ExportAsUAsset::ExportAsUAsset(Cycle, 0, ExportMap.second, Path);
 	}
 	TextureGraphEngine::GetScheduler()->AddBatch(Batch);
-	return cti::make_continuable<bool>([=](auto&& Promise) mutable
+	return cti::make_continuable<int32>([Batch, SettingsRef](auto&& Promise) mutable
 	{
-		Batch->OnDone([=, FWD_PROMISE(Promise)](JobBatch*) mutable
+		Batch->OnDone([SettingsRef, FWD_PROMISE(Promise)](JobBatch*) mutable
 		{
-			Promise.set_value(true);
+			Promise.set_value(SettingsRef->MapsExported);
 		});
 	});
 }
 AsyncInt TextureExporter::ExportRawAsUAsset(RawBufferPtr RawObj,const FExportMapSettings& Setting, const FString& CompletePath, const FName& TextureName)
 {
-	return cti::make_continuable<int32>([RawObj,Setting, CompletePath, TextureName](auto Promise) mutable {
-		FText PathError;
-		FPaths::ValidatePath(CompletePath, &PathError);
-		if (CompletePath.IsEmpty())
-		{
-			UE_LOG(LogTexture, Warning, TEXT("Provide path is empty!"));
-			Promise.set_value(-1);
-		}
-		else if (!PathError.IsEmpty())
-		{
-			UE_LOG(LogTexture, Warning, TEXT("Invalid file path provided: %s"), *(PathError.ToString()));
-			Promise.set_value(-1);
-		}
+	FText PathError;
+	FPaths::ValidatePath(CompletePath, &PathError);
+	if (CompletePath.IsEmpty())
+	{
+		UE_LOG(LogTexture, Warning, TEXT("Provide path is empty!"));
+		return cti::make_ready_continuable(-1);
+	}
+	else if (!PathError.IsEmpty())
+	{
+		UE_LOG(LogTexture, Warning, TEXT("Invalid file path provided: %s"), *(PathError.ToString()));
+		return cti::make_ready_continuable(-1);
+	}
 		
-		if (!RawObj || RawObj->GetData() == nullptr)
-		{
-			UE_LOG(LogTexture, Warning, TEXT("Can't export with invalid RawBuffer: %s"), *(RawObj->GetName()));
-			Promise.set_value(-1);
-		}
+	if (!RawObj || RawObj->GetData() == nullptr)
+	{
+		UE_LOG(LogTexture, Warning, TEXT("Can't export with invalid RawBuffer: %s"), (RawObj != nullptr) ? *RawObj->GetName() : TEXT("nullptr"));
+		return cti::make_ready_continuable(-1);
+	}
+		
+	return cti::make_continuable<int32>([RawObj,Setting, CompletePath, TextureName](auto Promise) mutable
+	{
 		
 		const EPixelFormat PixelFormat = BufferDescriptor::BufferPixelFormat(RawObj->GetDescriptor().Format, RawObj->GetDescriptor().ItemsPerPoint);
 		const ETextureSourceFormat SourceFormat = TextureHelper::GetTextureSourceFormat(RawObj->GetDescriptor().Format, RawObj->GetDescriptor().ItemsPerPoint);
@@ -276,7 +278,7 @@ AsyncInt TextureExporter::ExportRawAsUAsset(RawBufferPtr RawObj,const FExportMap
 #endif
 	});
 }
-AsyncBool TextureExporter::ExportAsUAsset(UMixInterface* MixObj, const FExportSettings& ExportSettings, FString CurrentPath)
+AsyncInt TextureExporter::ExportAsUAsset(UMixInterface* MixObj, TSharedRef<FExportSettings> ExportSettings, FString CurrentPath)
 {
 	return ExportAllMapsAsUAsset(MixObj, CurrentPath, ExportSettings);
 }
@@ -296,7 +298,7 @@ FExportMapSettings TextureExporter::GetExportSettingsForTarget(FExportSettings&
 	TargetMap.OnDone.BindRaw(&ExportSettings, &FExportSettings::OnMapExportDone);
 	return TargetMap;
 }
-void TextureExporter::ExportAsUAsset(UMixInterface* MixObj, FName AssetName, FString CurrentPath, const FExportSettings& Setting)
+void TextureExporter::ExportAsUAsset(UMixInterface* MixObj, FName AssetName, FString CurrentPath, TSharedRef<FExportSettings> Setting)
 {
 	ExportAllMapsAsUAsset(MixObj, CurrentPath, Setting);
 }
\ No newline at end of file
diff --git a/Engine/Plugins/TextureGraph/Source/TextureGraphEngine/Export/TextureExporter.h b/Engine/Plugins/TextureGraph/Source/TextureGraphEngine/Export/TextureExporter.h
--- a/Engine/Plugins/TextureGraph/Source/TextureGraphEngine/Export/TextureExporter.h
+++ b/Engine/Plugins/TextureGraph/Source/TextureGraphEngine/Export/TextureExporter.h
@@ -94,12 +94,12 @@ public:
 
 struct TextureExporter
 {
-	static TEXTUREGRAPHENGINE_API FExportMapSettings		GetExportSettingsForTarget(FExportSettings& ExportSettings, TiledBlobPtr BlobObj, FName Name);
-	static TEXTUREGRAPHENGINE_API AsyncBool				ExportAllMapsAsUAsset(UMixInterface* MixObj, const FString& OutputFolder, const FExportSettings& ExportSettings);
-	static TEXTUREGRAPHENGINE_API AsyncBool				ExportAsUAsset(UMixInterface* MixObj, const FExportSettings& ExportSettings, FString C);
-	static TEXTUREGRAPHENGINE_API void						ExportAsUAsset(UMixInterface* MixObj, FName AssetName, FString CurrentPath, const FExportSettings& Setting);
-	static TEXTUREGRAPHENGINE_API AsyncInt					ExportRawAsUAsset(RawBufferPtr RawObj,const FExportMapSettings& Setting, const FString& CompletePath, const FName& TextureName);
-	static TEXTUREGRAPHENGINE_API bool						IsPackageNameValid(FString Path, FString AssetName);
+	static TEXTUREGRAPHENGINE_API FExportMapSettings	GetExportSettingsForTarget(FExportSettings& ExportSettings, TiledBlobPtr BlobObj, FName Name);
+	static TEXTUREGRAPHENGINE_API AsyncInt				ExportAllMapsAsUAsset(UMixInterface* MixObj, const FString& OutputFolder, TSharedRef<FExportSettings> SettingsRef);
+	static TEXTUREGRAPHENGINE_API AsyncInt				ExportAsUAsset(UMixInterface* MixObj, TSharedRef<FExportSettings> SettingsRef, FString C);
+	static TEXTUREGRAPHENGINE_API void					ExportAsUAsset(UMixInterface* MixObj, FName AssetName, FString CurrentPath, TSharedRef<FExportSettings> SettingsRef);
+	static TEXTUREGRAPHENGINE_API AsyncInt				ExportRawAsUAsset(RawBufferPtr RawObj,const FExportMapSettings& Setting, const FString& CompletePath, const FName& TextureName);
+	static TEXTUREGRAPHENGINE_API bool					IsPackageNameValid(FString Path, FString AssetName);
 	static TEXTUREGRAPHENGINE_API bool 					IsFileNameValid(FName FileName, FText& Reason);
 	static TEXTUREGRAPHENGINE_API bool 					IsFolderPathValid(FString FolderPath, FText& Reason);
 	static TEXTUREGRAPHENGINE_API bool 					IsFilePathValid(const FName InFileName, const FName InFolderPath, FString& OutErrors);
