[TG] Refine the VirtualTexture warmup meachanism used by the TG_Expression_MaterialBase / FXMaterial_DrawMaterial:
- On each warm loop if a Virtual Texture:
- Update virtual texture streaming to trigger release of upload buffers.
- At the end of each warm loop:
- Increment frame number so that FVirtualTextureUploadCache releases upload buffers.
- Tick RDG pool, so that render target memory can be recycled.
- EndFrame and flush to give GPU allocator opportunity to release memory.

diff --git a/Engine/Plugins/TextureGraph/Source/TextureGraphEngine/FxMat/FxMaterial_DrawMaterial.cpp b/Engine/Plugins/TextureGraph/Source/TextureGraphEngine/FxMat/FxMaterial_DrawMaterial.cpp
--- a/Engine/Plugins/TextureGraph/Source/TextureGraphEngine/FxMat/FxMaterial_DrawMaterial.cpp
+++ b/Engine/Plugins/TextureGraph/Source/TextureGraphEngine/FxMat/FxMaterial_DrawMaterial.cpp
@@ -1,8 +1,10 @@
 // Copyright Epic Games, Inc. All Rights Reserved
 #include "FxMaterial_DrawMaterial.h"
 
+#include "ContentStreaming.h"
 #include "EngineModule.h"
 #include "SceneView.h"
+#include "VT/VirtualTextureChunkManager.h"
 #include "VT/VirtualTextureFeedbackResource.h"
 
 IMPLEMENT_MATERIAL_SHADER_TYPE(, FTextureGraphMaterialShaderVS, TEXT("/Plugin/TextureGraph/TextureGraphMaterialShader.usf"), TEXT("TextureGraphMaterialShaderVS"), SF_Vertex);
@@ -113,8 +115,21 @@ void FxMaterial_DrawMaterialBase::VirtualTextureFeedbackEnd(FRDGBuilder& GraphBu
 void FxMaterial_DrawMaterialBase::VirtualTextureFeedbackSync(FRHICommandListImmediate& RHICmdList, ERHIFeatureLevel::Type FeatureLevel) const
 {
 	// Need a flush before we can read back the virtual texture feedback buffer.
-	RHICmdList.ImmediateFlush(EImmediateFlushType::FlushRHIThread);
 	GetRendererModule().SyncVirtualTextureUpdates(RHICmdList, FeatureLevel);
+	// Update virtual texture streaming to trigger release of upload buffers.
+	// todo: Remove the need to access private implementation detail requiring private header access (enabled in our build.cs file).
+	IStreamingManager::Get().GetVirtualTextureStreamingManager().UpdateResourceStreaming(0.f);
+}
+
+void FxMaterial_DrawMaterialBase::EndFrame(FRHICommandListImmediate& RHICmdList)
+{
+	// Increment frame number so that FVirtualTextureUploadCache releases upload buffers.
+	GFrameNumberRenderThread++;
+	// Tick RDG pool, so that render target memory can be recycled.
+	FRDGBuilder::TickPoolElements();
+	// EndFrame and flush to give GPU allocator opportunity to release memory.
+	RHICmdList.EndFrame();
+	RHICmdList.ImmediateFlush(EImmediateFlushType::FlushRHIThreadFlushResources);
 }
 
 bool FxMaterial_DrawMaterialBase::ValidateMaterialShaderMap(UMaterial * InMaterial, FShaderType* InPixelShaderType)
diff --git a/Engine/Plugins/TextureGraph/Source/TextureGraphEngine/FxMat/FxMaterial_DrawMaterial.h b/Engine/Plugins/TextureGraph/Source/TextureGraphEngine/FxMat/FxMaterial_DrawMaterial.h
--- a/Engine/Plugins/TextureGraph/Source/TextureGraphEngine/FxMat/FxMaterial_DrawMaterial.h
+++ b/Engine/Plugins/TextureGraph/Source/TextureGraphEngine/FxMat/FxMaterial_DrawMaterial.h
@@ -111,6 +111,7 @@ public:
 	void VirtualTextureFeedbackBegin(FRDGBuilder& GraphBuilder, FIntPoint TargetSize, ERHIFeatureLevel::Type FeatureLevel) const;
 	void VirtualTextureFeedbackEnd(FRDGBuilder& GraphBuilder) const;
 	void VirtualTextureFeedbackSync(FRHICommandListImmediate& RHICmdList, ERHIFeatureLevel::Type FeatureLevel) const;
+	void EndFrame(FRHICommandListImmediate& RHICmdList);
 
 	static bool ValidateMaterialShaderMap(UMaterial* InMaterial, FShaderType* InPixelShaderType);
 };
@@ -265,6 +266,8 @@ public:
 			{
 				VirtualTextureFeedbackSync(RHI, FeatureLevel);
 			}
+
+			EndFrame(RHI);
 		}
 	}
 
diff --git a/Engine/Plugins/TextureGraph/Source/TextureGraphEngine/TextureGraphEngine.Build.cs b/Engine/Plugins/TextureGraph/Source/TextureGraphEngine/TextureGraphEngine.Build.cs
--- a/Engine/Plugins/TextureGraph/Source/TextureGraphEngine/TextureGraphEngine.Build.cs
+++ b/Engine/Plugins/TextureGraph/Source/TextureGraphEngine/TextureGraphEngine.Build.cs
@@ -45,6 +45,12 @@ public class TextureGraphEngine : ModuleRules
 
 		//PublicDefinitions.Add("WITH_MALLOC_STOMP=1");
 
+		PrivateIncludePaths.AddRange(
+			new string[] {
+				Path.Combine(ModuleDirectory, "../../../../Source/Runtime/Engine/Private")
+			}
+		);
+
 		PrivateDependencyModuleNames.AddRange(new string[] 
 		{ 
 			"Slate", 
diff --git a/Engine/Source/Runtime/Engine/Private/VT/VirtualTextureChunkManager.h b/Engine/Source/Runtime/Engine/Private/VT/VirtualTextureChunkManager.h
--- a/Engine/Source/Runtime/Engine/Private/VT/VirtualTextureChunkManager.h
+++ b/Engine/Source/Runtime/Engine/Private/VT/VirtualTextureChunkManager.h
@@ -32,7 +32,7 @@ public:
 	virtual ~FVirtualTextureChunkStreamingManager();
 
 	// IStreamingManager interface
-	virtual void UpdateResourceStreaming(float DeltaTime, bool bProcessEverything = false) override;
+	ENGINE_API void UpdateResourceStreaming(float DeltaTime, bool bProcessEverything = false) override;
 	virtual int32 BlockTillAllRequestsFinished(float TimeLimit = 0.0f, bool bLogResults = false) override;
 	virtual void CancelForcedResources() override;
 	virtual void NotifyLevelChange() override {}
