TG: Use render index instead of global counter when setting up virtual texture uniforms for material rendering.
Fixes issue with inconsistent virtual texture feedback.
#rnx

diff --git a/Engine/Plugins/TextureGraph/Source/TextureGraphEngine/FxMat/FxMaterial_DrawMaterial.cpp b/Engine/Plugins/TextureGraph/Source/TextureGraphEngine/FxMat/FxMaterial_DrawMaterial.cpp
--- a/Engine/Plugins/TextureGraph/Source/TextureGraphEngine/FxMat/FxMaterial_DrawMaterial.cpp
+++ b/Engine/Plugins/TextureGraph/Source/TextureGraphEngine/FxMat/FxMaterial_DrawMaterial.cpp
@@ -26,7 +26,7 @@ static FAutoConsoleVariableRef CVarVirtualTextureFeedbackFactor(
 	TEXT("Texel size of tile used to generate one item of virtual texture feedback."),
 	ECVF_RenderThreadSafe);
 
-FSceneView* FxMaterial_DrawMaterialBase::CreateSceneView(FSceneViewFamilyContext& ViewFamily, UTextureRenderTarget2D* RenderTarget, const FIntPoint& TargetSizeXY)
+FSceneView* FxMaterial_DrawMaterialBase::CreateSceneView(FSceneViewFamilyContext& ViewFamily, UTextureRenderTarget2D* RenderTarget, const FIntPoint& TargetSizeXY, int32 RenderIndex)
 {
 	FIntRect ViewRect(FIntPoint(0, 0), TargetSizeXY);
 
@@ -75,10 +75,9 @@ FSceneView* FxMaterial_DrawMaterialBase::CreateSceneView(FSceneViewFamilyContext
 		ViewUniformShaderParameters.MobilePreviewMode = 0.0f;
 
 		// Setup virtual texture feedback.
-		// Use a local static frame counter to vary feedback location independent of global render frame count.
-		static uint32 FrameIndex = 0;
+		// Use RenderIndex as the frame counter to vary feedback location independent of global render frame count.
 		VirtualTexture::FFeedbackShaderParams VirtualTextureFeedbackShaderParams;
-		VirtualTexture::GetFeedbackShaderParams(FrameIndex++, GVirtualTextureFeedbackFactor, VirtualTextureFeedbackShaderParams);
+		VirtualTexture::GetFeedbackShaderParams(RenderIndex, GVirtualTextureFeedbackFactor, VirtualTextureFeedbackShaderParams);
 		VirtualTexture::UpdateViewUniformShaderParameters(VirtualTextureFeedbackShaderParams, ViewUniformShaderParameters);
 
 		{
diff --git a/Engine/Plugins/TextureGraph/Source/TextureGraphEngine/FxMat/FxMaterial_DrawMaterial.h b/Engine/Plugins/TextureGraph/Source/TextureGraphEngine/FxMat/FxMaterial_DrawMaterial.h
--- a/Engine/Plugins/TextureGraph/Source/TextureGraphEngine/FxMat/FxMaterial_DrawMaterial.h
+++ b/Engine/Plugins/TextureGraph/Source/TextureGraphEngine/FxMat/FxMaterial_DrawMaterial.h
@@ -104,7 +104,7 @@ public:
 	{
 	}
 
-	FSceneView* CreateSceneView(FSceneViewFamilyContext& ViewFamily, UTextureRenderTarget2D* RenderTarget, const FIntPoint& TargetSizeXY);
+	FSceneView* CreateSceneView(FSceneViewFamilyContext& ViewFamily, UTextureRenderTarget2D* RenderTarget, const FIntPoint& TargetSizeXY, int32 RenderIndex);
 
 	bool HasVirtualTextureFeedback(FMaterial const& RenderMaterial, EShaderPlatform ShaderPlatform) const;
 	int32 GetDefaultVirtualTextureWarmupFrames() const;
@@ -230,7 +230,7 @@ public:
 				FSceneViewFamily::ConstructionValues(RenderTarget->GetRenderTargetResource(), nullptr, ESFIM_Game)
 				.SetTime(FGameTime())
 				.SetRealtimeUpdate(true));
-			FSceneView* SceneView = CreateSceneView(SceneViewFamilyContext, RenderTarget, Target->GetSizeXY());
+			FSceneView* SceneView = CreateSceneView(SceneViewFamilyContext, RenderTarget, Target->GetSizeXY(), RenderIndex);
 
 			auto* PassParameters = GraphBuilder.AllocParameters<FParameters>();
 			PassParameters->VS = VSHParams;
