Fix TextureGraph Instance skipping invalidation when changing parent graph

diff --git a/Engine/Plugins/TextureGraph/Source/TextureGraphEditor/Private/TG_Exporter.cpp b/Engine/Plugins/TextureGraph/Source/TextureGraphEditor/Private/TG_Exporter.cpp
--- a/Engine/Plugins/TextureGraph/Source/TextureGraphEditor/Private/TG_Exporter.cpp
+++ b/Engine/Plugins/TextureGraph/Source/TextureGraphEditor/Private/TG_Exporter.cpp
@@ -100,7 +100,7 @@ FTG_InstanceImpl::~FTG_InstanceImpl()
 
 }
 
-void FTG_InstanceImpl::Cleanup()
+void FTG_InstanceImpl::Cleanup(bool bFlushMixInvalidations /*= true*/)
 {
 	if (TextureGraphInstance && TextureGraphInstance->IsValidLowLevel())
 	{
@@ -111,9 +111,12 @@ void FTG_InstanceImpl::Cleanup()
 			Settings->GetViewportSettings().OnMaterialMappingChangedEvent.RemoveAll(this);
 			Settings->OnPreviewMeshChangedEvent.RemoveAll(this);
 		}
-		if (UTG_Graph* Graph = TextureGraphInstance->Graph())
+		if (bFlushMixInvalidations)
 		{
 			TextureGraphInstance->FlushInvalidations();
+		}
+		if (UTG_Graph* Graph = TextureGraphInstance->Graph())
+		{
 			Graph->OnGraphChangedDelegate.RemoveAll(this);
 		}
 		TextureGraphInstance->OnRenderDone.Unbind();
@@ -468,8 +471,8 @@ FReply FTG_InstanceImpl::OnExportClicked(EAppReturnType::Type ButtonID)
 
 void FTG_InstanceImpl::SetTextureGraphToExport(UTextureGraphInstance* InTextureGraph)
 {
-	// clear out previous handles
-	Cleanup();
+	// clear out previous handles but don't flush mix invalidations as that would skip invalidating the mix with the new parent
+	Cleanup(false);
 
 	TextureGraphInstance = InTextureGraph;
 	
diff --git a/Engine/Plugins/TextureGraph/Source/TextureGraphEditor/Public/TG_Exporter.h b/Engine/Plugins/TextureGraph/Source/TextureGraphEditor/Public/TG_Exporter.h
--- a/Engine/Plugins/TextureGraph/Source/TextureGraphEditor/Public/TG_Exporter.h
+++ b/Engine/Plugins/TextureGraph/Source/TextureGraphEditor/Public/TG_Exporter.h
@@ -79,7 +79,7 @@ public:
 													FTG_InstanceImpl();
 	virtual											~FTG_InstanceImpl() override;
 
-	void											Cleanup();
+	void											Cleanup(bool bFlushMixInvalidations = true);
 	/** Sets the Texture Graph to be exported in the Exporter */
 	void 											SetTextureGraphToExport(UTextureGraphInstance* InTextureGraph);
 	void 											OnGraphChanged(UTG_Graph* InGraph, UTG_Node* InNode, bool Tweaking);
