[TG]  optimization of the MaterialBase expression which consumes less memory in MetaHuman baking use case

diff --git a/Engine/Plugins/TextureGraph/Source/TextureGraph/Private/Expressions/TG_Expression_MaterialBase.cpp b/Engine/Plugins/TextureGraph/Source/TextureGraph/Private/Expressions/TG_Expression_MaterialBase.cpp
--- a/Engine/Plugins/TextureGraph/Source/TextureGraph/Private/Expressions/TG_Expression_MaterialBase.cpp
+++ b/Engine/Plugins/TextureGraph/Source/TextureGraph/Private/Expressions/TG_Expression_MaterialBase.cpp
@@ -187,27 +187,42 @@ void UTG_Expression_MaterialBase::Evaluate(FTG_EvaluationContext* InContext)
 	Output = TextureHelper::GetBlack();
 	if (GetMaterial())
 	{
+		UTG_Node* Node = GetParentNode();
+
 
 		FString AssetName = GetMaterial()->GetName();
 		UMaterialInstanceConstant* MIC = GetEvaluationMaterialInstanceConstant(InContext);
 
 		for (int32 i = 0; i < AvailableMaterialAttributeIds.Num(); ++i)
 		{
-			const auto RenderMaterial = std::make_shared<RenderMaterial_BP>(AssetName, MIC, GetNumVirtualTextureWarmupFrames());
-			TiledBlobPtr Result = TextureHelper::GetBlack();
-			Result = CreateRenderMaterialJob(InContext, RenderMaterial, Output.GetBufferDescriptor(), AvailableMaterialAttributeIds[i]);
-			Outputs[i] = Result;
-
-			if (GetRenderedAttributeId() == AvailableMaterialAttributeIds[i])
+			bool IsDefaultOutput = GetRenderedAttributeId() == AvailableMaterialAttributeIds[i];
+			if (IsDefaultOutput || Node->GetPin(MatAttributesOutputArgNames[i])->IsConnected())
 			{
-				Output = Result;
-			}
+				const auto RenderMaterial = std::make_shared<RenderMaterial_BP>(AssetName, MIC, GetNumVirtualTextureWarmupFrames());
+				TiledBlobPtr Result = TextureHelper::GetBlack();
+				Result = CreateRenderMaterialJob(InContext, RenderMaterial, Output.GetBufferDescriptor(), AvailableMaterialAttributeIds[i]);
+				Outputs[i] = Result;
+
+				if (IsDefaultOutput)
+				{
+					Output = Result;
+				}
 
-			// Reassign the generated result to the corresponding Output arg
-			FTG_Var* OutVar = InContext->Outputs.GetVar(MatAttributesOutputArgNames[i]);
-			if (OutVar)
+				// Reassign the generated result to the corresponding Output arg
+				FTG_Var* OutVar = InContext->Outputs.GetVar(MatAttributesOutputArgNames[i]);
+				if (OutVar)
+				{
+					OutVar->EditAs<FTG_Texture>() = Result;
+				}
+			}
+			else
 			{
-				OutVar->EditAs<FTG_Texture>() = Result;
+				// For non generated attributes, assign black
+				FTG_Var* OutVar = InContext->Outputs.GetVar(MatAttributesOutputArgNames[i]);
+				if (OutVar)
+				{
+					OutVar->EditAs<FTG_Texture>() = TextureHelper::GetBlack();
+				}
 			}
 		}
 	}
