Fix crash when exporting a new unsaved TextureGraph (by right click ->Export)

diff --git a/Engine/Plugins/TextureGraph/Source/TextureGraphEngine/Model/Mix/MixManager.cpp b/Engine/Plugins/TextureGraph/Source/TextureGraphEngine/Model/Mix/MixManager.cpp
--- a/Engine/Plugins/TextureGraph/Source/TextureGraphEngine/Model/Mix/MixManager.cpp
+++ b/Engine/Plugins/TextureGraph/Source/TextureGraphEngine/Model/Mix/MixManager.cpp
@@ -83,6 +83,9 @@ void MixManager::Update(float Delta)
 	{
 		MixInvalidateInfo& MixInfo = *Iter;
 
+		if (!MixInfo.MixObj->IsValidLowLevel())
+			continue;
+
 		UMixInterface* MixObj = MixInfo.MixObj;
 		auto RIter = std::find_if(BucketedMixInfos.begin(), BucketedMixInfos.end(),
 			[MixObj](const MixInvalidateInfo& Info)
