Fix crashes when exporting a specific TG

diff --git a/Engine/Plugins/TextureGraph/Source/TextureGraphEngine/Device/FX/Device_FX.cpp b/Engine/Plugins/TextureGraph/Source/TextureGraphEngine/Device/FX/Device_FX.cpp
--- a/Engine/Plugins/TextureGraph/Source/TextureGraphEngine/Device/FX/Device_FX.cpp
+++ b/Engine/Plugins/TextureGraph/Source/TextureGraphEngine/Device/FX/Device_FX.cpp
@@ -285,7 +285,19 @@ AsyncDeviceBufferRef Device_FX::FillTextureArray_Deferred(DeviceBufferRef Buffer
 
 AsyncDeviceBufferRef Device_FX::DrawTilesToBuffer_Deferred(DeviceBufferRef Buffer, const T_Tiles<DeviceBufferRef>& Tiles, const DrawTilesSettings* DrawSettings, FLinearColor* ClearColor)
 {
-	DeviceNativeTaskPtr Task = DeviceNativeTask_Lambda::Create(this, (int32)E_Priority::kSystem, TEXT("DrawTilesToBuffer"), [this, Buffer, Tiles, DrawSettings, ClearColor]() mutable -> int32
+	TOptional<DrawTilesSettings> DrawSettingsCopy;
+	if (DrawSettings != nullptr)
+	{
+		DrawSettingsCopy = *DrawSettings;
+	}
+
+	TOptional<FLinearColor> ClearColorCopy;
+	if (ClearColor != nullptr)
+	{
+		ClearColorCopy = *ClearColor;
+	}
+
+	DeviceNativeTaskPtr Task = DeviceNativeTask_Lambda::Create(this, (int32)E_Priority::kSystem, TEXT("DrawTilesToBuffer"), [this, Buffer, Tiles, DrawSettingsCopy, ClearColorCopy]() mutable -> int32
 		{
 			SCOPE_CYCLE_COUNTER(STAT_Device_FX_DrawTilesToBuffer);
 			check(IsInRenderingThread());
@@ -311,8 +323,8 @@ AsyncDeviceBufferRef Device_FX::DrawTilesToBuffer_Deferred(DeviceBufferRef Buffe
 			FRHIRenderPassInfo RenderPassInfo(RTResDest, ERenderTargetActions::Load_Store);
 			DevRHI.BeginRenderPass(RenderPassInfo, TEXT("CopyTexture"));
 
-			if (ClearColor)
-				DrawClearQuad(DevRHI, *ClearColor);
+			if (ClearColorCopy.IsSet())
+				DrawClearQuad(DevRHI, *ClearColorCopy);
 
 			{
 				FGraphicsPipelineStateInitializer GraphicsPSOInit;
@@ -328,11 +340,11 @@ AsyncDeviceBufferRef Device_FX::DrawTilesToBuffer_Deferred(DeviceBufferRef Buffe
 				GraphicsPSOInit.BoundShaderState.VertexDeclarationRHI = GFilterVertexDeclaration.VertexDeclarationRHI;
 
 				FRHISamplerState* PixelSampler = TStaticSamplerState<SF_Bilinear>::GetRHI();
-				bool bUseInputPositions = DrawSettings != nullptr ? true : false;
+				bool bUseInputPositions = DrawSettingsCopy.IsSet();
 
 				/// Verify that there are either no DrawSettings or they are setup correctly to match the input array dimensions
-				check(!DrawSettings || (DrawSettings->Position.Rows() == Tiles.Rows() && DrawSettings->Position.Cols() == Tiles.Cols()));
-				check(!DrawSettings || (DrawSettings->Size.Rows() == Tiles.Rows() && DrawSettings->Size.Cols() == Tiles.Cols()));
+				check(!DrawSettingsCopy.IsSet() || (DrawSettingsCopy->Position.Rows() == Tiles.Rows() && DrawSettingsCopy->Position.Cols() == Tiles.Cols()));
+				check(!DrawSettingsCopy.IsSet() || (DrawSettingsCopy->Size.Rows() == Tiles.Rows() && DrawSettingsCopy->Size.Cols() == Tiles.Cols()));
 
 				SetGraphicsPipelineState(DevRHI, GraphicsPSOInit, 0);
 
@@ -368,11 +380,11 @@ AsyncDeviceBufferRef Device_FX::DrawTilesToBuffer_Deferred(DeviceBufferRef Buffe
 						}
 						else
 						{
-							const FIntPoint& TilePos = DrawSettings->Position(RowId, ColId);
+							const FIntPoint& TilePos = DrawSettingsCopy->Position(RowId, ColId);
 							X1 = (uint32)TilePos.X;
 							Y1 = (uint32)TilePos.Y;
 
-							const FIntPoint& TileSize = DrawSettings->Size(RowId, ColId);
+							const FIntPoint& TileSize = DrawSettingsCopy->Size(RowId, ColId);
 							SrcTextureWidth = TileSize.X;
 							SrcTextureHeight = TileSize.Y;
 						}
