TG TextureExpression crashing on Linux

diff --git a/Engine/Plugins/TextureGraph/Source/TextureGraphEngine/Device/FX/Device_FX.cpp b/Engine/Plugins/TextureGraph/Source/TextureGraphEngine/Device/FX/Device_FX.cpp
--- a/Engine/Plugins/TextureGraph/Source/TextureGraphEngine/Device/FX/Device_FX.cpp
+++ b/Engine/Plugins/TextureGraph/Source/TextureGraphEngine/Device/FX/Device_FX.cpp
@@ -704,17 +704,16 @@ int32 Device_FX::InitRTResource(TexPtr TextureObj, UTextureRenderTarget* RT)
 	
 	FDeferredUpdateResource::UpdateResources(RHICmdList);
 
-	FTextureRenderTarget2DResource* RTRes = (FTextureRenderTarget2DResource*)RT->GetRenderTargetResource();
+	FTextureRenderTargetResource* RTRes = RT->GetRenderTargetResource();
 	check(RTRes);
 
-	UE_LOG(LogDevice, VeryVerbose, TEXT("New RT Array allocation: %s %llu [Ptr: 0x%x, Size: %dx%d]"), *TextureObj->GetDescriptor().Name,
+	UE_LOG(LogDevice, Verbose, TEXT("New RT allocation: %s %llu [Ptr: 0x%p, Size: %ux%u]"), *TextureObj->GetDescriptor().Name,
 		TextureObj->GetDescriptor().Format_HashValue(), RT, RTRes->GetSizeX(), RTRes->GetSizeY());
-
-	FTextureRHIRef RHITexture = RTRes->GetTextureRHI();
-	if(RHITexture)
-	{
-		RHICmdList.BindDebugLabelName(RHITexture, *TextureObj->GetDescriptor().Name);
-	}
+	
+	FTextureRHIRef RHITexture = RTRes->GetShaderResourceTexture();
+	check(RHITexture);
+	RHICmdList.BindDebugLabelName(RHITexture, *TextureObj->GetDescriptor().Name);
+	
 	return 0;
 }
 
diff --git a/Engine/Plugins/TextureGraph/Source/TextureGraphEngine/FxMat/RenderMaterial_BP.cpp b/Engine/Plugins/TextureGraph/Source/TextureGraphEngine/FxMat/RenderMaterial_BP.cpp
--- a/Engine/Plugins/TextureGraph/Source/TextureGraphEngine/FxMat/RenderMaterial_BP.cpp
+++ b/Engine/Plugins/TextureGraph/Source/TextureGraphEngine/FxMat/RenderMaterial_BP.cpp
@@ -279,7 +279,7 @@ void RenderMaterial_BP::BlitTo(FRHICommandListImmediate& RHI, UTextureRenderTarg
 		FTextureRenderTarget2DResource* RTRes = (FTextureRenderTarget2DResource*)RT->GetRenderTargetResource();
 		check(RTRes);
 
-		FTextureRHIRef TextureRHI = RTRes->GetTextureRHI();
+		FTextureRHIRef TextureRHI = RTRes->GetShaderResourceTexture();
 		check(TextureRHI);
 
 		TextureRHI->SetName(FName(*RT->GetName()));
diff --git a/Engine/Plugins/TextureGraph/Source/TextureGraphEngine/FxMat/RenderMaterial_FX.cpp b/Engine/Plugins/TextureGraph/Source/TextureGraphEngine/FxMat/RenderMaterial_FX.cpp
--- a/Engine/Plugins/TextureGraph/Source/TextureGraphEngine/FxMat/RenderMaterial_FX.cpp
+++ b/Engine/Plugins/TextureGraph/Source/TextureGraphEngine/FxMat/RenderMaterial_FX.cpp
@@ -83,7 +83,7 @@ void RenderMaterial_FX::BlitTo(FRHICommandListImmediate& RHI, UTextureRenderTarg
 	FTextureRenderTarget2DResource* RTRes = (FTextureRenderTarget2DResource*)RenderTarget->GetRenderTargetResource();
 	check(RTRes);
 
-	FTextureRHIRef TextureRHI = RTRes->GetTextureRHI();
+	FTextureRHIRef TextureRHI = RTRes->GetShaderResourceTexture();
 	check(TextureRHI);
 
 	TextureRHI->SetName(FName(*RenderTarget->GetName()));
