Fix for crash when cooking TG assets for MH

diff --git a/Engine/Plugins/TextureGraph/Source/TextureGraph/Private/Expressions/Input/TG_Expression_Material.cpp b/Engine/Plugins/TextureGraph/Source/TextureGraph/Private/Expressions/Input/TG_Expression_Material.cpp
--- a/Engine/Plugins/TextureGraph/Source/TextureGraph/Private/Expressions/Input/TG_Expression_Material.cpp
+++ b/Engine/Plugins/TextureGraph/Source/TextureGraph/Private/Expressions/Input/TG_Expression_Material.cpp
@@ -40,6 +40,11 @@ void UTG_Expression_Material::PostEditUndo()
 
 void UTG_Expression_Material::OnReferencedObjectPreSave(UObject* Object, FObjectPreSaveContext SaveContext)
 {
+	// early out if the parent node is invalid
+	UTG_Node* ParentNode = GetParentNode();
+	if (ParentNode == nullptr || !ParentNode->GetId().IsValid())
+		return;
+	
 	// every editor should check if your texture graph is dependent on the object being saved
 	UMaterialInterface* MaterialBeingSaved = Cast<UMaterialInterface>(Object);
 	UMaterialInterface* ReferencedMaterial = InputMaterial.GetMaterial();
