TG: Crash in Blueprint export node when running with a null texture graph

diff --git a/Engine/Plugins/TextureGraph/Source/TextureGraph/Private/Blueprint/TG_AsyncExportTask.cpp b/Engine/Plugins/TextureGraph/Source/TextureGraph/Private/Blueprint/TG_AsyncExportTask.cpp
--- a/Engine/Plugins/TextureGraph/Source/TextureGraph/Private/Blueprint/TG_AsyncExportTask.cpp
+++ b/Engine/Plugins/TextureGraph/Source/TextureGraph/Private/Blueprint/TG_AsyncExportTask.cpp
@@ -23,12 +23,12 @@ UTG_AsyncExportTask* UTG_AsyncExportTask::TG_AsyncExportTask( UTextureGraphBase*
 	Task->bSave = bSave;
 	Task->bExportAll = bExportAll;
 
-	Task->bBlobberCachingState = TextureGraphEngine::GetBlobber()->IsCacheEnabled();
-	if (bDisableCache)
-		TextureGraphEngine::GetBlobber()->SetEnableCache(false);
-
 	if (InTextureGraph != nullptr)
 	{
+		Task->bBlobberCachingState = TextureGraphEngine::GetBlobber()->IsCacheEnabled();
+		if (bDisableCache)
+			TextureGraphEngine::GetBlobber()->SetEnableCache(false);
+
 		Task->OriginalTextureGraphPtr = InTextureGraph;
 		Task->TextureGraphPtr = Cast<UTextureGraphBase>(StaticDuplicateObject(Task->OriginalTextureGraphPtr, GetTransientPackage(), NAME_None, RF_Standalone));
 		Task->TextureGraphPtr->Initialize();
@@ -41,6 +41,11 @@ UTG_AsyncExportTask* UTG_AsyncExportTask::TG_AsyncExportTask( UTextureGraphBase*
 
 const TArray<UTextureRenderTarget2D*>& UTG_AsyncExportTask::ActivateBlocking(JobBatchPtr Batch)
 {
+	if (!IsValid(TextureGraphPtr))
+	{
+		UE_LOG(LogTextureGraph, Warning, TEXT("TG_AsyncExportTask::Cannot export Texture Graph not selected"));
+		return OutputRts;
+	}
 	TargetExportSettings = FExportSettings();
 	TargetExportSettings.OnDone.BindUFunction(this, "OnExportDone");
 
diff --git a/Engine/Plugins/TextureGraph/Source/TextureGraph/Private/Blueprint/TG_BlueprintFunctionLibrary.cpp b/Engine/Plugins/TextureGraph/Source/TextureGraph/Private/Blueprint/TG_BlueprintFunctionLibrary.cpp
--- a/Engine/Plugins/TextureGraph/Source/TextureGraph/Private/Blueprint/TG_BlueprintFunctionLibrary.cpp
+++ b/Engine/Plugins/TextureGraph/Source/TextureGraph/Private/Blueprint/TG_BlueprintFunctionLibrary.cpp
@@ -275,8 +275,11 @@ const TArray<UTextureRenderTarget2D*>& UTG_BlueprintFunctionLibrary::RenderTextu
 
 void UTG_BlueprintFunctionLibrary::ExportTextureGraph(UObject* WorldContextObject, UTextureGraphBase* InTextureGraph, bool bOverwriteTextures, bool bSave, bool bExportAll)
 {
-	UTG_AsyncExportTask* Task = UTG_AsyncExportTask::TG_AsyncExportTask(InTextureGraph, bOverwriteTextures, bSave, bExportAll);
-	Task->ActivateBlocking(nullptr);
+	if (InTextureGraph != nullptr)
+	{
+		UTG_AsyncExportTask* Task = UTG_AsyncExportTask::TG_AsyncExportTask(InTextureGraph, bOverwriteTextures, bSave, bExportAll);
+		Task->ActivateBlocking(nullptr);
+	}
 }
 
 void UTG_BlueprintFunctionLibrary::AddParamWarning(FName ParamName, UObject* ObjectPtr, FString FunctionName)
