Mac Export issue fix

diff --git a/Engine/Plugins/TextureGraph/Source/TextureGraphEditor/Private/TG_Exporter.cpp b/Engine/Plugins/TextureGraph/Source/TextureGraphEditor/Private/TG_Exporter.cpp
--- a/Engine/Plugins/TextureGraph/Source/TextureGraphEditor/Private/TG_Exporter.cpp
+++ b/Engine/Plugins/TextureGraph/Source/TextureGraphEditor/Private/TG_Exporter.cpp
@@ -129,10 +129,13 @@ void FTG_InstanceImpl::Cleanup(bool bFlushMixInvalidations /*= true*/)
 
 void FTG_InstanceImpl::NotifyPostChange(const FPropertyChangedEvent& PropertyChangedEvent, FProperty* PropertyThatChanged)
 {
-	if(PropertyThatChanged->GetName() == GET_MEMBER_NAME_CHECKED(UTextureGraphInstance, ParentTextureGraph))
+	if(PropertyThatChanged != nullptr &&
+		PropertyThatChanged->GetName() == GET_MEMBER_NAME_CHECKED(UTextureGraphInstance, ParentTextureGraph))
 	{
 		TextureGraphInstance->SetParent(TextureGraphInstance->ParentTextureGraph);
-		SetTextureGraphToExport(TextureGraphInstance);
+
+		// clear out previous handles but don't flush mix invalidations as that would skip invalidating the mix with the new parent
+		SetTextureGraphToExport(TextureGraphInstance, false);
 	}
 }
 
@@ -469,10 +472,14 @@ FReply FTG_InstanceImpl::OnExportClicked(EAppReturnType::Type ButtonID)
 	return FReply::Handled();
 }
 
-void FTG_InstanceImpl::SetTextureGraphToExport(UTextureGraphInstance* InTextureGraph)
+void FTG_InstanceImpl::SetTextureGraphToExport(UTextureGraphInstance* InTextureGraph, bool bClearPreviousMixInvalidations)
 {
-	// clear out previous handles but don't flush mix invalidations as that would skip invalidating the mix with the new parent
-	Cleanup(false);
+	if (!InTextureGraph)
+	{
+		return;
+	}
+	
+	Cleanup(bClearPreviousMixInvalidations);
 
 	TextureGraphInstance = InTextureGraph;
 	
@@ -985,7 +992,7 @@ TSharedRef<SDockTab> FTG_ExporterUtility::CreateTGExporterTab(const FSpawnTabArg
 
 	TextureGraphInstance = NewObject<UTextureGraphInstance>();
 	TextureGraphInstance->Construct(FString());
-	Impl->SetTextureGraphToExport(TextureGraphInstance);
+	Impl->SetTextureGraphToExport(TextureGraphInstance, true);
 	return NomadTab;
 }
 
@@ -1009,7 +1016,7 @@ void FTG_ExporterUtility::SetTextureGraphToExport(UTextureGraphBase* InTextureGr
 		}
 		TextureGraphInstance->Initialize();
 	}
-	Impl->SetTextureGraphToExport(TextureGraphInstance);
+	Impl->SetTextureGraphToExport(TextureGraphInstance, true);
 }
 
 #undef LOCTEXT_NAMESPACE 
diff --git a/Engine/Plugins/TextureGraph/Source/TextureGraphEditor/Private/TG_InstanceEditor.cpp b/Engine/Plugins/TextureGraph/Source/TextureGraphEditor/Private/TG_InstanceEditor.cpp
--- a/Engine/Plugins/TextureGraph/Source/TextureGraphEditor/Private/TG_InstanceEditor.cpp
+++ b/Engine/Plugins/TextureGraph/Source/TextureGraphEditor/Private/TG_InstanceEditor.cpp
@@ -27,7 +27,10 @@ DEFINE_LOG_CATEGORY(LogTextureGraphInstanceEditor);
 
 void FTG_InstanceEditor::RefreshViewport()
 {
-	InstanceImpl->RefreshViewport();
+	if (InstanceImpl)
+	{
+		InstanceImpl->RefreshViewport();
+	}
 }
 
 void FTG_InstanceEditor::RefreshTool()
@@ -45,12 +48,18 @@ void FTG_InstanceEditor::SetMesh(class UMeshComponent* InPreviewMesh, class UWor
 void FTG_InstanceEditor::RegisterTabSpawners(const TSharedRef<class FTabManager>& InTabManager)
 {
 	FAssetEditorToolkit::RegisterTabSpawners(InTabManager);
-	InstanceImpl->RegisterTabSpawners(InTabManager);
+	if (InstanceImpl)
+	{
+		InstanceImpl->RegisterTabSpawners(InTabManager);
+	}
 }
 
 void FTG_InstanceEditor::UnregisterTabSpawners(const TSharedRef<class FTabManager>& InTabManager)
 {
-	InstanceImpl->UnregisterTabSpawners(InTabManager);
+	if (InstanceImpl)
+	{
+		InstanceImpl->UnregisterTabSpawners(InTabManager);
+	}
 	FAssetEditorToolkit::UnregisterTabSpawners(InTabManager);
 }
 
@@ -62,14 +71,14 @@ void FTG_InstanceEditor::OnTextureGraphPreSave(UObject* Object, FObjectPreSaveCo
 
 void FTG_InstanceEditor::InitEditor(const EToolkitMode::Type Mode, const TSharedPtr< class IToolkitHost >& InitToolkitHost, UTextureGraphInstance* InTextureGraph)
 {
-	FPropertyEditorModule& PropertyEditorModule = FModuleManager::GetModuleChecked<FPropertyEditorModule>("PropertyEditor");
-
-	// Initialize the TextureGraphInstance to create the runtime graph
-	if (InTextureGraph)
+	if (!InTextureGraph)
 	{
-		InTextureGraph->Initialize();
+		return;
 	}
-	
+
+	// Initialize the TextureGraphInstance to create the runtime graph
+	InTextureGraph->Initialize();
+		
 	InstanceImpl = MakeUnique<FTG_InstanceImpl>();
 	InstanceImpl->Initialize();
 
@@ -80,7 +89,7 @@ void FTG_InstanceEditor::InitEditor(const EToolkitMode::Type Mode, const TShared
 	}
 	FAssetEditorToolkit::InitAssetEditor(Mode, InitToolkitHost, TG_InstanceEditorAppIdentifier, InstanceImpl->GetDefaultLayout(), /*bCreateDefaultToolbar*/ true, /*bCreateDefaultStandaloneMenu*/ true, TextureGraphInstance);
 	
-	InstanceImpl->SetTextureGraphToExport(TextureGraphInstance);
+	InstanceImpl->SetTextureGraphToExport(TextureGraphInstance, true);
 	
 
 // 	// Add analytics tag
diff --git a/Engine/Plugins/TextureGraph/Source/TextureGraphEditor/Public/TG_Exporter.h b/Engine/Plugins/TextureGraph/Source/TextureGraphEditor/Public/TG_Exporter.h
--- a/Engine/Plugins/TextureGraph/Source/TextureGraphEditor/Public/TG_Exporter.h
+++ b/Engine/Plugins/TextureGraph/Source/TextureGraphEditor/Public/TG_Exporter.h
@@ -81,7 +81,7 @@ public:
 
 	void											Cleanup(bool bFlushMixInvalidations = true);
 	/** Sets the Texture Graph to be exported in the Exporter */
-	void 											SetTextureGraphToExport(UTextureGraphInstance* InTextureGraph);
+	void 											SetTextureGraphToExport(UTextureGraphInstance* InTextureGraph, bool bClearPreviousMixInvalidations);
 	void 											OnGraphChanged(UTG_Graph* InGraph, UTG_Node* InNode, bool Tweaking);
 	TSharedRef<FTabManager::FLayout>				GetDefaultLayout();
 	void 											Initialize();
