Fix unreachable code warnings revealed by editor ASan builds
#rnx

diff --git a/Engine/Plugins/Runtime/nDisplay/Source/DisplayClusterProjection/Private/Policy/VIOSO/DisplayClusterProjectionVIOSOPolicyFactory.cpp b/Engine/Plugins/Runtime/nDisplay/Source/DisplayClusterProjection/Private/Policy/VIOSO/DisplayClusterProjectionVIOSOPolicyFactory.cpp
--- a/Engine/Plugins/Runtime/nDisplay/Source/DisplayClusterProjection/Private/Policy/VIOSO/DisplayClusterProjectionVIOSOPolicyFactory.cpp
+++ b/Engine/Plugins/Runtime/nDisplay/Source/DisplayClusterProjection/Private/Policy/VIOSO/DisplayClusterProjectionVIOSOPolicyFactory.cpp
@@ -34,8 +34,8 @@ TSharedPtr<IDisplayClusterProjectionPolicy, ESPMode::ThreadSafe> FDisplayCluster
 	UE_LOG(LogDisplayClusterProjectionVIOSO, Verbose, TEXT("Instantiating projection policy <%s> id='%s'"), *InConfigurationProjectionPolicy->Type, *ProjectionPolicyId);
 
 	return  MakeShared<FDisplayClusterProjectionVIOSOPolicy, ESPMode::ThreadSafe>(ProjectionPolicyId, InConfigurationProjectionPolicy, VIOSOLibrary.ToSharedRef());
-#endif
-
+#else
 	UE_LOG(LogDisplayClusterProjectionVIOSO, Error, TEXT("VIOSO does not support the current platform: cannot instantiate projection policy <%s> id='%s'"), *InConfigurationProjectionPolicy->Type, *ProjectionPolicyId);
 	return nullptr;
+#endif
 }
diff --git a/Engine/Plugins/TextureGraph/Source/TextureGraphEditor/Private/Widgets/ActionMenuWidgets/STG_ActionMenuTileView.cpp b/Engine/Plugins/TextureGraph/Source/TextureGraphEditor/Private/Widgets/ActionMenuWidgets/STG_ActionMenuTileView.cpp
--- a/Engine/Plugins/TextureGraph/Source/TextureGraphEditor/Private/Widgets/ActionMenuWidgets/STG_ActionMenuTileView.cpp
+++ b/Engine/Plugins/TextureGraph/Source/TextureGraphEditor/Private/Widgets/ActionMenuWidgets/STG_ActionMenuTileView.cpp
@@ -1226,8 +1226,6 @@ FReply STG_ActionMenuTileView::OnKeyDown(const FGeometry& MyGeometry, const FKey
 		// When all else fails, it means we haven't filtered the list and we want to handle it as if we were just scrolling through a normal tree view
 		return TreeView->OnKeyDown(FindChildGeometry(MyGeometry, TreeView.ToSharedRef()), KeyEvent);
 	}
-
-	return FReply::Unhandled();
 }
 
 void STG_ActionMenuTileView::MarkActiveSuggestion()
diff --git a/Engine/Plugins/TextureGraph/Source/TextureGraphInsight/Private/View/STextureGraphInsightDeviceView.cpp b/Engine/Plugins/TextureGraph/Source/TextureGraphInsight/Private/View/STextureGraphInsightDeviceView.cpp
--- a/Engine/Plugins/TextureGraph/Source/TextureGraphInsight/Private/View/STextureGraphInsightDeviceView.cpp
+++ b/Engine/Plugins/TextureGraph/Source/TextureGraphInsight/Private/View/STextureGraphInsightDeviceView.cpp
@@ -56,42 +56,39 @@ public:
 
 	virtual TSharedRef<SWidget> GenerateWidgetForColumn(const FName& columnName) override
 	{
-			Column column = NameToColumn(columnName);
-
-			if (column == Main) {
-				// Rows in a TreeView need an expander button and some indentation
-				return SNew(SHorizontalBox)
-					+ SHorizontalBox::Slot()
-					.AutoWidth()
-					.HAlign(HAlign_Right)
-					.VAlign(VAlign_Fill)
-					[
-						SNew(SExpanderArrow, SharedThis(this))
-						.StyleSet(ExpanderStyleSet)
-					]
-					+ SHorizontalBox::Slot()
-					.FillWidth(1)
-					.Padding(FMargin(4.0f, 0.0f))
-					.VAlign(VAlign_Center)
-					[
-						SAssignNew(_textBoxes[column], STextBlock)
-						.Text(GetTextForColumn(column))
-					];
-			}
-			else
-			{
-				return SNew(SBox)
-					.Padding(FMargin(4.0f, 0.0f))
-					.VAlign(VAlign_Center)
-					[
-						SAssignNew(_textBoxes[column], STextBlock)
-						.Text(GetTextForColumn(column))
-						.OnDoubleClicked(this, &STextureGraphInsightDeviceListViewRow::OnDoubleClickedResultColumn)
-					];
-			}
-
-		// default to null widget if property cannot be found
-		return SNullWidget::NullWidget;
+		Column column = NameToColumn(columnName);
+
+		if (column == Main) {
+			// Rows in a TreeView need an expander button and some indentation
+			return SNew(SHorizontalBox)
+				+ SHorizontalBox::Slot()
+				.AutoWidth()
+				.HAlign(HAlign_Right)
+				.VAlign(VAlign_Fill)
+				[
+					SNew(SExpanderArrow, SharedThis(this))
+					.StyleSet(ExpanderStyleSet)
+				]
+				+ SHorizontalBox::Slot()
+				.FillWidth(1)
+				.Padding(FMargin(4.0f, 0.0f))
+				.VAlign(VAlign_Center)
+				[
+					SAssignNew(_textBoxes[column], STextBlock)
+					.Text(GetTextForColumn(column))
+				];
+		}
+		else
+		{
+			return SNew(SBox)
+				.Padding(FMargin(4.0f, 0.0f))
+				.VAlign(VAlign_Center)
+				[
+					SAssignNew(_textBoxes[column], STextBlock)
+					.Text(GetTextForColumn(column))
+					.OnDoubleClicked(this, &STextureGraphInsightDeviceListViewRow::OnDoubleClickedResultColumn)
+				];
+		}
 	}
 
 	FText GetTextForColumn(Column column) const
diff --git a/Engine/Source/Editor/SequencerCore/Public/MVVM/ViewModels/ViewModel.h b/Engine/Source/Editor/SequencerCore/Public/MVVM/ViewModels/ViewModel.h
--- a/Engine/Source/Editor/SequencerCore/Public/MVVM/ViewModels/ViewModel.h
+++ b/Engine/Source/Editor/SequencerCore/Public/MVVM/ViewModels/ViewModel.h
@@ -210,9 +210,9 @@ public:
 	template<typename T>
 	TViewModelPtr<T> FindAncestorOfType(bool bIncludeThis = false) const
 	{
-		for (TViewModelPtr<T> FirstAncestor : GetAncestorsOfType<T>(bIncludeThis))
+		if (auto AncestorsIt = GetAncestorsOfType<T>(bIncludeThis); AncestorsIt)
 		{
-			return FirstAncestor;
+			return *AncestorsIt;
 		}
 		return nullptr;
 	}
