Texture graph cook crash fix

diff --git a/Engine/Plugins/TextureGraph/Source/TextureGraph/Private/Expressions/Input/TG_Expression_Material.cpp b/Engine/Plugins/TextureGraph/Source/TextureGraph/Private/Expressions/Input/TG_Expression_Material.cpp
--- a/Engine/Plugins/TextureGraph/Source/TextureGraph/Private/Expressions/Input/TG_Expression_Material.cpp
+++ b/Engine/Plugins/TextureGraph/Source/TextureGraph/Private/Expressions/Input/TG_Expression_Material.cpp
@@ -188,8 +188,14 @@ void UTG_Expression_Material::SetAsset(UObject* Asset)
 
 void UTG_Expression_Material::SetTitleName(FName NewName)
 {
-	GetParentNode()->GetInputPin("InputMaterial")->SetAliasName(NewName);
-	TitleName = GetParentNode()->GetInputPin("InputMaterial")->GetAliasName();
+	if (UTG_Node* ParentNode = GetParentNode())
+	{
+		if (UTG_Pin* InputMaterialPin = ParentNode->GetInputPin("InputMaterial"))
+		{
+			InputMaterialPin->SetAliasName(NewName);
+			TitleName = InputMaterialPin->GetAliasName();
+		}
+	}
 }
 
 FName UTG_Expression_Material::GetTitleName() const
