Removing uses of implied command lists.

#Synced-CL 48381433

diff --git a/Engine/Plugins/TextureGraph/Source/TextureGraphEngine/2D/TextureHelper.cpp b/Engine/Plugins/TextureGraph/Source/TextureGraphEngine/2D/TextureHelper.cpp
--- a/Engine/Plugins/TextureGraph/Source/TextureGraphEngine/2D/TextureHelper.cpp
+++ b/Engine/Plugins/TextureGraph/Source/TextureGraphEngine/2D/TextureHelper.cpp
@@ -937,13 +937,15 @@ void TextureHelper::RawFromResource_Tiled(FTextureRHIRef ResourceRHI, const Buff
 		return;
 	}
 
+	FRHICommandListImmediate& RHICmdList = FRHICommandListImmediate::Get();
+
 	uint32 stride;
-	const uint8* SrcData = (const uint8*)RHILockTexture2D(ResourceRHI, 0, RLM_ReadOnly, stride, false);
+	const uint8* SrcData = (const uint8*)RHICmdList.LockTexture2D(ResourceRHI, 0, RLM_ReadOnly, stride, false);
 	size_t srcDataLength = Desc.Size();
 
 	RawFromMem_Tiled(SrcData, srcDataLength, Desc, TileSizeX, TileSizeY, Tiles);
 
-	RHIUnlockTexture2D(ResourceRHI, 0, false);
+	RHICmdList.UnlockTexture2D(ResourceRHI, 0, false);
 }
 
 BufferFormat TextureHelper::FindOptimalSupportedFormat(BufferFormat SrcFormat)
diff --git a/Engine/Source/Runtime/Engine/Private/Texture2DDynamic.cpp b/Engine/Source/Runtime/Engine/Private/Texture2DDynamic.cpp
--- a/Engine/Source/Runtime/Engine/Private/Texture2DDynamic.cpp
+++ b/Engine/Source/Runtime/Engine/Private/Texture2DDynamic.cpp
@@ -105,8 +105,10 @@ void FTexture2DDynamicResource::WriteRawToTexture_RenderThread(TArrayView64<cons
 		return;
 	}
 
+	FRHICommandListImmediate& RHICmdList = FRHICommandListImmediate::Get();
+
 	uint32 DestStride = 0;
-	uint8* DestData = reinterpret_cast<uint8*>(RHILockTexture2D(Texture2DRHI, 0, RLM_WriteOnly, DestStride, false, false));
+	uint8* DestData = reinterpret_cast<uint8*>(RHICmdList.LockTexture2D(Texture2DRHI, 0, RLM_WriteOnly, DestStride, false, false));
 
 
 	for (uint32 y = 0; y < Height; y++)
@@ -125,7 +127,7 @@ void FTexture2DDynamicResource::WriteRawToTexture_RenderThread(TArrayView64<cons
 		}
 	}
 
-	RHIUnlockTexture2D(Texture2DRHI, 0, false, false);
+	RHICmdList.UnlockTexture2D(Texture2DRHI, 0, false, false);
 }
 
 #endif
