[TG] Fixed a circular reference preventing Texture Graphs from being garbage collected

#Synced-CL 48612481

diff --git a/Engine/Plugins/TextureGraph/Source/TextureGraphEngine/FxMat/RenderMaterial_BP.cpp b/Engine/Plugins/TextureGraph/Source/TextureGraphEngine/FxMat/RenderMaterial_BP.cpp
--- a/Engine/Plugins/TextureGraph/Source/TextureGraphEngine/FxMat/RenderMaterial_BP.cpp
+++ b/Engine/Plugins/TextureGraph/Source/TextureGraphEngine/FxMat/RenderMaterial_BP.cpp
@@ -24,7 +24,7 @@ TSharedPtr<FMaterialInstanceStaticSwitchPermutationMap> FMaterialInstanceStaticS
 	{
 		MISSPMap = MakeShared<FMaterialInstanceStaticSwitchPermutationMap>();
 
-		TStrongObjectPtr<UMaterialInstanceConstant> MIC = TStrongObjectPtr<UMaterialInstanceConstant>(NewObject<UMaterialInstanceConstant>(InMaterial));
+		TStrongObjectPtr<UMaterialInstanceConstant> MIC = TStrongObjectPtr<UMaterialInstanceConstant>(NewObject<UMaterialInstanceConstant>());
 #if WITH_EDITOR
 		InMaterial->GetStaticParameterValues(MISSPMap->DefaultStaticParameterSet);
 		MISSPMap->DefaultKey = MISSPMap->KeyFromStaticSwitchParameters(MISSPMap->DefaultStaticParameterSet.StaticSwitchParameters);
@@ -74,7 +74,7 @@ UMaterialInstanceConstant* FMaterialInstanceStaticSwitchPermutationMap::GetMater
 	TStrongObjectPtr<UMaterialInstanceConstant>* MICP = PermutationsMap.Find(Key);
 	if (!MICP)
 	{
-		TStrongObjectPtr<UMaterialInstanceConstant> MIC = TStrongObjectPtr<UMaterialInstanceConstant>(NewObject<UMaterialInstanceConstant>(GetRootMaterial()));
+		TStrongObjectPtr<UMaterialInstanceConstant> MIC = TStrongObjectPtr<UMaterialInstanceConstant>(NewObject<UMaterialInstanceConstant>());
 #if WITH_EDITOR
 		MIC->SetParentEditorOnly(GetRootMaterial());
 		MIC->ClearParameterValuesEditorOnly();
