[TG] Made FMaterialInstanceStaticSwitchPermutationMap a USTRUCT to fix some object referencing issues that prevented the Texture Graph from being garbage collected

#Synced-CL 48653238

diff --git a/Engine/Plugins/TextureGraph/Source/TextureGraph/Private/Expressions/Input/TG_Expression_Material.cpp b/Engine/Plugins/TextureGraph/Source/TextureGraph/Private/Expressions/Input/TG_Expression_Material.cpp
--- a/Engine/Plugins/TextureGraph/Source/TextureGraph/Private/Expressions/Input/TG_Expression_Material.cpp
+++ b/Engine/Plugins/TextureGraph/Source/TextureGraph/Private/Expressions/Input/TG_Expression_Material.cpp
@@ -96,6 +96,8 @@ void UTG_Expression_Material::Initialize()
  	if (InputMaterial.IsValid())
 	{
 		MaterialCopy = DuplicateObject(InputMaterial.GetMaterial(), this);
+		MaterialCopy->ClearFlags(RF_Standalone | RF_Public);
+		MaterialCopy->SetFlags(RF_Transient);
 	}
 
 	Super::Initialize();
@@ -108,13 +110,11 @@ void UTG_Expression_Material::SetMaterialInternal(UMaterialInterface* InMaterial
 	{
 		MaterialCopy = nullptr;
 	}
-	else if (InMaterial->IsA<UMaterial>())
-	{
-		MaterialCopy = DuplicateObject(InMaterial, this);
-	}
-	else if (InMaterial->IsA<UMaterialInstance>())
+	else if (InMaterial->IsA<UMaterial>() || InMaterial->IsA<UMaterialInstance>())
 	{
 		MaterialCopy = DuplicateObject(InMaterial, this);
+		MaterialCopy->ClearFlags(RF_Standalone | RF_Public);
+		MaterialCopy->SetFlags(RF_Transient);
 	}
 
 	Super::SetMaterialInternal(MaterialCopy);
@@ -132,7 +132,7 @@ void UTG_Expression_Material::SetInputMaterial(const FTG_Material& InParamMateri
 	if (InputMaterial == InParamMaterial)
 	{
 		// Just check that the local MaterialCopy is valid, if not reassign below
-		if ((!InputMaterial.IsValid()) || (InputMaterial.IsValid() && MaterialCopy && MaterialPermutations))
+		if ((!InputMaterial.IsValid()) || (InputMaterial.IsValid() && MaterialCopy && MaterialPermutations.IsValid()))
 			return;
 	}
 
diff --git a/Engine/Plugins/TextureGraph/Source/TextureGraph/Private/Expressions/Input/TG_Expression_MaterialFunction.cpp b/Engine/Plugins/TextureGraph/Source/TextureGraph/Private/Expressions/Input/TG_Expression_MaterialFunction.cpp
--- a/Engine/Plugins/TextureGraph/Source/TextureGraph/Private/Expressions/Input/TG_Expression_MaterialFunction.cpp
+++ b/Engine/Plugins/TextureGraph/Source/TextureGraph/Private/Expressions/Input/TG_Expression_MaterialFunction.cpp
@@ -101,7 +101,7 @@ void UTG_Expression_MaterialFunction::SetMaterialFunction(UMaterialFunctionInter
 	if (InMaterialFunction == MaterialFunction)
 	{
 		// Just check that the MaterialInstance is valid, if not reassign below
-		if ((!MaterialFunction) || (MaterialFunction && MaterialPermutations))
+		if ((!MaterialFunction) || (MaterialFunction && MaterialPermutations.IsValid()))
 			return;
 	}
 
diff --git a/Engine/Plugins/TextureGraph/Source/TextureGraph/Private/Expressions/TG_Expression_MaterialBase.cpp b/Engine/Plugins/TextureGraph/Source/TextureGraph/Private/Expressions/TG_Expression_MaterialBase.cpp
--- a/Engine/Plugins/TextureGraph/Source/TextureGraph/Private/Expressions/TG_Expression_MaterialBase.cpp
+++ b/Engine/Plugins/TextureGraph/Source/TextureGraph/Private/Expressions/TG_Expression_MaterialBase.cpp
@@ -86,9 +86,10 @@ void UTG_Expression_MaterialBase::SetMaterialInternal(UMaterialInterface* InMate
 
 	if (!InMaterial)
 	{
-		MaterialPermutations.Reset();
+		MaterialPermutations = FMaterialInstanceStaticSwitchPermutationMap();
 	}
-	else {
+	else 
+	{
 		MaterialPermutations = FMaterialInstanceStaticSwitchPermutationMap::Create(GetMaterial());
 	}
 
@@ -144,7 +145,7 @@ void UTG_Expression_MaterialBase::GenerateMaterialAttributeOptions()
 
 void UTG_Expression_MaterialBase::Initialize()
 {
-	if (GetMaterial() && !MaterialPermutations)
+	if (GetMaterial() && !MaterialPermutations.IsValid())
 	{
 		MaterialPermutations = FMaterialInstanceStaticSwitchPermutationMap::Create(GetMaterial());
 	}
@@ -173,7 +174,7 @@ UMaterialInstanceConstant* UTG_Expression_MaterialBase::GetEvaluationMaterialIns
 		}
 	}
 
-	return MaterialPermutations->GetMaterialInstance(EvaluationCombination);
+	return MaterialPermutations.GetMaterialInstance(EvaluationCombination);
 }
 
 void UTG_Expression_MaterialBase::Evaluate(FTG_EvaluationContext* InContext)
diff --git a/Engine/Plugins/TextureGraph/Source/TextureGraph/Public/Expressions/TG_Expression_MaterialBase.h b/Engine/Plugins/TextureGraph/Source/TextureGraph/Public/Expressions/TG_Expression_MaterialBase.h
--- a/Engine/Plugins/TextureGraph/Source/TextureGraph/Public/Expressions/TG_Expression_MaterialBase.h
+++ b/Engine/Plugins/TextureGraph/Source/TextureGraph/Public/Expressions/TG_Expression_MaterialBase.h
@@ -74,7 +74,8 @@ public:
 protected:
 	// A local Material Instance Constant is recreated from the reference material assigned through SetMaterialInternal
 	// All static switches permutations are cached in this map and only created on demand
-	TSharedPtr<FMaterialInstanceStaticSwitchPermutationMap> MaterialPermutations;
+	UPROPERTY(Transient)
+	FMaterialInstanceStaticSwitchPermutationMap MaterialPermutations;
 
 	// The list of the material attributes declared in the signature as Output of the expression,
 	TArray<EDrawMaterialAttributeTarget>	AvailableMaterialAttributeIds; // The set of material properties available for rendering
diff --git a/Engine/Plugins/TextureGraph/Source/TextureGraphEngine/2D/Tex.cpp b/Engine/Plugins/TextureGraph/Source/TextureGraphEngine/2D/Tex.cpp
--- a/Engine/Plugins/TextureGraph/Source/TextureGraphEngine/2D/Tex.cpp
+++ b/Engine/Plugins/TextureGraph/Source/TextureGraphEngine/2D/Tex.cpp
@@ -408,7 +408,12 @@ void Tex::AddReferencedObjects(FReferenceCollector& Collector)
 {
 	/// Don't do anything if caching is disabled
 	if (!TextureGraphEngine::GetBlobber()->IsCacheEnabled() || bNoCache)
+	{
+		// Clear these pointers so that they don't become dangling if the objects are deleted
+		RT = nullptr;
+		Texture = nullptr;
 		return;
+	}
 
 	if (RT)
 		Collector.AddReferencedObject(RT);
diff --git a/Engine/Plugins/TextureGraph/Source/TextureGraphEngine/FxMat/RenderMaterial_BP.cpp b/Engine/Plugins/TextureGraph/Source/TextureGraphEngine/FxMat/RenderMaterial_BP.cpp
--- a/Engine/Plugins/TextureGraph/Source/TextureGraphEngine/FxMat/RenderMaterial_BP.cpp
+++ b/Engine/Plugins/TextureGraph/Source/TextureGraphEngine/FxMat/RenderMaterial_BP.cpp
@@ -17,38 +17,35 @@
 #include "FxMaterial_DrawMaterial.h"
 #include "Job/Job.h"
 
-TSharedPtr<FMaterialInstanceStaticSwitchPermutationMap> FMaterialInstanceStaticSwitchPermutationMap::Create(UMaterialInterface* InMaterial)
+FMaterialInstanceStaticSwitchPermutationMap FMaterialInstanceStaticSwitchPermutationMap::Create(UMaterialInterface* InMaterial)
 {
-	TSharedPtr<FMaterialInstanceStaticSwitchPermutationMap> MISSPMap;
+	FMaterialInstanceStaticSwitchPermutationMap MISSPMap;
 	if (InMaterial)
 	{
-		MISSPMap = MakeShared<FMaterialInstanceStaticSwitchPermutationMap>();
-
-		TStrongObjectPtr<UMaterialInstanceConstant> MIC = TStrongObjectPtr<UMaterialInstanceConstant>(NewObject<UMaterialInstanceConstant>());
+		UMaterialInstanceConstant* MIC = NewObject<UMaterialInstanceConstant>();
 #if WITH_EDITOR
-		InMaterial->GetStaticParameterValues(MISSPMap->DefaultStaticParameterSet);
-		MISSPMap->DefaultKey = MISSPMap->KeyFromStaticSwitchParameters(MISSPMap->DefaultStaticParameterSet.StaticSwitchParameters);
+		InMaterial->GetStaticParameterValues(MISSPMap.DefaultStaticParameterSet);
+		MISSPMap.DefaultKey = MISSPMap.KeyFromStaticSwitchParameters(MISSPMap.DefaultStaticParameterSet.StaticSwitchParameters);
 
 		MIC->SetParentEditorOnly(InMaterial);
 		MIC->ClearParameterValuesEditorOnly();
-		MIC->SetFlags(RF_Standalone);
-		MIC->MarkPackageDirty();
+		MIC->SetFlags(RF_Transient);
 		MIC->PreEditChange(NULL);
 		MIC->PostEditChange();
 #endif
-		MISSPMap->PermutationsMap.Add(MISSPMap->DefaultKey, MIC);
+		MISSPMap.PermutationsMap.Add(MISSPMap.DefaultKey, MIC);
 	}
 	return MISSPMap;
 }
 
 UMaterialInterface* FMaterialInstanceStaticSwitchPermutationMap::GetRootMaterial()
 {
-	TStrongObjectPtr<UMaterialInstanceConstant>* MICP = PermutationsMap.Find(DefaultKey);
+	TObjectPtr<UMaterialInstanceConstant>* MICP = PermutationsMap.Find(DefaultKey);
 	if (MICP)
 	{
 		// Use the FIRST MIC that initiated the MISWPermutationMap as the "root Material" 
 		// in order to preserve the default values
-		return Cast<UMaterialInterface>(MICP->Get());
+		return *MICP;
 	}
 	return nullptr;
 }
@@ -71,15 +68,14 @@ UMaterialInstanceConstant* FMaterialInstanceStaticSwitchPermutationMap::GetMater
 
 	int32 Key = KeyFromStaticSwitchParameters(Parameters);
 
-	TStrongObjectPtr<UMaterialInstanceConstant>* MICP = PermutationsMap.Find(Key);
+	TObjectPtr<UMaterialInstanceConstant>* MICP = PermutationsMap.Find(Key);
 	if (!MICP)
 	{
-		TStrongObjectPtr<UMaterialInstanceConstant> MIC = TStrongObjectPtr<UMaterialInstanceConstant>(NewObject<UMaterialInstanceConstant>());
+		UMaterialInstanceConstant* MIC = NewObject<UMaterialInstanceConstant>();
 #if WITH_EDITOR
 		MIC->SetParentEditorOnly(GetRootMaterial());
 		MIC->ClearParameterValuesEditorOnly();
-		MIC->SetFlags(RF_Standalone);
-		MIC->MarkPackageDirty();
+		MIC->SetFlags(RF_Transient);
 
 		FStaticParameterSet InstanceStaticParameterSet = DefaultStaticParameterSet;
 		InstanceStaticParameterSet.StaticSwitchParameters = Parameters;
@@ -89,14 +85,20 @@ UMaterialInstanceConstant* FMaterialInstanceStaticSwitchPermutationMap::GetMater
 #endif
 		PermutationsMap.Add(Key, MIC);
 
-		return MIC.Get();
+		return MIC;
 	}
 	else
 	{
-		return MICP->Get();
+		return *MICP;
 	}
 }
 
+bool FMaterialInstanceStaticSwitchPermutationMap::IsValid() const
+{
+	// If a non-null material is provided to Create, it will be in the PermutationsMap, so if 
+	// PermutationsMap has no entries, it implies Create was not used to initialize this struct.
+	return !PermutationsMap.IsEmpty();
+}
 
 bool RenderMaterial_BP::ValidateMaterialCompatible(UMaterialInterface* InMaterial)
 {
@@ -117,12 +119,11 @@ RenderMaterial_BP::RenderMaterial_BP(FString Name, UMaterialInterface* InMateria
 	else if (Material && !MaterialInstance)
 	{
 
-		MaterialInstance = TStrongObjectPtr<UMaterialInstanceConstant>(NewObject<UMaterialInstanceConstant>(Material));
+		MaterialInstance = TStrongObjectPtr<UMaterialInstanceConstant>(NewObject<UMaterialInstanceConstant>());
 #if WITH_EDITOR
 		MaterialInstance->SetParentEditorOnly(Material);
 		MaterialInstance->ClearParameterValuesEditorOnly();
-		MaterialInstance->SetFlags(RF_Standalone);
-		MaterialInstance->MarkPackageDirty();
+		MaterialInstance->SetFlags(RF_Transient);
 		MaterialInstance->PreEditChange(NULL);
 		MaterialInstance->PostEditChange();
 #endif
diff --git a/Engine/Plugins/TextureGraph/Source/TextureGraphEngine/FxMat/RenderMaterial_BP.h b/Engine/Plugins/TextureGraph/Source/TextureGraphEngine/FxMat/RenderMaterial_BP.h
--- a/Engine/Plugins/TextureGraph/Source/TextureGraphEngine/FxMat/RenderMaterial_BP.h
+++ b/Engine/Plugins/TextureGraph/Source/TextureGraphEngine/FxMat/RenderMaterial_BP.h
@@ -5,7 +5,8 @@
 #include "RenderMaterial.h"
 #include "Materials/MaterialInstanceDynamic.h"
 #include "Materials/MaterialInstanceConstant.h"
-#include "UObject/StrongObjectPtr.h"
+
+#include "RenderMaterial_BP.generated.h"
 
 #define UE_API TEXTUREGRAPHENGINE_API
 
@@ -15,16 +16,24 @@ typedef std::shared_ptr< FxMaterial> FxMaterialPtr;
 // Helper container of MIC for the different static switch permutations of a Material
 // Used by the caller of a RenderMaterial_BP to pick a specific permutation different from the default set.
 // See TG_Expression_MaterialBase for a usage example
+USTRUCT()
 struct FMaterialInstanceStaticSwitchPermutationMap
 {
+	GENERATED_BODY()
+
 	int32												DefaultKey = 0;
+	
+	UPROPERTY()
 	FStaticParameterSet									DefaultStaticParameterSet;
-	TMap<int32, TStrongObjectPtr<UMaterialInstanceConstant>>	PermutationsMap;
+	
+	UPROPERTY()
+	TMap<int32, TObjectPtr<UMaterialInstanceConstant>>	PermutationsMap;
 
-	static UE_API TSharedPtr<FMaterialInstanceStaticSwitchPermutationMap> Create(UMaterialInterface* InMaterial);
+	static UE_API FMaterialInstanceStaticSwitchPermutationMap Create(UMaterialInterface* InMaterial);
 	UE_API UMaterialInterface*				GetRootMaterial();
 	UE_API int32							KeyFromStaticSwitchParameters(const TArray<FStaticSwitchParameter>& Parameters);
 	UE_API UMaterialInstanceConstant*		GetMaterialInstance(const TArray<FStaticSwitchParameter>& Parameters);
+	UE_API bool								IsValid() const;
 };
 
 class RenderMaterial_BP : public RenderMaterial
