Fixing TextureScriptTests

#Synced-CL 48720034

diff --git a/Engine/Plugins/TextureGraph/Source/TextureGraphEngine/Data/Blobber.h b/Engine/Plugins/TextureGraph/Source/TextureGraphEngine/Data/Blobber.h
--- a/Engine/Plugins/TextureGraph/Source/TextureGraphEngine/Data/Blobber.h
+++ b/Engine/Plugins/TextureGraph/Source/TextureGraphEngine/Data/Blobber.h
@@ -20,7 +20,7 @@ typedef cti::continuable<JobResultPtr>		AsyncJobResultPtr;
 typedef T_BlobRef<Blob>						BlobRef;
 typedef T_BlobRef<TiledBlob>				TiledBlobRef;
 
-DECLARE_LOG_CATEGORY_EXTERN(LogBlob, Log, Verbose);
+DECLARE_LOG_CATEGORY_EXTERN(LogBlob, Log, All);
 
 typedef std::vector<HashType>   HashTypes;
 class JobArg;
