[TG] Improve HistogramRange node and output portion of T_Levels.
In the levels operator output low and high were applied in the wrong space, this is now fixed.
Renammed the levels configuration call to map the expression name
fixed a bug in the way the mid value was passed in the case of AUtoLevels

#Synced-CL 48629549

diff --git a/Engine/Plugins/TextureGraph/Shaders/Expressions/Expression_Levels.usf b/Engine/Plugins/TextureGraph/Shaders/Expressions/Expression_Levels.usf
--- a/Engine/Plugins/TextureGraph/Shaders/Expressions/Expression_Levels.usf
+++ b/Engine/Plugins/TextureGraph/Shaders/Expressions/Expression_Levels.usf
@@ -12,10 +12,6 @@
 #define AUTO_LEVELS 0
 #endif
 
-#ifndef OUT_LEVELS
-#define OUT_LEVELS 0
-#endif
-
 Texture2D SourceTexture;
 float LowValue;
 float HighValue;
@@ -43,18 +39,6 @@ float LevelsControlInputCurve(float Value, float LowInput, float HighInput, floa
 	// thus
 	float MidExponent = log(0.5) / log(MidRanged);
 	
-#if OUT_LEVELS
-	float4 HistoMin			= Histogram.Load(int3(2, 1, 0));
-	float4 HistoMax			= Histogram.Load(int3(3, 1, 0));
-		
-	LowInput = 	HistoMin.w;
-	HighInput = HistoMax.w;
-	
-	float InputRange = clamp(HighInput - LowInput, 0, 1);
-	
-	ValueRanged = OutLow + ((ValueRanged - LowInput) / InputRange) * OutputRange;
-#endif /// OUT_LEVELS
-
 	return pow(ValueRanged, MidExponent);
 }
 
@@ -81,13 +65,17 @@ float4 FSH_Levels(float2 uv : TEXCOORD0) : SV_Target0
 	float4 HistoMaxCount	= Histogram.Load(int3(1, 1, 0));
 	float4 HistoMin			= Histogram.Load(int3(2, 1, 0));
 	float4 HistoMax			= Histogram.Load(int3(3, 1, 0));
-		
-	LowInput = 	HistoMin.w;
-	HighInput = HistoMax.w;
+			
+	LowInput = 	HistoMin.x;
+	HighInput = HistoMax.x;
+	MidInput = MidPercentage * (HighInput - LowInput) + LowInput;
 #endif /// AUTO_LEVELS
 
-	float LeveledValue = LevelsControlInputCurve(SourceValue, LowInput, HighInput, MidValue);
+	float LeveledValue = LevelsControlInputCurve(SourceValue, LowInput, HighInput, MidInput);
 
+	// Apply out levels
+	LeveledValue = OutLow + (LeveledValue) * OutputRange;
+	
 #if CONVERT_TO_GRAYSCALE
 	// The new  Grayscale level is used to scale the color darker or brighter
 	float Ratio = (SourceValue == 0 ? 0.0 : LeveledValue / SourceValue); 	
diff --git a/Engine/Plugins/TextureGraph/Source/TextureGraph/Private/Expressions/FIlter/TG_Expression_Levels.cpp b/Engine/Plugins/TextureGraph/Source/TextureGraph/Private/Expressions/FIlter/TG_Expression_Levels.cpp
--- a/Engine/Plugins/TextureGraph/Source/TextureGraph/Private/Expressions/FIlter/TG_Expression_Levels.cpp
+++ b/Engine/Plugins/TextureGraph/Source/TextureGraph/Private/Expressions/FIlter/TG_Expression_Levels.cpp
@@ -266,7 +266,7 @@ void UTG_Expression_HistogramScan::Evaluate(FTG_EvaluationContext* InContext)
 	}
 
 	LevelsControl = MakeShared<FLevels>();
-	LevelsControl->InitFromPositionContrast(Position, Contrast);
+	LevelsControl->InitFromHistogramScan(Position, Contrast);
 
 
 	BufferDescriptor Desc = Output.GetBufferDescriptor();
@@ -287,7 +287,7 @@ void UTG_Expression_HistogramRange::Evaluate(FTG_EvaluationContext* InContext)
 	}
 
 	LevelsControl = MakeShared<FLevels>();
-	LevelsControl->InitFromRange(Range, Position);
+	LevelsControl->InitFromHistogramRange(Position, Range);
 
 	BufferDescriptor Desc = Output.GetBufferDescriptor();
 	Output = T_Levels::Create(InContext->Cycle, Desc, Input.RasterBlob, LevelsControl, InContext->TargetId);
diff --git a/Engine/Plugins/TextureGraph/Source/TextureGraph/Public/Expressions/Filter/TG_Expression_Levels.h b/Engine/Plugins/TextureGraph/Source/TextureGraph/Public/Expressions/Filter/TG_Expression_Levels.h
--- a/Engine/Plugins/TextureGraph/Source/TextureGraph/Public/Expressions/Filter/TG_Expression_Levels.h
+++ b/Engine/Plugins/TextureGraph/Source/TextureGraph/Public/Expressions/Filter/TG_Expression_Levels.h
@@ -167,7 +167,7 @@ public:
 
 	// Drives the contrast of the histogram of the input image
 	UPROPERTY(EditAnywhere, Category = NoCategory, meta = (TGType = "TG_Setting", UIMin = "0", ClampMin = "0", UIMax = "1", ClampMax = "1", MD_ScalarEditor))
-	float 								Contrast = 0.5f; 
+	float 								Contrast = 0.0f; 
 
 	// The output image filtered result of the Levels operator 
 	UPROPERTY(EditAnywhere, Category = NoCategory, meta = (TGType = "TG_Output", PinDisplayName = ""))
@@ -196,14 +196,13 @@ public:
 	UPROPERTY(EditAnywhere, Category = NoCategory, meta = (TGType = "TG_Input", PinDisplayName = "", MD_HistogramLuminance))
 	FTG_Texture							Input;
 
-	// Drives the range of the levels out
+	// Drives the center of the levels out
 	UPROPERTY(EditAnywhere, Category = NoCategory, meta = (TGType = "TG_Setting", UIMin = "0", ClampMin = "0", UIMax = "1", ClampMax = "1", MD_ScalarEditor))
-	float 								Range = 1.0f; 
-
+	float 								Position = 0.5f;
 
-	// Drives the position of the histogram of the input image
+	// Drives the range of the levels out
 	UPROPERTY(EditAnywhere, Category = NoCategory, meta = (TGType = "TG_Setting", UIMin = "0", ClampMin = "0", UIMax = "1", ClampMax = "1", MD_ScalarEditor))
-	float 								Position = 0.5f;
+	float 								Range = 1.0f; 
 
 	// The output image filtered result of the Levels operator 
 	UPROPERTY(EditAnywhere, Category = NoCategory, meta = (TGType = "TG_Output", PinDisplayName = ""))
diff --git a/Engine/Plugins/TextureGraph/Source/TextureGraphEngine/Transform/Expressions/T_Levels.cpp b/Engine/Plugins/TextureGraph/Source/TextureGraphEngine/Transform/Expressions/T_Levels.cpp
--- a/Engine/Plugins/TextureGraph/Source/TextureGraphEngine/Transform/Expressions/T_Levels.cpp
+++ b/Engine/Plugins/TextureGraph/Source/TextureGraphEngine/Transform/Expressions/T_Levels.cpp
@@ -11,7 +11,6 @@
 
 IMPLEMENT_GLOBAL_SHADER(FSH_Levels, "/Plugin/TextureGraph/Expressions/Expression_Levels.usf", "FSH_Levels", SF_Pixel);
 
-
 bool FLevels::SetHigh(float InValue)
 {
 	float NewValue = std::min(1.f, std::max(InValue, Low));
@@ -98,30 +97,38 @@ void FLevels::InitFromAutoLevels(float InMidPercentage)
 	MidPercentage = std::min(std::max(0.0f, InMidPercentage), 1.0f);
 }
 
-// Histogram scan
-void FLevels::InitFromPositionContrast(float InPosition, float InContrast)
+void FLevels::InitFromHistogramScan(float InPosition, float InContrast)
 {
-	float C = InContrast * 0.5f;
+	// Histogram scan is affecting the input low and high and the mid.
+	// InPosition determines the center of the remapped input range
+	// Contrast controls the width of the range of remapping
+	// Contrast is scaled depending on the InPosition
+	// Contrast shrink to 0 when InPosition is at 0 or 1
+	// Contrast is whole when InPosition is at 0.5
+
 	float P = 1.0f - FMath::Clamp(InPosition, 0, 1);
-	float P1 = (FMath::Max(P, 0.5f) - 0.5f) * 2.0f;
-	float P2 = FMath::Min(P * 2.0f, 1.0f);
-	Low = FMath::Lerp(P1, P2, C);
-	High = FMath::Lerp(P2, P1, C);
+
+	float C = (1.0 - FMath::Clamp(InContrast, 0, 1)) * 0.5f;
+	float D = 1.0f - FMath::Abs(0.5f - P) * 2.0f;
+	float ScaledC = C * D;
+
+	Low = P - ScaledC;
+	High = P + ScaledC;
+
 	Mid = Low + (High - Low) * 0.5f;
-	
+
 	IsAutoLevels = false;
 }
 
-void FLevels::InitFromRange(float InRange_, float InPosition_)
+void FLevels::InitFromHistogramRange(float InPosition,  float InRange)
 {
-	float InRange = FMath::Clamp(1.0f - InRange_, 0.0f, 1.0f) * 0.5f;
-	float InPosition = FMath::Clamp(1.0f - InPosition_, 0.0f, 1.0f) * 0.5f;
-	float C = InRange; /// *0.5f;
-	float P = 1.0f - FMath::Clamp(InPosition, 0, 1);
-	float P1 = (FMath::Max(P, 0.5f) - 0.5f) * 2.0f;
-	float P2 = FMath::Min(P * 2.0f, 1.0f);
-	OutLow = FMath::Lerp(P1, P2, C);
-	OutHigh = FMath::Lerp(P2, P1, C);
+	// Histogram range is just affecting the output low and high
+
+	float range = FMath::Clamp(InRange, 0.0f, 1.0f) * 0.5f;
+	float pos = FMath::Clamp(InPosition, 0.0f, 1.0f);
+
+	OutLow = FMath::Clamp(pos - range, 0.0f, 1.0f);
+	OutHigh = FMath::Clamp(pos + range, 0.0f, 1.0f);
 
 	IsAutoLevels = false;
 }
@@ -291,9 +298,8 @@ TiledBlobPtr	T_Levels::Create(MixUpdateCyclePtr Cycle, BufferDescriptor DesiredO
 
 	bool bNeedHistogramRaw = false;
 	TiledBlobPtr Histogram = TextureHelper::GBlack;
-	bool bIsModifyingOut = InLevels->OutLow > 0 || InLevels->OutHigh < 1;
 
-	if (InLevels->IsAutoLevels || bIsModifyingOut)
+	if (InLevels->IsAutoLevels)
 	{
 		bNeedHistogramRaw = true;
 		Histogram = T_TextureHistogram::Create(Cycle, Source, TargetId);
@@ -308,7 +314,6 @@ TiledBlobPtr	T_Levels::Create(MixUpdateCyclePtr Cycle, BufferDescriptor DesiredO
 	FSH_Levels::FPermutationDomain PermutationVector;
 	PermutationVector.Set<FSH_Levels::FConvertToGrayscale>(NeedsConvertToGrayscale);
 	PermutationVector.Set<FSH_Levels::FIsAutoLevels>(InLevels->IsAutoLevels);
-	PermutationVector.Set<FSH_Levels::FIsOutLevels>(bIsModifyingOut);
 
 	std::shared_ptr<FxMaterial_Normal<VSH_Simple, FSH_Levels>> Mat = 
 		std::make_shared<FxMaterial_Normal<VSH_Simple, FSH_Levels>>(typename VSH_Simple::FPermutationDomain(), PermutationVector);
diff --git a/Engine/Plugins/TextureGraph/Source/TextureGraphEngine/Transform/Expressions/T_Levels.h b/Engine/Plugins/TextureGraph/Source/TextureGraphEngine/Transform/Expressions/T_Levels.h
--- a/Engine/Plugins/TextureGraph/Source/TextureGraphEngine/Transform/Expressions/T_Levels.h
+++ b/Engine/Plugins/TextureGraph/Source/TextureGraphEngine/Transform/Expressions/T_Levels.h
@@ -35,8 +35,7 @@ public:
 
 	class FConvertToGrayscale : SHADER_PERMUTATION_BOOL("CONVERT_TO_GRAYSCALE");
 	class FIsAutoLevels : SHADER_PERMUTATION_BOOL("AUTO_LEVELS");
-	class FIsOutLevels : SHADER_PERMUTATION_BOOL("OUT_LEVELS");
-	using FPermutationDomain = TShaderPermutationDomain<FConvertToGrayscale, FIsAutoLevels, FIsOutLevels>;
+	using FPermutationDomain = TShaderPermutationDomain<FConvertToGrayscale, FIsAutoLevels>;
 
 	static bool						ShouldCompilePermutation(const FGlobalShaderPermutationParameters& Parameters) 
 	{
@@ -87,10 +86,11 @@ struct FLevels
 	UE_API float EvalMidExponent() const;
 	UE_API bool SetMidFromMidExponent(float InExponent);
 
+	// various init from different parameter set
 	UE_API void InitFromLowMidHigh(float LowValue, float MidValue, float HighValue, float OutLow, float OutHigh);
 	UE_API void InitFromAutoLevels(float MidPercentage);
-	UE_API void InitFromPositionContrast(float MidValue, float Contrast);
-	UE_API void InitFromRange(float InRange, float InPosition);
+	UE_API void InitFromHistogramScan(float InPosition, float InContrast);
+	UE_API void InitFromHistogramRange(float InPosition, float InRange);
 
 	FLevelsHistogramData HistogramData;
 };
