TG: remove usage of RF_MarkAsNative

#Synced-CL 49032025

diff --git a/Engine/Plugins/TextureGraph/Source/TextureGraphEngine/2D/Tex.cpp b/Engine/Plugins/TextureGraph/Source/TextureGraphEngine/2D/Tex.cpp
--- a/Engine/Plugins/TextureGraph/Source/TextureGraphEngine/2D/Tex.cpp
+++ b/Engine/Plugins/TextureGraph/Source/TextureGraphEngine/2D/Tex.cpp
@@ -181,7 +181,7 @@ size_t TexDescriptor::GetPitch() const
 }
 
 //////////////////////////////////////////////////////////////////////////
-EObjectFlags Tex::Flags = (RF_Public | RF_MarkAsNative | RF_Standalone);
+EObjectFlags Tex::Flags = (RF_Public | RF_Standalone);
 
 Tex::Tex(int32 Width, int32 Height, EPixelFormat PixelFormat)
 	: Desc(Width, Height, PixelFormat)
