[TG] Fix an issue in material compilation when SUbstrate is enabled and a Material is usiung the node "Forward Selected Directional Light"
added the missing include to the TGMaterialShader

#Synced-CL 49553855

diff --git a/Engine/Plugins/TextureGraph/Shaders/TextureGraphMaterialShader.usf b/Engine/Plugins/TextureGraph/Shaders/TextureGraphMaterialShader.usf
--- a/Engine/Plugins/TextureGraph/Shaders/TextureGraphMaterialShader.usf
+++ b/Engine/Plugins/TextureGraph/Shaders/TextureGraphMaterialShader.usf
@@ -8,6 +8,7 @@
 #include "/Engine/Public/Platform.ush"
 #include "/Engine/Private/Common.ush"
 #include "/Engine/Generated/Material.ush"
+#include "/Engine/Private/BasePassCommonFunctions.ush"
 
 #define SHADER_MATERIAL_USES_SUBSTRATE	(SUBSTRATE_ENABLED && MATERIAL_IS_SUBSTRATE)
 #if SHADER_MATERIAL_USES_SUBSTRATE
