Fix crash when opening a TG asset after force deleting material referenced in the TG asset.

#Synced-CL 49735945

diff --git a/Engine/Plugins/TextureGraph/Source/TextureGraph/Private/Expressions/Input/TG_Expression_Material.cpp b/Engine/Plugins/TextureGraph/Source/TextureGraph/Private/Expressions/Input/TG_Expression_Material.cpp
--- a/Engine/Plugins/TextureGraph/Source/TextureGraph/Private/Expressions/Input/TG_Expression_Material.cpp
+++ b/Engine/Plugins/TextureGraph/Source/TextureGraph/Private/Expressions/Input/TG_Expression_Material.cpp
@@ -93,7 +93,7 @@ void UTG_Expression_Material::PostLoad()
 void UTG_Expression_Material::Initialize()
 {
 	// If the referenced material is valid, then we need to recreate a duplicate
- 	if (InputMaterial.IsValid())
+ 	if (InputMaterial.IsValid() && InputMaterial.GetMaterial() != nullptr)
 	{
 		MaterialCopy = DuplicateObject(InputMaterial.GetMaterial(), this);
 		MaterialCopy->ClearFlags(RF_Standalone | RF_Public);
