Fix TexturePath expression evaluating black blobs in some cases

#Synced-CL 49841360

diff --git a/Engine/Plugins/TextureGraph/Source/TextureGraph/Private/Expressions/Input/TG_Expression_TexturePath.cpp b/Engine/Plugins/TextureGraph/Source/TextureGraph/Private/Expressions/Input/TG_Expression_TexturePath.cpp
--- a/Engine/Plugins/TextureGraph/Source/TextureGraph/Private/Expressions/Input/TG_Expression_TexturePath.cpp
+++ b/Engine/Plugins/TextureGraph/Source/TextureGraph/Private/Expressions/Input/TG_Expression_TexturePath.cpp
@@ -72,14 +72,19 @@ void UTG_Expression_TexturePath::Evaluate(FTG_EvaluationContext* InContext)
 {
 	Super::Evaluate(InContext);
 
-	/// This is our default state
-	Output.SetNum(1);
-	Output.Set(0, FTG_Texture::GetBlack());
-
 	FString LocalPath;
 	bool bValidation = ValidateInputPath(LocalPath);
 
-	if (Output.Num() == 0 || LocalPath != OutputPath)
+	// Determine if we actually need to reload
+	// Key insight: Don't touch Output unless we're going to replace it!
+	bool bOutputIsValid = (Output.Num() > 0) && 
+						  Output.Get(0).GetTexture().RasterBlob && 
+						  Output.Get(0).GetTexture().RasterBlob->IsValid();
+    
+	bool bPathChanged = (LocalPath != OutputPath);
+	bool bNeedsReload = bPathChanged || !bOutputIsValid;
+
+	if (bNeedsReload)
 	{
 		if (!Path.IsEmpty() && bValidation)
 		{
@@ -110,12 +115,31 @@ void UTG_Expression_TexturePath::Evaluate(FTG_EvaluationContext* InContext)
 						Output.Set(Index, Item);
 					}
 				}
+				else
+				{
+					// Directory exists but empty - set black
+					Output.SetNum(1);
+					Output.Set(0, FTG_Texture::GetBlack());
+				}
+			}
+			else
+			{
+				// Path validation passed but file/dir doesn't exist (edge case)
+				Output.SetNum(1);
+				Output.Set(0, FTG_Texture::GetBlack());
 			}
 		}
-
+		else
+		{
+			// Invalid or empty path - set black as fallback
+			Output.SetNum(1);
+			Output.Set(0, FTG_Texture::GetBlack());
+		}
+		
+		// Update cached path
 		OutputPath = LocalPath;
 	}
-
+	
 	//For the connected pin we will report error here in eveluate because it does not have the updated value during validation.
 	UTG_Pin* PathPin = GetParentNode()->GetPin(GET_MEMBER_NAME_CHECKED(UTG_Expression_TexturePath, Path));
 
