TG: [Linux] Fix various crashes on loading TG assets

#Synced-CL 49841360

diff --git a/Engine/Plugins/TextureGraph/Source/TextureGraphEngine/2D/Tex.cpp b/Engine/Plugins/TextureGraph/Source/TextureGraphEngine/2D/Tex.cpp
--- a/Engine/Plugins/TextureGraph/Source/TextureGraphEngine/2D/Tex.cpp
+++ b/Engine/Plugins/TextureGraph/Source/TextureGraphEngine/2D/Tex.cpp
@@ -392,11 +392,10 @@ void Tex::InitRT(bool ForceFloat /* = false */)
 	RT->bCanCreateUAV = Desc.bUAV;
 	RT->RenderTargetFormat = TextureHelper::GetRenderTargetFormatFromPixelFormat(Desc.Format);
 	RT->OverrideFormat = Desc.Format;
-	RT->SizeX = Desc.Width;
-	RT->SizeY = Desc.Height;
 	RT->LODBias = 0;
-
+	RT->InitAutoFormat(Desc.Width, Desc.Height);
 	RT->UpdateResource();
+	FlushRenderingCommands();
 }
 
 FString Tex::GetReferencerName() const
diff --git a/Engine/Plugins/TextureGraph/Source/TextureGraphEngine/Device/FX/Device_FX.cpp b/Engine/Plugins/TextureGraph/Source/TextureGraphEngine/Device/FX/Device_FX.cpp
--- a/Engine/Plugins/TextureGraph/Source/TextureGraphEngine/Device/FX/Device_FX.cpp
+++ b/Engine/Plugins/TextureGraph/Source/TextureGraphEngine/Device/FX/Device_FX.cpp
@@ -710,7 +710,7 @@ int32 Device_FX::InitRTResource(TexPtr TextureObj, UTextureRenderTarget* RT)
 	UE_LOG(LogDevice, Verbose, TEXT("New RT allocation: %s %llu [Ptr: 0x%p, Size: %ux%u]"), *TextureObj->GetDescriptor().Name,
 		TextureObj->GetDescriptor().Format_HashValue(), RT, RTRes->GetSizeX(), RTRes->GetSizeY());
 	
-	FTextureRHIRef RHITexture = RTRes->GetShaderResourceTexture();
+	FTextureRHIRef RHITexture = RTRes->GetTextureRHI();
 	check(RHITexture);
 	RHICmdList.BindDebugLabelName(RHITexture, *TextureObj->GetDescriptor().Name);
 	
diff --git a/Engine/Plugins/TextureGraph/Source/TextureGraphEngine/FxMat/RenderMaterial_BP.cpp b/Engine/Plugins/TextureGraph/Source/TextureGraphEngine/FxMat/RenderMaterial_BP.cpp
--- a/Engine/Plugins/TextureGraph/Source/TextureGraphEngine/FxMat/RenderMaterial_BP.cpp
+++ b/Engine/Plugins/TextureGraph/Source/TextureGraphEngine/FxMat/RenderMaterial_BP.cpp
@@ -275,17 +275,16 @@ void RenderMaterial_BP::BlitTo(FRHICommandListImmediate& RHI, UTextureRenderTarg
 	if (!MaterialInstanceValidated)
 		return;
 
-	if (!MeshObj)
+	if (!MeshObj && RT)
 	{
-		FTextureRenderTarget2DResource* RTRes = (FTextureRenderTarget2DResource*)RT->GetRenderTargetResource();
+		FTextureRenderTargetResource* RTRes = RT->GetRenderTargetResource();
 		check(RTRes);
 
-		FTextureRHIRef TextureRHI = RTRes->GetShaderResourceTexture();
+		FTextureRHIRef TextureRHI = RTRes->GetTextureRHI();
 		check(TextureRHI);
 
 		TextureRHI->SetName(FName(*RT->GetName()));
 
-		check(RT);
 		std::static_pointer_cast<FxMaterial_QuadDrawMaterial>(FXMaterialObj)->MyBlit(RHI, RT, TextureRHI, MeshObj, TargetId);
 	}
 	else
diff --git a/Engine/Plugins/TextureGraph/Source/TextureGraphEngine/FxMat/RenderMaterial_FX.cpp b/Engine/Plugins/TextureGraph/Source/TextureGraphEngine/FxMat/RenderMaterial_FX.cpp
--- a/Engine/Plugins/TextureGraph/Source/TextureGraphEngine/FxMat/RenderMaterial_FX.cpp
+++ b/Engine/Plugins/TextureGraph/Source/TextureGraphEngine/FxMat/RenderMaterial_FX.cpp
@@ -80,10 +80,10 @@ void RenderMaterial_FX::BlitTo(FRHICommandListImmediate& RHI, UTextureRenderTarg
 {
 	check(RenderTarget);
 
-	FTextureRenderTarget2DResource* RTRes = (FTextureRenderTarget2DResource*)RenderTarget->GetRenderTargetResource();
+	FTextureRenderTargetResource* RTRes = RenderTarget->GetRenderTargetResource();
 	check(RTRes);
 
-	FTextureRHIRef TextureRHI = RTRes->GetShaderResourceTexture();
+	FTextureRHIRef TextureRHI = RTRes->GetTextureRHI();
 	check(TextureRHI);
 
 	TextureRHI->SetName(FName(*RenderTarget->GetName()));
diff --git a/Engine/Plugins/TextureGraph/Source/TextureGraphEngine/Job/Job.cpp b/Engine/Plugins/TextureGraph/Source/TextureGraphEngine/Job/Job.cpp
--- a/Engine/Plugins/TextureGraph/Source/TextureGraphEngine/Job/Job.cpp
+++ b/Engine/Plugins/TextureGraph/Source/TextureGraphEngine/Job/Job.cpp
@@ -465,11 +465,11 @@ TiledBlobPtr Job::InitResult(FString InNewName, const BufferDescriptor* InDesire
 	ResultDesc.bIsTransient = IsDiscard();
 	ResultDesc.Name = NewName;
 
-	if ((int)ResultDesc.Width < (int)EResolution::Auto || (int)ResultDesc.Width > (int)EResolution::Resolution8192)
-		ResultDesc.Width = (int)EResolution::Auto;
+	if (ResultDesc.Width < static_cast<uint32>(EResolution::Auto) || ResultDesc.Width > static_cast<uint32>(EResolution::Resolution8192))
+		ResultDesc.Width = static_cast<uint32>(EResolution::Auto);
 
-	if ((int)ResultDesc.Height < (int)EResolution::Auto || (int)ResultDesc.Height > (int)EResolution::Resolution8192)
-		ResultDesc.Height = (int)EResolution::Auto;
+	if (ResultDesc.Height < static_cast<uint32>(EResolution::Auto) || ResultDesc.Height > static_cast<uint32>(EResolution::Resolution8192))
+		ResultDesc.Height = static_cast<uint32>(EResolution::Auto);
 
 	/// Update the job name as well
 	Name = NewName;
diff --git a/Engine/Plugins/TextureGraph/Source/TextureGraphEngine/Transform/Layer/T_Thumbnail.cpp b/Engine/Plugins/TextureGraph/Source/TextureGraphEngine/Transform/Layer/T_Thumbnail.cpp
--- a/Engine/Plugins/TextureGraph/Source/TextureGraphEngine/Transform/Layer/T_Thumbnail.cpp
+++ b/Engine/Plugins/TextureGraph/Source/TextureGraphEngine/Transform/Layer/T_Thumbnail.cpp
@@ -51,9 +51,12 @@ TiledBlobRef T_Thumbnail::Bind(UMixInterface* Mix, UObject* Model, TiledBlobPtr
 
 	// Normalize dimensions and remap to expected thumbnail size
 	const float MaxDimension = FMath::Max(Desc.Width, Desc.Height);
-	Desc.Width = RenderMaterial_Thumbnail::GThumbWidth * Desc.Width / MaxDimension;
-	Desc.Height = RenderMaterial_Thumbnail::GThumbHeight * Desc.Height / MaxDimension;
-
+	if (MaxDimension > 0)
+	{
+		Desc.Width = RenderMaterial_Thumbnail::GThumbWidth * Desc.Width / MaxDimension;
+		Desc.Height = RenderMaterial_Thumbnail::GThumbHeight * Desc.Height / MaxDimension;
+	}
+	
 	FString Name = FString::Printf(TEXT("T_Thumbnail [%s]"), *InBlobToBind->GetDescriptor().Name);
 	
 	TiledBlobPtr Result = JobPtr->InitResult(Name, &Desc);
