TG: Fix cook crash when AllowCommandlets=1

#Synced-CL 49841360
#Synced-CL 49841360

diff --git a/Engine/Plugins/TextureGraph/Source/TextureGraphEngine/TextureGraphEngine.cpp b/Engine/Plugins/TextureGraph/Source/TextureGraphEngine/TextureGraphEngine.cpp
--- a/Engine/Plugins/TextureGraph/Source/TextureGraphEngine/TextureGraphEngine.cpp
+++ b/Engine/Plugins/TextureGraph/Source/TextureGraphEngine/TextureGraphEngine.cpp
@@ -244,7 +244,7 @@ void TextureGraphEngine::DoneTests()
 
 	// Currently we're only supporting editor only 
 #if WITH_EDITOR
-	if (!GEditor)
+	if (!GEditor || !GInstance || !GInstance->bIsEngineInitDone)
 		return;
 		
 	bGIsEngineDestroying = true;
@@ -294,7 +294,7 @@ void TextureGraphEngine::SetRunEngine()
 
 bool TextureGraphEngine::BlockCommandlets()
 {
-	return !CVarAllowTextureGraphInCommandlets.GetValueOnAnyThread()
+	return (IsRunningCookCommandlet() || IsRunningDLCCookCommandlet() || !CVarAllowTextureGraphInCommandlets.GetValueOnAnyThread())
 	#if WITH_EDITOR
 		&& !(GEditor && FSlateApplication::IsInitialized() && FApp::CanEverRender())
 	#endif
