TextureGraphEngine: Only reference RenderDocPlugin when building developer tools

#rnx
#Synced-CL 50554652

diff --git a/Engine/Plugins/TextureGraph/Source/TextureGraphEngine/TextureGraphEngine.Build.cs b/Engine/Plugins/TextureGraph/Source/TextureGraphEngine/TextureGraphEngine.Build.cs
--- a/Engine/Plugins/TextureGraph/Source/TextureGraphEngine/TextureGraphEngine.Build.cs
+++ b/Engine/Plugins/TextureGraph/Source/TextureGraphEngine/TextureGraphEngine.Build.cs
@@ -85,7 +85,7 @@ public class TextureGraphEngine : ModuleRules
 			"DeveloperSettings",
 		});
 
-		if (Target.Platform == UnrealTargetPlatform.Win64 || Target.Platform == UnrealTargetPlatform.Linux)
+		if (Target.bBuildDeveloperTools && (Target.Platform == UnrealTargetPlatform.Win64 || Target.Platform == UnrealTargetPlatform.Linux))
 		{
 			PublicDependencyModuleNames.Add("RenderDocPlugin");
 			PublicDefinitions.Add("TEXTUREGRAPH_RENDERDOC_ENABLED=1");
