Remove FlushRenderingCommands from Tex::InitRT, as it is no longer needed for Linux resource initialization and also fixes recursion warnings on commandlets

#Synced-CL 50537359

diff --git a/Engine/Plugins/TextureGraph/Source/TextureGraphEngine/2D/Tex.cpp b/Engine/Plugins/TextureGraph/Source/TextureGraphEngine/2D/Tex.cpp
--- a/Engine/Plugins/TextureGraph/Source/TextureGraphEngine/2D/Tex.cpp
+++ b/Engine/Plugins/TextureGraph/Source/TextureGraphEngine/2D/Tex.cpp
@@ -393,9 +393,9 @@ void Tex::InitRT(bool ForceFloat /* = false */)
 	RT->RenderTargetFormat = TextureHelper::GetRenderTargetFormatFromPixelFormat(Desc.Format);
 	RT->OverrideFormat = Desc.Format;
 	RT->LODBias = 0;
-	RT->InitAutoFormat(Desc.Width, Desc.Height);
+	RT->SizeX = Desc.Width;
+	RT->SizeY = Desc.Height;
 	RT->UpdateResource();
-	FlushRenderingCommands();
 }
 
 FString Tex::GetReferencerName() const
