Fix stale entries in TextureGraphInstance OutputSettingsMap

Reset OutputSettingsMap before collecting output settings in UpdateOutputSettingsFromGraph.
Previously, removed output nodes would persist as stale entries because CollectOutputSettings
only adds entries without clearing old ones first. This caused the MetaHuman Export to DCC to fail
when trying to load textures for outputs that no longer exist in the graph.

#Synced-CL 52673835

diff --git a/Engine/Plugins/TextureGraph/Source/TextureGraph/Private/TextureGraph.cpp b/Engine/Plugins/TextureGraph/Source/TextureGraph/Private/TextureGraph.cpp
--- a/Engine/Plugins/TextureGraph/Source/TextureGraph/Private/TextureGraph.cpp
+++ b/Engine/Plugins/TextureGraph/Source/TextureGraph/Private/TextureGraph.cpp
@@ -111,6 +111,7 @@ bool UTextureGraphInstance::CheckOutputSettingsMatch(const TObjectPtr<UTextureGr
 }
 void UTextureGraphInstance::UpdateOutputSettingsFromGraph()
 {
+	OutputSettingsMap.Reset();
 	Graph()->CollectOutputSettings(OutputSettingsMap);
 	//TODO: Keep values if names match
 }
